VR Editor :
- UISystem is ticking after worldinteraction tick, so it is not a frame behind.
- Added selectionbar to the dockable window
- Fixed RadialMenu crosshair being inverted after 2889817

diff --git a/Engine/Source/Editor/VREditor/VREditorDockableWindow.cpp b/Engine/Source/Editor/VREditor/VREditorDockableWindow.cpp
--- a/Engine/Source/Editor/VREditor/VREditorDockableWindow.cpp
+++ b/Engine/Source/Editor/VREditor/VREditorDockableWindow.cpp
@@ -9,34 +9,63 @@
 
 namespace VREd
 {
-	static FAutoConsoleVariable DockUIZBarOffset( TEXT( "VREd.DockUIZBarOffset" ), 0.5f, TEXT( "Z Distance between the selectionbar and the UI" ) );
+	static FAutoConsoleVariable DockWindowTickness( TEXT( "VREd.DockWindowTickness" ), 0.003f, TEXT( "Z Distance between the selectionbar and the UI" ) );
+	static FAutoConsoleVariable DockUIZBarOffset( TEXT( "VREd.DockUIZBarOffset" ), 1.0f, TEXT( "Z Distance between the selectionbar and the UI" ) );
 }
 
 AVREditorDockableWindow::AVREditorDockableWindow() 
 	: Super(),
-	SelectionBarScale( FVector(1.0f, 1.0f, 1.0f) )
+	SelectionMeshScale( FVector(0.6f, 2.0f, 1.2f) )
 {
-	UStaticMesh* SelectionMesh = nullptr;
+
 	{
-		static ConstructorHelpers::FObjectFinder<UStaticMesh> ObjectFinder( TEXT( "/Engine/VREditor/SM_ContentWindow_01" ) );
-		SelectionMesh = ObjectFinder.Object; 
-		check( SelectionMesh != nullptr );
+		UStaticMesh* WindowMesh = nullptr;
+		{
+			static ConstructorHelpers::FObjectFinder<UStaticMesh> ObjectFinder( TEXT( "/Engine/VREditor/SM_ContentWindow_01" ) );
+			WindowMesh = ObjectFinder.Object;
+			check( WindowMesh != nullptr );
+		}
+
+		WindowMeshComponent = CreateDefaultSubobject<UStaticMeshComponent>( TEXT( "SelectionBar" ) );
+		WindowMeshComponent->SetStaticMesh( WindowMesh );
+		WindowMeshComponent->SetMobility( EComponentMobility::Movable );
+		WindowMeshComponent->AttachTo( RootComponent );
+
+		WindowMeshComponent->SetCollisionEnabled( ECollisionEnabled::QueryOnly );
+		WindowMeshComponent->SetCollisionResponseToAllChannels( ECollisionResponse::ECR_Block );
+		WindowMeshComponent->SetCollisionResponseToChannel( ECC_EditorGizmo, ECollisionResponse::ECR_Block );
+
+		WindowMeshComponent->bGenerateOverlapEvents = false;
+		WindowMeshComponent->SetCanEverAffectNavigation( false );
+		WindowMeshComponent->bCastDynamicShadow = false;
+		WindowMeshComponent->bCastStaticShadow = false;
+		WindowMeshComponent->bAffectDistanceFieldLighting = false;
 	}
 
-	SelectionBar = CreateDefaultSubobject<UStaticMeshComponent>( TEXT( "SelectionBar" ) ); 
-	SelectionBar->SetStaticMesh( SelectionMesh );
-	SelectionBar->SetMobility( EComponentMobility::Movable );
-	SelectionBar->AttachTo( RootComponent );
+	{
+		UStaticMesh* SelectionMesh = nullptr;
+		{
+			static ConstructorHelpers::FObjectFinder<UStaticMesh> ObjectFinder( TEXT( "/Engine/VREditor/TransformGizmo/UniformScaleHandle" ) );
+			SelectionMesh = ObjectFinder.Object;
+			check( SelectionMesh != nullptr );
+		}
 
-	SelectionBar->SetCollisionEnabled( ECollisionEnabled::QueryOnly );
-	SelectionBar->SetCollisionResponseToAllChannels( ECollisionResponse::ECR_Block );
-	SelectionBar->SetCollisionResponseToChannel( ECC_EditorGizmo, ECollisionResponse::ECR_Block );
+		SelectionMeshComponent = CreateDefaultSubobject<UStaticMeshComponent>( TEXT( "SelectionMesh" ) );
+		SelectionMeshComponent->SetStaticMesh( SelectionMesh );
+		SelectionMeshComponent->SetMobility( EComponentMobility::Movable );
+		SelectionMeshComponent->AttachTo( RootComponent );
+
+		SelectionMeshComponent->SetCollisionEnabled( ECollisionEnabled::QueryOnly );
+		SelectionMeshComponent->SetCollisionResponseToAllChannels( ECollisionResponse::ECR_Block );
+		SelectionMeshComponent->SetCollisionResponseToChannel( ECC_EditorGizmo, ECollisionResponse::ECR_Block );
+
+		SelectionMeshComponent->bGenerateOverlapEvents = false;
+		SelectionMeshComponent->SetCanEverAffectNavigation( false );
+		SelectionMeshComponent->bCastDynamicShadow = false;
+		SelectionMeshComponent->bCastStaticShadow = false;
+		SelectionMeshComponent->bAffectDistanceFieldLighting = false;
+	}
 
-	SelectionBar->bGenerateOverlapEvents = false;
-	SelectionBar->SetCanEverAffectNavigation( false );
-	SelectionBar->bCastDynamicShadow = false;
-	SelectionBar->bCastStaticShadow = false;
-	SelectionBar->bAffectDistanceFieldLighting = false;
 }
 
 void AVREditorDockableWindow::TickManually( float DeltaTime )
@@ -46,7 +75,15 @@ void AVREditorDockableWindow::TickManually( float DeltaTime )
 	if( WidgetComponent->IsVisible() )
 	{
 		const FVector AnimatedScale = CalculateAnimatedScale();
-		SelectionBar->SetWorldScale3D( ( FVector( 0.001f, 1.0 / (float)Resolution.X, 1.0f / (float)Resolution.Y ) * ( Size *  GetOwner().GetOwner().GetWorldScaleFactor() ) ) * 10 * AnimatedScale );
+		const float WordScaleFactor = GetOwner().GetOwner().GetWorldScaleFactor();
+		const FVector SelectionMeshWorldScale = SelectionMeshScale * WordScaleFactor;
+		SelectionMeshComponent->SetWorldScale3D( SelectionMeshWorldScale * AnimatedScale );
+
+		const FVector WindowMeshWorldScale = FVector( VREd::DockWindowTickness->GetFloat() , 1.0 / (float)Resolution.X, 1.0f / (float)Resolution.Y  ) * ( Size *  GetOwner().GetOwner().GetWorldScaleFactor() ) * AnimatedScale;
+		WindowMeshComponent->SetWorldScale3D( WindowMeshWorldScale * 10 );	
+		
+		const FVector NewRelativeLocation = FVector( 0.f, 0.f, -(((WindowMeshWorldScale.Z * 0.5f) * 1000) + ((SelectionMeshWorldScale.Z * 0.5) * 10) + (VREd::DockUIZBarOffset->GetFloat() * WordScaleFactor)) );
+		SelectionMeshComponent->SetRelativeLocation( NewRelativeLocation );
 	}
 } 
 
diff --git a/Engine/Source/Editor/VREditor/VREditorDockableWindow.h b/Engine/Source/Editor/VREditor/VREditorDockableWindow.h
--- a/Engine/Source/Editor/VREditor/VREditorDockableWindow.h
+++ b/Engine/Source/Editor/VREditor/VREditorDockableWindow.h
@@ -13,23 +13,26 @@ class AVREditorDockableWindow : public AVREditorFloatingUI
 {
 	GENERATED_BODY()
 	
-	/** The dockable */
+	/** The dockable window mesh */
 	UPROPERTY()
-	class UStaticMeshComponent* SelectionBar;
+	class UStaticMeshComponent* WindowMeshComponent;
+
+	/** Mesh underneath the window for easy selecting and dragging */
+	UPROPERTY()
+	class UStaticMeshComponent* SelectionMeshComponent;
 
 public:
 
 	/** Default constructor */
 	AVREditorDockableWindow();
 
-	/** Updates the selection bar scale and the widgetcomponent position */
+	/** Updates the meshes for the UI */
 	virtual void TickManually( float DeltaTime ) override;
 
 	/** Updates the last dragged relative position */
 	void UpdateRelativeRoomTransform();
 
 private:
-
-	/** Scale for the selection bar that will be updated with the worldscale */
-	FVector SelectionBarScale;
+	/** The room space scale for the SelectionMeshComponent */
+	FVector SelectionMeshScale;
 };	
\ No newline at end of file
diff --git a/Engine/Source/Editor/VREditor/VREditorRadialMenu.cpp b/Engine/Source/Editor/VREditor/VREditorRadialMenu.cpp
--- a/Engine/Source/Editor/VREditor/VREditorRadialMenu.cpp
+++ b/Engine/Source/Editor/VREditor/VREditorRadialMenu.cpp
@@ -34,14 +34,14 @@ UVREditorRadialMenu::UVREditorRadialMenu( const FObjectInitializer& ObjectInitia
 
 void UVREditorRadialMenu::Update( const FVirtualHand& Hand )
 {
-	TrackpadPosition = FVector2D( Hand.TrackpadPosition.X, -Hand.TrackpadPosition.Y );
+	TrackpadPosition = Hand.TrackpadPosition;
 	
 	if( !GetOwner()->bHidden )
 	{
 		// Check if position is from the center
 		if( !IsInMenuRadius() )
 		{
-			float Angle = FRotator::NormalizeAxis( FMath::RadiansToDegrees(FMath::Atan2( Hand.TrackpadPosition.X,  Hand.TrackpadPosition.Y ) ) );
+			float Angle = FRotator::NormalizeAxis( FMath::RadiansToDegrees(FMath::Atan2( TrackpadPosition.X, -TrackpadPosition.Y ) ) );
 			TrackpadAngle = Angle;
 			if( Angle < 0)
 			{
diff --git a/Engine/Source/Editor/VREditor/VREditorUISystem.cpp b/Engine/Source/Editor/VREditor/VREditorUISystem.cpp
--- a/Engine/Source/Editor/VREditor/VREditorUISystem.cpp
+++ b/Engine/Source/Editor/VREditor/VREditorUISystem.cpp
@@ -41,7 +41,7 @@ namespace VREd
 	static FAutoConsoleVariable EditorUISize( TEXT( "VREd.EditorUISize" ), 50.0f, TEXT( "How big editor UIs should be" ) );
 	static FAutoConsoleVariable ContentBrowserUIScale( TEXT( "VREd.ContentBrowserUIScale" ), 2.5f, TEXT( "How much to scale up (or down) the content browser for VR" ) );
 	static FAutoConsoleVariable EditorUIScale( TEXT( "VREd.EditorUIScale" ), 2.5f, TEXT( "How much to scale up (or down) editor UIs for VR" ) );
-	static FAutoConsoleVariable EditorUIOffsetFromHand( TEXT( "VREd.EditorUIOffsetFromHand" ), 30.0f, TEXT( "How far to offset editor UIs from the origin of the hand mesh" ) );
+	static FAutoConsoleVariable EditorUIOffsetFromHand( TEXT( "VREd.EditorUIOffsetFromHand" ), 40.0f, TEXT( "How far to offset editor UIs from the origin of the hand mesh" ) );
 	static FAutoConsoleVariable RadialMenuFadeDelay( TEXT( "VREd.RadialMenuFadeDelay" ), 0.2f, TEXT( "The delay for the radial menu after selecting a button" ) );
 	static FAutoConsoleVariable UIAbsoluteScrollSpeed( TEXT( "VREd.UIAbsoluteScrollSpeed" ), 8.0f, TEXT( "How fast the UI scrolls when dragging the touchpad" ) );
 	static FAutoConsoleVariable UIRelativeScrollSpeed( TEXT( "VREd.UIRelativeScrollSpeed" ), 0.75f, TEXT( "How fast the UI scrolls when holding an analog stick" ) );
