VREditor: Added double clicking support!

diff --git a/Engine/Source/Editor/VREditor/VREditorUISystem.cpp b/Engine/Source/Editor/VREditor/VREditorUISystem.cpp
--- a/Engine/Source/Editor/VREditor/VREditorUISystem.cpp
+++ b/Engine/Source/Editor/VREditor/VREditorUISystem.cpp
@@ -47,6 +47,7 @@ namespace VREd
 	static FAutoConsoleVariable UIRelativeScrollSpeed( TEXT( "VREd.UIRelativeScrollSpeed" ), 0.75f, TEXT( "How fast the UI scrolls when holding an analog stick" ) );
 	static FAutoConsoleVariable MinUIScrollDeltaForInertia( TEXT( "VREd.MinUIScrollDeltaForInertia" ), 0.25f, TEXT( "Minimum amount of touch pad input before inertial UI scrolling kicks in" ) );
 	static FAutoConsoleVariable MinDockDragDistance( TEXT( "VREd.MinDockDragDistance" ), 10.0f, TEXT( "Minimum amount of distance needed to drag dock from hand" ) );
+	static FAutoConsoleVariable DoubleClickTime( TEXT( "VREd.DoubleClickTime" ), 0.25f, TEXT( "Minumum duration between clicks for a double click event to fire" ) );
 }
 
 
@@ -200,6 +201,15 @@ void FVREditorUISystem::OnVRAction( FEditorViewportClient& ViewportClient, const
 									else if( Event == IE_Released )
 									{
 										Reply = FSlateApplication::Get().RoutePointerUpEvent( WidgetPathUnderFinger, PointerEvent );
+
+										const double CurrentTime = FPlatformTime::Seconds();
+										if( CurrentTime - Hand.LastClickReleaseTime <= VREd::DoubleClickTime->GetFloat() )
+										{
+											// Trigger a double click event!
+											Reply = FSlateApplication::Get().RoutePointerDoubleClickEvent( WidgetPathUnderFinger, PointerEvent );
+										}
+
+										Hand.LastClickReleaseTime = CurrentTime;
 									}
 								}
 							}
diff --git a/Engine/Source/Editor/VREditor/VirtualHand.h b/Engine/Source/Editor/VREditor/VirtualHand.h
--- a/Engine/Source/Editor/VREditor/VirtualHand.h
+++ b/Engine/Source/Editor/VREditor/VirtualHand.h
@@ -168,6 +168,9 @@ struct FVirtualHand
 	/** True if we're currently holding the 'SelectAndMove' button down after clicking on UI */
 	bool bIsClickingOnUI;
 
+	/** Last real time that we released the 'SelectAndMove' button on UI.  This is used to detect double-clicks. */
+	double LastClickReleaseTime;
+
 	/** The last real time we played a haptic effect.  This is used to turn off haptic effects shortly after they're triggered. */
 	double LastHapticTime;
 
@@ -316,6 +319,7 @@ struct FVirtualHand
 		}
 		ClickingOnComponent = nullptr;
 		bIsClickingOnUI = false;
+		LastClickReleaseTime = 0.0;
 		LastHapticTime = 0.0;
 		
 		bHasUIInFront = false;
