VR Editor: UI no longer teleports when picked up

diff --git a/Engine/Source/Editor/VREditor/VREditorUISystem.cpp b/Engine/Source/Editor/VREditor/VREditorUISystem.cpp
--- a/Engine/Source/Editor/VREditor/VREditorUISystem.cpp
+++ b/Engine/Source/Editor/VREditor/VREditorUISystem.cpp
@@ -60,6 +60,7 @@ FVREditorUISystem::FVREditorUISystem( FVREditorMode& InitOwner )
 	  QuickRadialWidgetClass( nullptr ),
 	  DraggingUI( nullptr ),
 	  DraggingUIHandIndex( INDEX_NONE ),
+	  DraggingUIOffsetTransform( FTransform::Identity ),
 	  bPanelVisibilityToggle( false )
 {
 	// Register to find out about VR events
@@ -925,7 +926,32 @@ void FVREditorUISystem::HideRadialMenu( const int32 HandIndex )
 	}
 }
 
-AVREditorFloatingUI* FVREditorUISystem::StartDraggingDockUI( AVREditorDockableWindow* InitDraggingDockUI, const int32 HandIndex )
+
+FTransform FVREditorUISystem::MakeDockableUITransformOnLaser( AVREditorDockableWindow* InitDraggingDockUI, const int32 HandIndex, const float DockSelectDistance ) const
+{
+	const FVirtualHand& Hand = Owner.GetVirtualHand( HandIndex );
+
+	const FVector NewLocation = Hand.Transform.GetLocation() + ( Hand.Transform.GetRotation().Vector().GetSafeNormal() * DockSelectDistance );
+	
+	FRotator NewRotation = ( Hand.Transform.GetLocation() - NewLocation ).ToOrientationRotator();
+	NewRotation.Roll = -Hand.Transform.GetRotation().Rotator().Roll;
+	
+	const FTransform LaserImpactToWorld( NewRotation, NewLocation );
+	return LaserImpactToWorld;
+}
+
+
+FTransform FVREditorUISystem::MakeDockableUITransform( AVREditorDockableWindow* InitDraggingDockUI, const int32 HandIndex, const float DockSelectDistance ) const
+{
+	const FTransform UIOnLaserToWorld = MakeDockableUITransformOnLaser( DraggingUI, HandIndex, DockSelectDistance );
+	const FTransform UIToUIOnLaser = DraggingUIOffsetTransform;
+	
+	const FTransform UIToWorld = UIToUIOnLaser * UIOnLaserToWorld;
+	return UIToWorld;
+}
+
+
+AVREditorFloatingUI* FVREditorUISystem::StartDraggingDockUI( AVREditorDockableWindow* InitDraggingDockUI, const int32 HandIndex, const float DockSelectDistance )
 {
 	AVREditorFloatingUI::EDockedTo DockTo = InitDraggingDockUI->GetDockedTo();
 	if( DockTo == AVREditorFloatingUI::EDockedTo::LeftHand || DockTo == AVREditorFloatingUI::EDockedTo::RightHand )
@@ -934,6 +960,16 @@ AVREditorFloatingUI* FVREditorUISystem::StartDraggingDockUI( AVREditorDockableWi
 	}
 
 	DraggingUIHandIndex = HandIndex;
+
+	FTransform UIToWorld = InitDraggingDockUI->GetActorTransform();
+	UIToWorld.SetScale3D( FVector( 1.0f ) );
+	const FTransform WorldToUI = UIToWorld.Inverse();
+
+	const FTransform UIOnLaserToWorld = MakeDockableUITransformOnLaser( InitDraggingDockUI, HandIndex, DockSelectDistance );
+	const FTransform UIOnLaserToUI = UIOnLaserToWorld * WorldToUI;
+	const FTransform UIToUIOnLaser = UIOnLaserToUI.Inverse();
+	DraggingUIOffsetTransform = UIToUIOnLaser;
+
 	DraggingUI = InitDraggingDockUI;
 	DraggingUI->SetDockedTo( AVREditorFloatingUI::EDockedTo::Nothing );
 	return DraggingUI;
diff --git a/Engine/Source/Editor/VREditor/VREditorUISystem.h b/Engine/Source/Editor/VREditor/VREditorUISystem.h
--- a/Engine/Source/Editor/VREditor/VREditorUISystem.h
+++ b/Engine/Source/Editor/VREditor/VREditorUISystem.h
@@ -76,8 +76,14 @@ public:
 	void HideRadialMenu( const int32 HandIndex );
 
 	/** Start dragging a dock window on the hand */
-	class AVREditorFloatingUI* StartDraggingDockUI( class AVREditorDockableWindow* InitDraggingDockUI, const int32 HandIndex );
+	class AVREditorFloatingUI* StartDraggingDockUI( class AVREditorDockableWindow* InitDraggingDockUI, const int32 HandIndex, const float DockSelectDistance );
 
+	/** Makes up a transform for a dockable UI when dragging it with a laser at the specified distance from the laser origin */
+	FTransform MakeDockableUITransformOnLaser( AVREditorDockableWindow* InitDraggingDockUI, const int32 HandIndex, const float DockSelectDistance ) const;
+
+	/** Makes up a transform for a dockable UI when dragging it that includes the original offset from the laser's impact point */
+	FTransform MakeDockableUITransform( AVREditorDockableWindow* InitDraggingDockUI, const int32 HandIndex, const float DockSelectDistance ) const;
+	
 	/** Returns the current dragged dock window, nullptr if none */
 	class AVREditorDockableWindow* GetDraggingDockUI() const;
 
@@ -132,6 +138,9 @@ protected:
 	/** Hand that is dragging the UI */
 	int32 DraggingUIHandIndex;
 
+	/** Offset transform from room-relative transform to the object where we picked it up at */
+	FTransform DraggingUIOffsetTransform;
+
 	/** The current UI that is being dragged */
 	class AVREditorDockableWindow* DraggingUI;
 
diff --git a/Engine/Source/Editor/VREditor/VREditorWorldInteraction.cpp b/Engine/Source/Editor/VREditor/VREditorWorldInteraction.cpp
--- a/Engine/Source/Editor/VREditor/VREditorWorldInteraction.cpp
+++ b/Engine/Source/Editor/VREditor/VREditorWorldInteraction.cpp
@@ -384,7 +384,7 @@ void FVREditorWorldInteraction::OnVRAction( FEditorViewportClient& ViewportClien
 								{
 									Hand.DraggingMode = EVREditorDraggingMode::DockableWindow;
 									DockSelectDistance = ( LaserPointerStart - HitResult.Location ).Size();
-									Owner.GetUISystem().StartDraggingDockUI( DockableWindow, VRAction.HandIndex );
+									Owner.GetUISystem().StartDraggingDockUI( DockableWindow, VRAction.HandIndex, DockSelectDistance );
 									bWasHandled = true;
 								}
 							}
@@ -1160,13 +1160,9 @@ void FVREditorWorldInteraction::Tick( FEditorViewportClient* ViewportClient, con
 				AVREditorDockableWindow* DraggingUI = Owner.GetUISystem().GetDraggingDockUI();
 				if( DraggingUI )
 				{
-					FVector NewLocation = Hand.Transform.GetLocation() + ( Hand.Transform.GetRotation().Vector().GetSafeNormal() * DockSelectDistance );
-					DraggingUI->SetActorLocation( NewLocation );
-					
-					FRotator Rotation = (Hand.Transform.GetLocation() - NewLocation).ToOrientationRotator();
-					Rotation.Roll = -Hand.Transform.GetRotation().Rotator().Roll;
-					DraggingUI->SetActorRotation( FQuat( Rotation ) );
-					
+					const FTransform DockedUIToWorld = Owner.GetUISystem().MakeDockableUITransform( DraggingUI, HandIndex, DockSelectDistance );
+					DraggingUI->SetActorTransform( DockedUIToWorld );
+
 					DraggingUI->UpdateRelativeRoomTransform();
 				}
 			}
