VR Editor :
- Removed duplicated code to set dynamic materials for gizmo handles by creating a function in the base gizmo handle
- Removed static colors in headers
- Added list of colors to the mode that can be accessed using color enum in the VREditorMode
- Replaced old static color references to new system
- Added on hover laser pointer leave and enter to the dockable window that changes its color

diff --git a/Engine/Source/Editor/VREditor/Gizmo/VREditorBaseTransformGizmo.cpp b/Engine/Source/Editor/VREditor/Gizmo/VREditorBaseTransformGizmo.cpp
--- a/Engine/Source/Editor/VREditor/Gizmo/VREditorBaseTransformGizmo.cpp
+++ b/Engine/Source/Editor/VREditor/Gizmo/VREditorBaseTransformGizmo.cpp
@@ -4,6 +4,7 @@
 #include "VREditorBaseTransformGizmo.h"
 #include "UnitConversion.h"
 #include "VREditorGizmoHandle.h"
+#include "VREditorMode.h"
 
 namespace VREd
 {
@@ -12,7 +13,8 @@ namespace VREd
 	static FAutoConsoleVariable GizmoSelectionAnimationCurvePower( TEXT( "VREd.GizmoSelectionAnimationCurvePower" ), 2.0f, TEXT( "Controls the animation curve for the gizmo after objects are selected" ) );
 }
 
-ABaseTransformGizmo::ABaseTransformGizmo( )
+ABaseTransformGizmo::ABaseTransformGizmo( ) :
+	Owner( nullptr )
 {
 	// Create root default scene component
 	{
@@ -167,6 +169,16 @@ float ABaseTransformGizmo::GetAnimationAlpha()
 	return AnimationAlpha;
 }
 
+void ABaseTransformGizmo::SetOwnerMode( FVREditorMode* InOwner )
+{
+	Owner = InOwner;
+}
+
+FVREditorMode* ABaseTransformGizmo::GetOwnerMode() const
+{
+	return Owner;
+}
+
 void ABaseTransformGizmo::UpdateHandleVisibility( const EGizmoHandleTypes GizmoType, const ECoordSystem GizmoCoordinateSpace, const bool bAllHandlesVisible, UActorComponent* DraggingHandle )
 {
 	for ( UVREditorGizmoHandleGroup* HandleGroup : AllHandleGroups )
diff --git a/Engine/Source/Editor/VREditor/Gizmo/VREditorBaseTransformGizmo.h b/Engine/Source/Editor/VREditor/Gizmo/VREditorBaseTransformGizmo.h
--- a/Engine/Source/Editor/VREditor/Gizmo/VREditorBaseTransformGizmo.h
+++ b/Engine/Source/Editor/VREditor/Gizmo/VREditorBaseTransformGizmo.h
@@ -41,6 +41,12 @@ public:
 	/** Updates the animation with the currenttime and selectedtime */
 	float GetAnimationAlpha();
 
+	/** Sets the owner */
+	void SetOwnerMode( class FVREditorMode* InOwner );
+
+	/** Gets the owner */
+	class FVREditorMode* GetOwnerMode() const;
+
 protected:
 
 	/** Makes up a name string for a handle */
@@ -72,4 +78,7 @@ protected:
 	/** All gizmo components */
 	UPROPERTY()
 	TArray< class UVREditorGizmoHandleGroup* > AllHandleGroups;
+
+	/** Owning object */
+	class FVREditorMode* Owner;
 };
\ No newline at end of file
diff --git a/Engine/Source/Editor/VREditor/Gizmo/VREditorGizmoHandle.cpp b/Engine/Source/Editor/VREditor/Gizmo/VREditorGizmoHandle.cpp
--- a/Engine/Source/Editor/VREditor/Gizmo/VREditorGizmoHandle.cpp
+++ b/Engine/Source/Editor/VREditor/Gizmo/VREditorGizmoHandle.cpp
@@ -3,6 +3,7 @@
 #include "VREditorModule.h"
 #include "VREditorGizmoHandle.h"
 #include "VREditorTransformGizmo.h"
+#include "VREditorMode.h"
 
 UVREditorGizmoHandleGroup::UVREditorGizmoHandleGroup()
 	: Super()
@@ -164,4 +165,62 @@ TArray<FVREditorGizmoHandle>& UVREditorGizmoHandleGroup::GetHandles()
 EGizmoHandleTypes UVREditorGizmoHandleGroup::GetHandleType() const
 {
 	return EGizmoHandleTypes::All;
-}
\ No newline at end of file
+}
+
+void UVREditorGizmoHandleGroup::UpdateHandleColor( const int32 AxisIndex, FVREditorGizmoHandle& Handle, class UActorComponent* DraggingHandle, const TArray< UActorComponent* >& HoveringOverHandles )
+{
+	UStaticMeshComponent* HandleMesh = Handle.HandleMesh;
+
+	if ( !HandleMesh->GetMaterial( 0 )->IsA( UMaterialInstanceDynamic::StaticClass() ) )
+	{
+		UMaterialInstanceDynamic* MID = UMaterialInstanceDynamic::Create( GizmoMaterial, this );
+		HandleMesh->SetMaterial( 0, MID );
+	}
+	if ( !HandleMesh->GetMaterial( 1 )->IsA( UMaterialInstanceDynamic::StaticClass() ) )
+	{
+		UMaterialInstanceDynamic* MID = UMaterialInstanceDynamic::Create( TranslucentGizmoMaterial, this );
+		HandleMesh->SetMaterial( 1, MID );
+	}
+	UMaterialInstanceDynamic* MID0 = CastChecked<UMaterialInstanceDynamic>( HandleMesh->GetMaterial( 0 ) );
+	UMaterialInstanceDynamic* MID1 = CastChecked<UMaterialInstanceDynamic>( HandleMesh->GetMaterial( 1 ) );
+
+	ABaseTransformGizmo* GizmoActor = CastChecked<ABaseTransformGizmo>( GetOwner() );
+	if (GizmoActor)
+	{
+		FVREditorMode* Mode = GizmoActor->GetOwnerMode();
+		if (Mode)
+		{
+			FLinearColor HandleColor = Mode->GetColor( FVREditorMode::EColors::WhiteGizmoColor );
+			if (HandleMesh == DraggingHandle)
+			{
+				HandleColor = Mode->GetColor( FVREditorMode::EColors::DraggingGizmoColor );
+			}
+			else if (AxisIndex != INDEX_NONE)
+			{
+				switch (AxisIndex)
+				{
+				case 0:
+					HandleColor = Mode->GetColor( FVREditorMode::EColors::RedGizmoColor );
+					break;
+
+				case 1:
+					HandleColor = Mode->GetColor( FVREditorMode::EColors::GreenGizmoColor );
+					break;
+
+				case 2:
+					HandleColor = Mode->GetColor( FVREditorMode::EColors::BlueGizmoColor );
+					break;
+				}
+
+				if (HoveringOverHandles.Contains( HandleMesh ))
+				{
+					HandleColor = FLinearColor::LerpUsingHSV( HandleColor, Mode->GetColor( FVREditorMode::EColors::HoverGizmoColor ), Handle.HoverAlpha );
+				}
+			}
+
+			static FName StaticHandleColorParameter( "Color" );
+			MID0->SetVectorParameterValue( StaticHandleColorParameter, HandleColor );
+			MID1->SetVectorParameterValue( StaticHandleColorParameter, HandleColor );
+		}
+	}
+}
diff --git a/Engine/Source/Editor/VREditor/Gizmo/VREditorGizmoHandle.h b/Engine/Source/Editor/VREditor/Gizmo/VREditorGizmoHandle.h
--- a/Engine/Source/Editor/VREditor/Gizmo/VREditorGizmoHandle.h
+++ b/Engine/Source/Editor/VREditor/Gizmo/VREditorGizmoHandle.h
@@ -92,6 +92,8 @@ public:
 	}
 
 protected:
+	/** Updates the colors of the dynamic material instances for the handle passed using its axis index */	
+	void UpdateHandleColor( const int32 AxisIndex, FVREditorGizmoHandle& Handle, class UActorComponent* DraggingHandle, const TArray< UActorComponent* >& HoveringOverHandles );
 
 	/** Gizmo material (opaque) */
 	UPROPERTY()
diff --git a/Engine/Source/Editor/VREditor/Gizmo/VREditorPlaneTranslationGizmoHandle.cpp b/Engine/Source/Editor/VREditor/Gizmo/VREditorPlaneTranslationGizmoHandle.cpp
--- a/Engine/Source/Editor/VREditor/Gizmo/VREditorPlaneTranslationGizmoHandle.cpp
+++ b/Engine/Source/Editor/VREditor/Gizmo/VREditorPlaneTranslationGizmoHandle.cpp
@@ -2,7 +2,7 @@
 
 #include "VREditorModule.h"
 #include "VREditorPlaneTranslationGizmoHandle.h"
-#include "VREditorTransformGizmo.h"
+#include "VREditorBaseTransformGizmo.h"
 #include "UnitConversion.h"
 
 UVREditorPlaneTranslationGizmoHandleGroup::UVREditorPlaneTranslationGizmoHandleGroup()
@@ -130,52 +130,9 @@ void UVREditorPlaneTranslationGizmoHandleGroup::UpdateGizmoHandleGroup( const FT
 			PlaneTranslationHandle->SetRelativeRotation( GizmoSpaceOrientation );
 
 			PlaneTranslationHandle->SetRelativeScale3D( FVector( GizmoHandleScale ) );
-
+			
 			// Update material
-			{
-				if (!PlaneTranslationHandle->GetMaterial( 0 )->IsA( UMaterialInstanceDynamic::StaticClass() ))
-				{
-					UMaterialInstanceDynamic* MID = UMaterialInstanceDynamic::Create( GizmoMaterial, this );
-					PlaneTranslationHandle->SetMaterial( 0, MID );
-				}
-				if (!PlaneTranslationHandle->GetMaterial( 1 )->IsA( UMaterialInstanceDynamic::StaticClass() ))
-				{
-					UMaterialInstanceDynamic* MID = UMaterialInstanceDynamic::Create( TranslucentGizmoMaterial, this );
-					PlaneTranslationHandle->SetMaterial( 1, MID );
-				}
-				UMaterialInstanceDynamic* MID0 = CastChecked<UMaterialInstanceDynamic>( PlaneTranslationHandle->GetMaterial( 0 ) );
-				UMaterialInstanceDynamic* MID1 = CastChecked<UMaterialInstanceDynamic>( PlaneTranslationHandle->GetMaterial( 1 ) );
-
-				FLinearColor HandleColor = FLinearColor::White;
-				if (PlaneTranslationHandle == DraggingHandle)
-				{
-					HandleColor = VREd::GizmoColor::DraggingGizmoColor;
-				}
-				else if (FacingAxisIndex != INDEX_NONE)
-				{
-					switch (FacingAxisIndex)
-					{
-					case 0:
-						HandleColor = VREd::GizmoColor::RedGizmoColor;
-						break;
-
-					case 1:
-						HandleColor = VREd::GizmoColor::GreenGizmoColor;
-						break;
-
-					case 2:
-						HandleColor = VREd::GizmoColor::BlueGizmoColor;
-						break;
-					}
-
-					if( HoveringOverHandles.Contains( PlaneTranslationHandle ) )
-					{
-						HandleColor = FLinearColor::LerpUsingHSV( HandleColor, VREd::GizmoColor::HoverGizmoColor, Handle.HoverAlpha );
-					}
-				}
-				MID0->SetVectorParameterValue( "Color", HandleColor );
-				MID1->SetVectorParameterValue( "Color", HandleColor );
-			}
+			UpdateHandleColor( FacingAxisIndex, Handle, DraggingHandle, HoveringOverHandles );
 		}
 	}
 }
diff --git a/Engine/Source/Editor/VREditor/Gizmo/VREditorRotationGizmoHandle.cpp b/Engine/Source/Editor/VREditor/Gizmo/VREditorRotationGizmoHandle.cpp
--- a/Engine/Source/Editor/VREditor/Gizmo/VREditorRotationGizmoHandle.cpp
+++ b/Engine/Source/Editor/VREditor/Gizmo/VREditorRotationGizmoHandle.cpp
@@ -231,51 +231,7 @@ void UVREditorRotationGizmoHandleGroup::UpdateGizmoHandleGroup( const FTransform
 			HandlePlacement.GetCenterHandleCountAndFacingAxisIndex( /* Out */ CenterHandleCount, /* Out */ FacingAxisIndex, /* Out */ CenterAxisIndex );
 
 			// Update material
-			{
-				if (!RotationHandle->GetMaterial( 0 )->IsA( UMaterialInstanceDynamic::StaticClass() ))
-				{
-					UMaterialInstanceDynamic* MID = UMaterialInstanceDynamic::Create( GizmoMaterial, this );
-					RotationHandle->SetMaterial( 0, MID );
-				}
-				if (!RotationHandle->GetMaterial( 1 )->IsA( UMaterialInstanceDynamic::StaticClass() ))
-				{
-					UMaterialInstanceDynamic* MID = UMaterialInstanceDynamic::Create( TranslucentGizmoMaterial, this );
-					RotationHandle->SetMaterial( 1, MID );
-				}
-				UMaterialInstanceDynamic* MID0 = CastChecked<UMaterialInstanceDynamic>( RotationHandle->GetMaterial( 0 ) );
-				UMaterialInstanceDynamic* MID1 = CastChecked<UMaterialInstanceDynamic>( RotationHandle->GetMaterial( 1 ) );
-
-				FLinearColor HandleColor = FLinearColor::White;
-				if (RotationHandle == DraggingHandle)
-				{
-					HandleColor = VREd::GizmoColor::DraggingGizmoColor;
-				}
-				else
-				{
-					check( CenterAxisIndex != INDEX_NONE );
-					switch (CenterAxisIndex)
-					{
-					case 0:
-						HandleColor = VREd::GizmoColor::RedGizmoColor;
-						break;
-
-					case 1:
-						HandleColor = VREd::GizmoColor::GreenGizmoColor;
-						break;
-
-					case 2:
-						HandleColor = VREd::GizmoColor::BlueGizmoColor;
-						break;
-					}
-
-					if( HoveringOverHandles.Contains( RotationHandle ) )
-					{
-						HandleColor = FLinearColor::LerpUsingHSV( HandleColor, VREd::GizmoColor::HoverGizmoColor, Handle.HoverAlpha );
-					}
-				}
-				MID0->SetVectorParameterValue( "Color", HandleColor );
-				MID1->SetVectorParameterValue( "Color", HandleColor );
-			}
+			UpdateHandleColor( CenterAxisIndex, Handle, DraggingHandle, HoveringOverHandles);
 		}
 	}
 }
diff --git a/Engine/Source/Editor/VREditor/Gizmo/VREditorStretchGizmoHandle.cpp b/Engine/Source/Editor/VREditor/Gizmo/VREditorStretchGizmoHandle.cpp
--- a/Engine/Source/Editor/VREditor/Gizmo/VREditorStretchGizmoHandle.cpp
+++ b/Engine/Source/Editor/VREditor/Gizmo/VREditorStretchGizmoHandle.cpp
@@ -249,53 +249,7 @@ void UVREditorStretchGizmoHandleGroup::UpdateGizmoHandleGroup( const FTransform&
 			StretchingHandle->SetRelativeScale3D( FVector( GizmoHandleScale ) );
 
 			// Update material
-			{
-				if (!StretchingHandle->GetMaterial( 0 )->IsA( UMaterialInstanceDynamic::StaticClass() ))
-				{
-					UMaterialInstanceDynamic* MID = UMaterialInstanceDynamic::Create( GizmoMaterial, this );
-					StretchingHandle->SetMaterial( 0, MID );
-				}
-				if (!StretchingHandle->GetMaterial( 1 )->IsA( UMaterialInstanceDynamic::StaticClass() ))
-				{
-					UMaterialInstanceDynamic* MID = UMaterialInstanceDynamic::Create( TranslucentGizmoMaterial, this );
-					StretchingHandle->SetMaterial( 1, MID );
-				}
-				UMaterialInstanceDynamic* MID0 = CastChecked<UMaterialInstanceDynamic>( StretchingHandle->GetMaterial( 0 ) );
-				UMaterialInstanceDynamic* MID1 = CastChecked<UMaterialInstanceDynamic>( StretchingHandle->GetMaterial( 1 ) );
-
-				FLinearColor HandleColor = FLinearColor::White;
-				if (StretchingHandle == DraggingHandle)
-				{
-					HandleColor = VREd::GizmoColor::DraggingGizmoColor;
-				}
-				else 
-				{
-					if (FacingAxisIndex != INDEX_NONE)
-					{
-						switch (FacingAxisIndex)
-						{
-						case 0:
-							HandleColor = VREd::GizmoColor::RedGizmoColor;
-							break;
-
-						case 1:
-							HandleColor = VREd::GizmoColor::GreenGizmoColor;
-							break;
-
-						case 2:
-							HandleColor = VREd::GizmoColor::BlueGizmoColor;
-							break;
-						}
-					}
-
-					if( HoveringOverHandles.Contains( StretchingHandle ) )
-					{
-						HandleColor = FLinearColor::LerpUsingHSV( HandleColor, VREd::GizmoColor::HoverGizmoColor, Handle.HoverAlpha );
-					}
-				}
-				MID0->SetVectorParameterValue( "Color", HandleColor );
-				MID1->SetVectorParameterValue( "Color", HandleColor );
-			}
+			UpdateHandleColor( FacingAxisIndex, Handle, DraggingHandle, HoveringOverHandles );
 		}
 	}
 }
diff --git a/Engine/Source/Editor/VREditor/Gizmo/VREditorTranslationGizmoHandle.cpp b/Engine/Source/Editor/VREditor/Gizmo/VREditorTranslationGizmoHandle.cpp
--- a/Engine/Source/Editor/VREditor/Gizmo/VREditorTranslationGizmoHandle.cpp
+++ b/Engine/Source/Editor/VREditor/Gizmo/VREditorTranslationGizmoHandle.cpp
@@ -134,50 +134,7 @@ void UVREditorTranslationGizmoHandleGroup::UpdateGizmoHandleGroup( const FTransf
 			TranslationHandle->SetRelativeScale3D( FVector( GizmoHandleScale ) );
 
 			// Update material
-			{
-				if (!TranslationHandle->GetMaterial( 0 )->IsA( UMaterialInstanceDynamic::StaticClass() ))
-				{
-					UMaterialInstanceDynamic* MID = UMaterialInstanceDynamic::Create( GizmoMaterial, this );
-					TranslationHandle->SetMaterial( 0, MID );
-				}
-				if (!TranslationHandle->GetMaterial( 1 )->IsA( UMaterialInstanceDynamic::StaticClass() ))
-				{
-					UMaterialInstanceDynamic* MID = UMaterialInstanceDynamic::Create( TranslucentGizmoMaterial, this );
-					TranslationHandle->SetMaterial( 1, MID );
-				}
-				UMaterialInstanceDynamic* MID0 = CastChecked<UMaterialInstanceDynamic>( TranslationHandle->GetMaterial( 0 ) );
-				UMaterialInstanceDynamic* MID1 = CastChecked<UMaterialInstanceDynamic>( TranslationHandle->GetMaterial( 1 ) );
-
-				FLinearColor HandleColor = FLinearColor::White;
-				if (TranslationHandle == DraggingHandle)
-				{
-					HandleColor = VREd::GizmoColor::DraggingGizmoColor;
-				}
-				else if (FacingAxisIndex != INDEX_NONE)
-				{
-					switch (FacingAxisIndex)
-					{
-					case 0:
-						HandleColor = VREd::GizmoColor::RedGizmoColor;
-						break;
-
-					case 1:
-						HandleColor = VREd::GizmoColor::GreenGizmoColor;
-						break;
-
-					case 2:
-						HandleColor = VREd::GizmoColor::BlueGizmoColor;
-						break;
-					}
-
-					if( HoveringOverHandles.Contains( TranslationHandle ) )
-					{
-						HandleColor = FLinearColor::LerpUsingHSV( HandleColor, VREd::GizmoColor::HoverGizmoColor, Handle.HoverAlpha );
-					}
-				}
-				MID0->SetVectorParameterValue( "Color", HandleColor );
-				MID1->SetVectorParameterValue( "Color", HandleColor );
-			}
+			UpdateHandleColor( FacingAxisIndex, Handle, DraggingHandle, HoveringOverHandles );
 		}
 	}
 }
diff --git a/Engine/Source/Editor/VREditor/Gizmo/VREditorUniformScaleGizmoHandle.cpp b/Engine/Source/Editor/VREditor/Gizmo/VREditorUniformScaleGizmoHandle.cpp
--- a/Engine/Source/Editor/VREditor/Gizmo/VREditorUniformScaleGizmoHandle.cpp
+++ b/Engine/Source/Editor/VREditor/Gizmo/VREditorUniformScaleGizmoHandle.cpp
@@ -2,7 +2,8 @@
 
 #include "VREditorModule.h"
 #include "VREditorUniformScaleGizmoHandle.h"
-#include "VREditorTransformGizmo.h"
+#include "VREditorBaseTransformGizmo.h"
+#include "VREditorMode.h"
 #include "UnitConversion.h"
 
 UVREditorUniformScaleGizmoHandleGroup::UVREditorUniformScaleGizmoHandleGroup()
@@ -74,23 +75,31 @@ void UVREditorUniformScaleGizmoHandleGroup::UpdateGizmoHandleGroup( const FTrans
 			UMaterialInstanceDynamic* MID0 = CastChecked<UMaterialInstanceDynamic>( UniformScaleHandle->GetMaterial( 0 ) );
 			UMaterialInstanceDynamic* MID1 = CastChecked<UMaterialInstanceDynamic>( UniformScaleHandle->GetMaterial( 1 ) );
 
-			FLinearColor HandleColor = FLinearColor::White;
-			if (UniformScaleHandle == DraggingHandle)
+			ABaseTransformGizmo* GizmoActor = CastChecked<ABaseTransformGizmo>( GetOwner() );
+			if( GizmoActor )
 			{
-				HandleColor = VREd::GizmoColor::DraggingGizmoColor;
-			}
-			else
-			{
-				HandleColor = VREd::GizmoColor::WhiteGizmoColor;
-
-				if( HoveringOverHandles.Contains( UniformScaleHandle ) )
+				FVREditorMode* Mode =GizmoActor->GetOwnerMode();
+				if( Mode )
 				{
-					HandleColor = FLinearColor::LerpUsingHSV( HandleColor, VREd::GizmoColor::HoverGizmoColor, Handle.HoverAlpha );
+					FLinearColor HandleColor = Mode->GetColor( FVREditorMode::EColors::WhiteGizmoColor );
+					if (UniformScaleHandle == DraggingHandle)
+					{
+						HandleColor = Mode->GetColor( FVREditorMode::EColors::DraggingGizmoColor );
+					}
+					else
+					{
+						HandleColor = VREd::GizmoColor::WhiteGizmoColor;
+
+						if (HoveringOverHandles.Contains( UniformScaleHandle ))
+						{
+							HandleColor = FLinearColor::LerpUsingHSV( HandleColor, Mode->GetColor( FVREditorMode::EColors::HoverGizmoColor ), Handle.HoverAlpha );
+						}
+					}
+
+					MID0->SetVectorParameterValue( "Color", HandleColor );
+					MID1->SetVectorParameterValue( "Color", HandleColor );
 				}
 			}
-
-			MID0->SetVectorParameterValue( "Color", HandleColor );
-			MID1->SetVectorParameterValue( "Color", HandleColor );
 		}
 	}
 }
diff --git a/Engine/Source/Editor/VREditor/VREditorDockableWindow.cpp b/Engine/Source/Editor/VREditor/VREditorDockableWindow.cpp
--- a/Engine/Source/Editor/VREditor/VREditorDockableWindow.cpp
+++ b/Engine/Source/Editor/VREditor/VREditorDockableWindow.cpp
@@ -26,7 +26,7 @@ AVREditorDockableWindow::AVREditorDockableWindow()
 			check( WindowMesh != nullptr );
 		}
 
-		WindowMeshComponent = CreateDefaultSubobject<UStaticMeshComponent>( TEXT( "SelectionBar" ) );
+		WindowMeshComponent = CreateDefaultSubobject<UStaticMeshComponent>( TEXT( "WindowMesh" ) );
 		WindowMeshComponent->SetStaticMesh( WindowMesh );
 		WindowMeshComponent->SetMobility( EComponentMobility::Movable );
 		WindowMeshComponent->AttachTo( RootComponent );
@@ -64,6 +64,29 @@ AVREditorDockableWindow::AVREditorDockableWindow()
 		SelectionMeshComponent->bCastDynamicShadow = false;
 		SelectionMeshComponent->bCastStaticShadow = false;
 		SelectionMeshComponent->bAffectDistanceFieldLighting = false;
+
+		UMaterial* HoverMaterial = nullptr;
+		{
+			static ConstructorHelpers::FObjectFinder<UMaterial> ObjectFinder( TEXT( "/Engine/VREditor/TransformGizmo/TransformGizmoMaterial" ) );
+			HoverMaterial = ObjectFinder.Object;
+			check( HoverMaterial != nullptr );
+		}
+		HoverMaterialMID = UMaterialInstanceDynamic::Create( HoverMaterial, GetTransientPackage() );
+		check( HoverMaterialMID != nullptr );
+		SelectionMeshComponent->SetMaterial( 0, HoverMaterialMID );
+
+		UMaterial* TranslucentHoverMaterial = nullptr;
+		{
+			static ConstructorHelpers::FObjectFinder<UMaterial> ObjectFinder( TEXT( "/Engine/VREditor/TransformGizmo/TranslucentTransformGizmoMaterial" ) );
+			TranslucentHoverMaterial = ObjectFinder.Object;
+			check( TranslucentHoverMaterial != nullptr );
+		}
+		TranslucentHoverMID = UMaterialInstanceDynamic::Create( TranslucentHoverMaterial, GetTransientPackage() );
+		check( TranslucentHoverMID != nullptr );
+		SelectionMeshComponent->SetMaterial( 1, TranslucentHoverMID );
+
+		// Can't get the owner here to set the default color
+		SetSelectionBarColor( FLinearColor( 0.7f, 0.7f, 0.7f, 1.0f ) );
 	}
 
 }
@@ -101,3 +124,23 @@ void AVREditorDockableWindow::UpdateRelativeRoomTransform()
 	RelativeOffset = RoomSpaceWindowLocation;
 	LocalRotation = RoomSpaceWindowRotation.Rotator();
 }
+
+void AVREditorDockableWindow::OnEnterHover( const FHitResult& HitResult )
+{
+	if( SelectionMeshComponent == HitResult.GetComponent() )
+	{
+		SetSelectionBarColor( GetOwner().GetOwner().GetColor( FVREditorMode::EColors::HoverGizmoColor ) );
+	}
+}
+
+void AVREditorDockableWindow::OnLeaveHover()
+{
+	SetSelectionBarColor( GetOwner().GetOwner().GetColor( FVREditorMode::EColors::WhiteGizmoColor ) );
+}
+
+void AVREditorDockableWindow::SetSelectionBarColor( const FLinearColor& LinearColor )
+{
+	static FName StaticColorParameterName( "Color" );
+	HoverMaterialMID->SetVectorParameterValue( StaticColorParameterName, LinearColor );
+	TranslucentHoverMID->SetVectorParameterValue( StaticColorParameterName, LinearColor );
+}
\ No newline at end of file
diff --git a/Engine/Source/Editor/VREditor/VREditorDockableWindow.h b/Engine/Source/Editor/VREditor/VREditorDockableWindow.h
--- a/Engine/Source/Editor/VREditor/VREditorDockableWindow.h
+++ b/Engine/Source/Editor/VREditor/VREditorDockableWindow.h
@@ -32,7 +32,25 @@ public:
 	/** Updates the last dragged relative position */
 	void UpdateRelativeRoomTransform();
 
+	/** Enter hover with laser changes the color of SelectionMesh */
+	void OnEnterHover( const FHitResult& HitResult );
+
+	/** Leaving hover with laser changes the color of SelectionMesh */
+	void OnLeaveHover();
+
 private:
+
+	/** Setting the color on the dynamic materials of the SelectionMesh */
+	void SetSelectionBarColor( const FLinearColor& LinearColor );
+
 	/** The room space scale for the SelectionMeshComponent */
 	FVector SelectionMeshScale;
+
+	/** Dynamic material  (opaque) */
+	UPROPERTY()
+	class UMaterialInstanceDynamic* HoverMaterialMID;
+
+	/** Dynamic material (translucent) */
+	UPROPERTY()
+	class UMaterialInstanceDynamic* TranslucentHoverMID;
 };	
\ No newline at end of file
diff --git a/Engine/Source/Editor/VREditor/VREditorMode.cpp b/Engine/Source/Editor/VREditor/VREditorMode.cpp
--- a/Engine/Source/Editor/VREditor/VREditorMode.cpp
+++ b/Engine/Source/Editor/VREditor/VREditorMode.cpp
@@ -197,10 +197,25 @@ void FVREditorMode::Enter()
 		KeyToActionMap.Add( EKeys::MotionController_Right_Thumbstick_Y, FVRAction( EVRActionType::TrackpadPositionY, VREditorConstants::RightHandIndex ) );
 	}
 
+	// Setting up colors
+	Colors.SetNumZeroed( (int32)EColors::TotalCount );
+	{
+		Colors[ (int32)EColors::DefaultColor ] = FLinearColor::Red;	
+		Colors[ (int32)EColors::SelectionColor ] = FLinearColor::Green;
+		Colors[ (int32)EColors::TeleportColor ] = FLinearColor( 1.0f, 0.0f, 0.75f, 1.0f );
+		Colors[ (int32)EColors::WorldDraggingColor_OneHand ] = FLinearColor::Yellow;
+		Colors[ (int32)EColors::WorldDraggingColor_TwoHands ] = FLinearColor::Blue;
+		Colors[ (int32)EColors::RedGizmoColor ] = FLinearColor( 0.4f, 0.05f, 0.05f, 1.0f );
+		Colors[ (int32)EColors::GreenGizmoColor ] = FLinearColor( 0.05f, 0.4f, 0.05f, 1.0f );
+		Colors[ (int32)EColors::BlueGizmoColor ] = FLinearColor( 0.05f, 0.05f, 0.4f, 1.0f );
+		Colors[ (int32)EColors::WhiteGizmoColor ] = FLinearColor( 0.7f, 0.7f, 0.7f, 1.0f );
+		Colors[ (int32)EColors::HoverGizmoColor ] = FLinearColor( FLinearColor::Yellow );
+		Colors[ (int32)EColors::DraggingGizmoColor ] = FLinearColor( FLinearColor::White );
+	}
 
 	// @todo vreditor: We need to make sure the user can never switch to orthographic mode, or activate settings that
 	// would disrupt the user's ability to view the VR scene.
-
+		
 	// @todo vreditor: Don't bother drawing toolbars in VR, or other things that won't matter in VR
 
 	{
@@ -605,6 +620,7 @@ void FVREditorMode::SpawnAvatarMeshActor()
 				UMaterialInstance* TranslucentLaserPointerMaterialInst = LoadObject<UMaterialInstance>( nullptr, TEXT( "/Engine/VREditor/LaserPointer/TranslucentLaserPointerMaterialInst" ) );
 				check( TranslucentLaserPointerMaterialInst != nullptr );
 				Hand.TranslucentLaserPointerMID = UMaterialInstanceDynamic::Create( TranslucentLaserPointerMaterialInst, GetTransientPackage() );
+				check( Hand.TranslucentLaserPointerMID != nullptr );
 				Hand.LaserPointerMeshComponent->SetMaterial( 1, Hand.TranslucentLaserPointerMID );
 			}
 
@@ -984,7 +1000,7 @@ void FVREditorMode::Tick( FEditorViewportClient* ViewportClient, float DeltaTime
 					const bool CrawlFade = 1.0f;
 					const float CrawlSpeed = 10.0f;
 
-					SetLaserVisuals( HandIndex, VREd::ModeColors::DefaultColor );
+					SetLaserVisuals( HandIndex, GetColor( EColors::DefaultColor ) );
 
 					static FName StaticLengthParameterName( "Length" );
 					Hand.LaserPointerMID->SetScalarParameterValue( StaticLengthParameterName, LaserPointerLength );
@@ -1077,7 +1093,7 @@ void FVREditorMode::Tick( FEditorViewportClient* ViewportClient, float DeltaTime
 	{
 		for( int32 HandIndex = 0; HandIndex < VREditorConstants::NumVirtualHands; ++HandIndex )
 		{
-			FLinearColor ResultColor = VREd::ModeColors::DefaultColor;
+			FLinearColor ResultColor = GetColor( EColors::DefaultColor );
 			float CrawlSpeed = 0.0f;
 			float CrawlFade = 0.0f;
 
@@ -1090,20 +1106,20 @@ void FVREditorMode::Tick( FEditorViewportClient* ViewportClient, float DeltaTime
 
 			if( bIsDraggingWorldWithTwoHands )
 			{
-				ResultColor = VREd::ModeColors::WorldDraggingColor_TwoHands;
+				ResultColor = GetColor( EColors::WorldDraggingColor_TwoHands );
 			}
 			else if( Hand.DraggingMode == EVREditorDraggingMode::World )
 			{
 				if( WorldInteraction->IsTeleporting() )
 				{
-					ResultColor = VREd::ModeColors::TeleportColor;
+					ResultColor = GetColor( EColors::TeleportColor );
 				}
 				else
 				{
 					// We can teleport in this mode, so animate the laser a bit
 					CrawlFade = 1.0f;
 					CrawlSpeed = 5.0f;
-					ResultColor = VREd::ModeColors::WorldDraggingColor_OneHand;
+					ResultColor = GetColor( EColors::WorldDraggingColor_OneHand );
 				}
 			}
 			else if( Hand.DraggingMode == EVREditorDraggingMode::ActorsAtLaserImpact ||
@@ -1111,7 +1127,7 @@ void FVREditorMode::Tick( FEditorViewportClient* ViewportClient, float DeltaTime
 				Hand.DraggingMode == EVREditorDraggingMode::ActorsWithGizmo ||
 				Hand.DraggingMode == EVREditorDraggingMode::AssistingDrag )
 			{
-				ResultColor = VREd::ModeColors::SelectionColor;
+				ResultColor = GetColor( EColors::SelectionColor );
 			}
 
 			SetLaserVisuals( HandIndex, ResultColor, CrawlFade, CrawlSpeed );
@@ -2580,5 +2596,10 @@ EHMDDeviceType::Type FVREditorMode::GetHMDDeviceType() const
 	return GEngine->HMDDevice.IsValid() ? GEngine->HMDDevice->GetHMDDeviceType() : EHMDDeviceType::DT_SteamVR;
 }
 
+FLinearColor FVREditorMode::GetColor( const EColors Color ) const
+{
+	return Colors[ (int32)Color ];
+}
+
 
 #undef LOCTEXT_NAMESPACE
\ No newline at end of file
diff --git a/Engine/Source/Editor/VREditor/VREditorMode.h b/Engine/Source/Editor/VREditor/VREditorMode.h
--- a/Engine/Source/Editor/VREditor/VREditorMode.h
+++ b/Engine/Source/Editor/VREditor/VREditorMode.h
@@ -10,15 +10,6 @@
 
 namespace VREd
 {
-	namespace ModeColors
-	{
-		static const FLinearColor DefaultColor = FLinearColor::Red;
-		static const FLinearColor SelectionColor = FLinearColor::Green;
-		static const FLinearColor TeleportColor = FLinearColor( 1.0f, 0.0f, 0.75f, 1.0f );
-		static const FLinearColor WorldDraggingColor_OneHand = FLinearColor::Yellow;	// With one hand, a laser is visible and you can teleport by pulling the trigger
-		static const FLinearColor WorldDraggingColor_TwoHands = FLinearColor::Blue;		// With two hands you are freely rotating/scaling the world, no laser is visible, and you can't teleport
-	}
-
 	static FAutoConsoleVariable ScaleMax( TEXT( "VREd.ScaleMax" ), 4000.0f, TEXT( "Maximum world scale in centimeters" ) );
 	static FAutoConsoleVariable ScaleMin( TEXT( "VREd.ScaleMin" ), 10.0f, TEXT( "Minimum world scale in centimeters" ) );
 }
@@ -499,5 +490,34 @@ protected:
 
 	/** The current Gizmo type that is used for the TransformGizmo Actor */
 	EGizmoHandleTypes CurrentGizmoType;
+
+	//
+	// Colors
+	//
+public:
+	// Color types
+	enum EColors
+	{
+		DefaultColor,
+		SelectionColor,
+		TeleportColor,
+		WorldDraggingColor_OneHand,
+		WorldDraggingColor_TwoHands,
+		RedGizmoColor,
+		GreenGizmoColor,
+		BlueGizmoColor,
+		WhiteGizmoColor,
+		HoverGizmoColor,
+		DraggingGizmoColor,
+		TotalCount
+	};
+
+	// Gets the color
+	FLinearColor GetColor( const EColors Color ) const;
+
+private:
+
+	// All the colors for this mode
+	TArray<FLinearColor> Colors;
 };
 
diff --git a/Engine/Source/Editor/VREditor/VREditorWorldInteraction.cpp b/Engine/Source/Editor/VREditor/VREditorWorldInteraction.cpp
--- a/Engine/Source/Editor/VREditor/VREditorWorldInteraction.cpp
+++ b/Engine/Source/Editor/VREditor/VREditorWorldInteraction.cpp
@@ -658,9 +658,10 @@ void FVREditorWorldInteraction::OnVRHoverUpdate( FEditorViewportClient& Viewport
 		FVirtualHand& Hand = Owner.GetVirtualHand( HandIndex );
 
 		UActorComponent* PreviousHoverGizmoComponent = Hand.HoveringOverTransformGizmoComponent.Get();
-
 		Hand.HoveringOverTransformGizmoComponent = nullptr;
 
+		UActorComponent* NewHoveredActorComponent = nullptr;
+
 		FHitResult HitResult = Owner.GetHitResultFromLaserPointer( HandIndex );
 		if( HitResult.Actor.IsValid() )
 		{
@@ -671,6 +672,7 @@ void FVREditorWorldInteraction::OnVRHoverUpdate( FEditorViewportClient& Viewport
 				this->HoveredObjects.Add( FViewportHoverTarget( Actor ) );
 
 				HoverImpactPoint = HitResult.ImpactPoint;
+				NewHoveredActorComponent = HitResult.GetComponent();
 
 				if( Actor == TransformGizmoActor )
 				{
@@ -696,10 +698,38 @@ void FVREditorWorldInteraction::OnVRHoverUpdate( FEditorViewportClient& Viewport
 
 					Hand.HoverHapticCheckLastHoveredGizmoComponent = HitResult.GetComponent();
 				}
+				else if ( Actor->IsA( AVREditorDockableWindow::StaticClass() ) )
+				{
+					AVREditorDockableWindow* DockableWindow = Cast< AVREditorDockableWindow >( Actor );
+					if (DockableWindow)
+					{
+						DockableWindow->OnEnterHover( HitResult );
+					}
+				}
 
 				bWasHandled = true;
 			}
 		}
+
+		// Leave dockable window hover
+		if( Hand.HoveredActorComponent != nullptr && ( NewHoveredActorComponent == nullptr || Hand.HoveredActorComponent != NewHoveredActorComponent || !HitResult.Actor.IsValid() ) )
+		{
+			AActor* Actor = Hand.HoveredActorComponent->GetOwner();
+			if (Actor->IsA( AVREditorDockableWindow::StaticClass() ))
+			{
+				AVREditorDockableWindow* DockableWindow = Cast< AVREditorDockableWindow >( Actor );
+				if (DockableWindow)
+				{
+					DockableWindow->OnLeaveHover();
+				}
+			}
+		}
+
+		// Set the new hovered actor component if there is any
+		if( HitResult.Actor.IsValid() && NewHoveredActorComponent )
+		{
+			Hand.HoveredActorComponent = NewHoveredActorComponent;
+		}
 	}
 }
 
@@ -1393,6 +1423,7 @@ void FVREditorWorldInteraction::RefreshTransformGizmo( const bool bNewObjectsSel
 			const bool bWithSceneComponent = false;	// We already have our own scene component
 			TransformGizmoActor = Owner.SpawnTransientSceneActor<ATransformGizmo>( TEXT( "TransformGizmo" ), bWithSceneComponent );
 			check( TransformGizmoActor != nullptr );
+			TransformGizmoActor->SetOwnerMode( &Owner ) ;
 
 			if( VREd::ShowTransformGizmo->GetInt() == 0 )
 			{
diff --git a/Engine/Source/Editor/VREditor/VirtualHand.h b/Engine/Source/Editor/VREditor/VirtualHand.h
--- a/Engine/Source/Editor/VREditor/VirtualHand.h
+++ b/Engine/Source/Editor/VREditor/VirtualHand.h
@@ -279,7 +279,8 @@ struct FVirtualHand
 	/** Time that we either started showing or hiding help labels (for fade transitions) */
 	FTimespan HelpLabelShowOrHideStartTime;
 
-
+	/** The current component hovered by the laser pointer of this hand */
+	TWeakObjectPtr<class UActorComponent> HoveredActorComponent;
 
 	/** Default constructor for FVirtualHand that initializes safe defaults */
 	FVirtualHand()
@@ -351,6 +352,8 @@ struct FVirtualHand
 
 		bWantHelpLabels = false;
 		HelpLabelShowOrHideStartTime = FTimespan::MinValue();
+
+		HoveredActorComponent = nullptr;
 	}
 };
 
