VR Editor: Re-enabled late frame updates on motion controllers

diff --git a/Engine/Source/Editor/VREditor/VREditorMode.cpp b/Engine/Source/Editor/VREditor/VREditorMode.cpp
--- a/Engine/Source/Editor/VREditor/VREditorMode.cpp
+++ b/Engine/Source/Editor/VREditor/VREditorMode.cpp
@@ -559,7 +559,7 @@ void FVREditorMode::SpawnAvatarMeshActor()
 				// the top of every frame, potentially clobbering what the RT uses before rendering (popping).  On the Vive,
 				// controller state is cached by the HMD's late update code, then re-used by the game thread later, however
 				// the state can be clobbered either before or after the game thread reads it (popping.)
-				Hand.MotionControllerComponent->bDisableLowLatencyUpdate = true;
+				Hand.MotionControllerComponent->bDisableLowLatencyUpdate = false;
 			}
 
 			// Hand mesh
