VR Editor: Fixed bugs with trackpad-based movement and scroll direction

diff --git a/Engine/Source/Editor/VREditor/VREditorMode.cpp b/Engine/Source/Editor/VREditor/VREditorMode.cpp
--- a/Engine/Source/Editor/VREditor/VREditorMode.cpp
+++ b/Engine/Source/Editor/VREditor/VREditorMode.cpp
@@ -1479,13 +1479,6 @@ bool FVREditorMode::InputAxis( FEditorViewportClient* ViewportClient, FViewport*
 				Hand.LastTrackpadPosition.Y = Hand.bIsTrackpadPositionValid[1] ? Hand.TrackpadPosition.Y : Delta;
 				Hand.LastTrackpadPositionUpdateTime = FTimespan::FromSeconds( FPlatformTime::Seconds() );
 				Hand.TrackpadPosition.Y = Delta;
-
-				// @todo vreditor: On the Rift, trackpad vertical values are reversed
-				if( GetHMDDeviceType() == EHMDDeviceType::DT_OculusRift )
-				{
-					Hand.TrackpadPosition.Y = -Hand.TrackpadPosition.Y;
-				}
-
 				Hand.bIsTrackpadPositionValid[1] = true;
 			}
 		}
diff --git a/Engine/Source/Editor/VREditor/VREditorWorldInteraction.cpp b/Engine/Source/Editor/VREditor/VREditorWorldInteraction.cpp
--- a/Engine/Source/Editor/VREditor/VREditorWorldInteraction.cpp
+++ b/Engine/Source/Editor/VREditor/VREditorWorldInteraction.cpp
@@ -794,20 +794,18 @@ void FVREditorWorldInteraction::Tick( FEditorViewportClient* ViewportClient, con
 				// Make sure they are touching the trackpad, otherwise we get bad data
 				if( Hand.bIsTrackpadPositionValid[ 1 ] )
 				{
-					const bool bIsAbsolute = ( Owner.GetHMDDeviceType() == EHMDDeviceType::DT_OculusRift );
+					const bool bIsAbsolute = ( Owner.GetHMDDeviceType() == EHMDDeviceType::DT_SteamVR );
 
-					// @todo vreditor: Instead of using only the vertical axis of the pad, we should take the whole swipe direction into account
-					// @todo vreditor: This thing could use a tiny bit of smoothing.  It's not particularly accurate.
 					float SlideDelta = 0.0f;
-					if( Hand.bIsTouchingTrackpad || Owner.GetHMDDeviceType() != EHMDDeviceType::DT_SteamVR )
+					if( Hand.bIsTouchingTrackpad || !bIsAbsolute )
 					{
 						if( bIsAbsolute )
 						{
-							SlideDelta = ( Hand.TrackpadPosition.Y * VREd::TrackpadRelativeDragSpeed->GetFloat() ) * WorldScaleFactor;
+							SlideDelta = -( ( Hand.TrackpadPosition.Y - Hand.LastTrackpadPosition.Y ) * VREd::TrackpadAbsoluteDragSpeed->GetFloat() ) * WorldScaleFactor;
 						}
 						else
 						{
-							SlideDelta = ( ( Hand.TrackpadPosition.Y - Hand.LastTrackpadPosition.Y ) * VREd::TrackpadAbsoluteDragSpeed->GetFloat() ) * WorldScaleFactor;
+							SlideDelta = -( Hand.TrackpadPosition.Y * VREd::TrackpadRelativeDragSpeed->GetFloat() ) * WorldScaleFactor;
 						}
 					}
 
