VR Editor: Shrunk the Quick Menu UI a little bit

diff --git a/Engine/Source/Editor/VREditor/VREditorUISystem.cpp b/Engine/Source/Editor/VREditor/VREditorUISystem.cpp
--- a/Engine/Source/Editor/VREditor/VREditorUISystem.cpp
+++ b/Engine/Source/Editor/VREditor/VREditorUISystem.cpp
@@ -36,7 +36,7 @@ namespace VREd
 	static FAutoConsoleVariable EditorUIResolutionX( TEXT( "VREd.EditorUIResolutionX" ), 1024, TEXT( "Horizontal resolution to use for VR editor UI render targets" ) );
 	static FAutoConsoleVariable EditorUIResolutionY( TEXT( "VREd.EditorUIResolutionY" ), 1024, TEXT( "Vertical resolution to use for VR editor UI render targets" ) );
 	static FAutoConsoleVariable QuickMenuUIResolutionX( TEXT( "VREd.QuickMenuUIResolutionX" ), 1024, TEXT( "Horizontal resolution to use for Quick Menu VR UI render targets" ) );
-	static FAutoConsoleVariable QuickMenuUIResolutionY( TEXT( "VREd.QuickMenuUIResolutionY" ), 1300, TEXT( "Vertical resolution to use for Quick Menu VR UI render targets" ) );
+	static FAutoConsoleVariable QuickMenuUIResolutionY( TEXT( "VREd.QuickMenuUIResolutionY" ), 1130, TEXT( "Vertical resolution to use for Quick Menu VR UI render targets" ) );
 	static FAutoConsoleVariable ContentBrowserUISize( TEXT( "VREd.ContentBrowserUISize" ), 60.0f, TEXT( "How big content browser UIs should be" ) );
 	static FAutoConsoleVariable EditorUISize( TEXT( "VREd.EditorUISize" ), 50.0f, TEXT( "How big editor UIs should be" ) );
 	static FAutoConsoleVariable ContentBrowserUIScale( TEXT( "VREd.ContentBrowserUIScale" ), 2.5f, TEXT( "How much to scale up (or down) the content browser for VR" ) );
@@ -462,7 +462,7 @@ void FVREditorUISystem::Tick( FEditorViewportClient* ViewportClient, const float
 				// @todo vreditor tweak: Weird to hard code this here.  Probably should be an accessor on the hand itself, and based on the actual device type
 				const float UICapsuleLocalRadius = 6.0f * WorldScaleFactor;
 				const FVector UICapsuleStart = FVector( -6.0f, 0.0f, 0.0f ) * WorldScaleFactor;
-				const FVector UICapsuleEnd = FVector( -16.0f, 0.0f, 0.0f ) * WorldScaleFactor;
+				const FVector UICapsuleEnd = FVector( -14.0f, 0.0f, 0.0f ) * WorldScaleFactor;
 				FVector ClosestPointOnLaserPointer, ClosestPointOnUICapsule;
 				FMath::SegmentDistToSegment(
 					LaserPointerStartInOtherHandSpace, LaserPointerEndInOtherHandSpace,
