VR Editor:
- Fixed: Deselecting all actors when one hand with ui in front presses the trigger
- Fixed: User can't select and drag the same dockwindow if already dragging with the other hand
- Fixed: Calling on hover leave and enter if the dockwindow selection bar component is already selected by the other hand

diff --git a/Engine/Source/Editor/VREditor/VREditorWorldInteraction.cpp b/Engine/Source/Editor/VREditor/VREditorWorldInteraction.cpp
--- a/Engine/Source/Editor/VREditor/VREditorWorldInteraction.cpp
+++ b/Engine/Source/Editor/VREditor/VREditorWorldInteraction.cpp
@@ -357,6 +357,12 @@ void FVREditorWorldInteraction::OnVRAction( FEditorViewportClient& ViewportClien
 					bJustTeleported = false;
 				}
 
+				// Nothing should happen when the hand has a UI in front, otherwise it will deselect everything
+				if( Hand.bHasUIInFront )
+				{
+					bWasHandled = true;
+				}
+
 				// No clicking while we're dragging the world.  (No laser pointers are visible, anyway.)
 				const bool bIsDraggingWorldWithTwoHands =
 					( Hand.DraggingMode == EVREditorDraggingMode::World && OtherHand.DraggingMode == EVREditorDraggingMode::AssistingDrag ) ||
@@ -377,7 +383,7 @@ void FVREditorWorldInteraction::OnVRAction( FEditorViewportClient& ViewportClien
 
 						if (Actor->IsA( AVREditorDockableWindow::StaticClass() ))
 						{
-							if( VRAction.ActionType == EVRActionType::SelectAndMove )
+							if( VRAction.ActionType == EVRActionType::SelectAndMove && Owner.GetUISystem().GetDraggingDockUI() == nullptr )
 							{
 								AVREditorDockableWindow* DockableWindow = Cast< AVREditorDockableWindow >( Actor );
 								if (DockableWindow)
@@ -656,6 +662,7 @@ void FVREditorWorldInteraction::OnVRHoverUpdate( FEditorViewportClient& Viewport
 	if( !bWasHandled )
 	{
 		FVirtualHand& Hand = Owner.GetVirtualHand( HandIndex );
+		FVirtualHand& OtherHand = Owner.GetOtherHand( HandIndex );
 
 		UActorComponent* PreviousHoverGizmoComponent = Hand.HoveringOverTransformGizmoComponent.Get();
 		Hand.HoveringOverTransformGizmoComponent = nullptr;
@@ -701,7 +708,7 @@ void FVREditorWorldInteraction::OnVRHoverUpdate( FEditorViewportClient& Viewport
 				else if ( Actor->IsA( AVREditorDockableWindow::StaticClass() ) )
 				{
 					AVREditorDockableWindow* DockableWindow = Cast< AVREditorDockableWindow >( Actor );
-					if (DockableWindow)
+					if (DockableWindow && NewHoveredActorComponent != OtherHand.HoveredActorComponent )
 					{
 						DockableWindow->OnEnterHover( HitResult );
 					}
@@ -710,9 +717,10 @@ void FVREditorWorldInteraction::OnVRHoverUpdate( FEditorViewportClient& Viewport
 				bWasHandled = true;
 			}
 		}
-
+		
 		// Leave dockable window hover
-		if( Hand.HoveredActorComponent != nullptr && ( NewHoveredActorComponent == nullptr || Hand.HoveredActorComponent != NewHoveredActorComponent || !HitResult.Actor.IsValid() ) )
+		if( Hand.HoveredActorComponent != nullptr && Hand.HoveredActorComponent != OtherHand.HoveredActorComponent &&
+			( NewHoveredActorComponent == nullptr || Hand.HoveredActorComponent != NewHoveredActorComponent || !HitResult.Actor.IsValid() ) )
 		{
 			AActor* Actor = Hand.HoveredActorComponent->GetOwner();
 			if (Actor->IsA( AVREditorDockableWindow::StaticClass() ))
