VR Editor : - Fixed: The radial menu has a delay before hiding when not being updated for a while, so it doesn't hide when the trackpad gives incorrect values once.

diff --git a/Engine/Source/Editor/VREditor/VREditorUISystem.cpp b/Engine/Source/Editor/VREditor/VREditorUISystem.cpp
--- a/Engine/Source/Editor/VREditor/VREditorUISystem.cpp
+++ b/Engine/Source/Editor/VREditor/VREditorUISystem.cpp
@@ -61,7 +61,8 @@ FVREditorUISystem::FVREditorUISystem( FVREditorMode& InitOwner )
 	  DraggingUI( nullptr ),
 	  DraggingUIHandIndex( INDEX_NONE ),
 	  DraggingUIOffsetTransform( FTransform::Identity ),
-	  bPanelVisibilityToggle( false )
+	  bPanelVisibilityToggle( false ),
+	  RadialMenuHideDelayTime( 0.0f )
 {
 	// Register to find out about VR events
 	Owner.OnVRAction().AddRaw( this, &FVREditorUISystem::OnVRAction );
@@ -556,6 +557,15 @@ void FVREditorUISystem::Tick( FEditorViewportClient* ViewportClient, const float
 			}
 		}
 	}
+	else
+	{
+		// Close the radial menu if it was not updated for a while
+		RadialMenuHideDelayTime += DeltaTime;
+		if( RadialMenuHideDelayTime >= VREd::RadialMenuFadeDelay->GetFloat() )
+		{
+			HideRadialMenu( QuickRadialMenu->GetDockedTo() == AVREditorFloatingUI::EDockedTo::LeftHand ? VREditorConstants::LeftHandIndex : VREditorConstants::RightHandIndex );
+		}
+	}
 
 	// Tick all of our floating UIs
 	for( AVREditorFloatingUI* FloatingUIPtr : FloatingUIs )
@@ -875,6 +885,7 @@ void FVREditorUISystem::UpdateRadialMenu( const int32 HandIndex )
 
 	if(!QuickRadialMenu->bHidden)
 	{
+		RadialMenuHideDelayTime = 0.0f;
 		QuickRadialMenu->GetUserWidget<UVREditorRadialMenu>()->Update( Owner.GetVirtualHand( HandIndex ) );
 	}
 }
diff --git a/Engine/Source/Editor/VREditor/VREditorUISystem.h b/Engine/Source/Editor/VREditor/VREditorUISystem.h
--- a/Engine/Source/Editor/VREditor/VREditorUISystem.h
+++ b/Engine/Source/Editor/VREditor/VREditorUISystem.h
@@ -125,6 +125,9 @@ protected:
 	/** The Radial Menu UI */
 	AVREditorFloatingUI* QuickRadialMenu;
 
+	/** The time since the radial menu was updated */
+	float RadialMenuHideDelayTime;
+
 	/** Quick menu widget class */
 	UClass* QuickMenuWidgetClass;
 
diff --git a/Engine/Source/Editor/VREditor/VREditorWorldInteraction.cpp b/Engine/Source/Editor/VREditor/VREditorWorldInteraction.cpp
--- a/Engine/Source/Editor/VREditor/VREditorWorldInteraction.cpp
+++ b/Engine/Source/Editor/VREditor/VREditorWorldInteraction.cpp
@@ -881,11 +881,6 @@ void FVREditorWorldInteraction::Tick( FEditorViewportClient* ViewportClient, con
 						Owner.GetUISystem().UpdateRadialMenu( HandIndex );
 					}
 				}
-				else
-				{
-					// Hide the radial menu when not touching the trackpad
-					Owner.GetUISystem().HideRadialMenu( HandIndex );
-				}
 			}
 		}
 
