VR Editor: Cleaned up selection bar for dockable windows

diff --git a/Engine/Source/Editor/VREditor/VREditorDockableWindow.cpp b/Engine/Source/Editor/VREditor/VREditorDockableWindow.cpp
--- a/Engine/Source/Editor/VREditor/VREditorDockableWindow.cpp
+++ b/Engine/Source/Editor/VREditor/VREditorDockableWindow.cpp
@@ -9,19 +9,18 @@
 
 namespace VREd
 {
-	static FAutoConsoleVariable DockWindowTickness( TEXT( "VREd.DockWindowTickness" ), 0.003f, TEXT( "Z Distance between the selectionbar and the UI" ) );
-	static FAutoConsoleVariable DockUIZBarOffset( TEXT( "VREd.DockUIZBarOffset" ), 1.0f, TEXT( "Z Distance between the selectionbar and the UI" ) );
+	static FAutoConsoleVariable DockWindowThickness( TEXT( "VREd.DockWindowTickness" ), 1.0f, TEXT( "How thick the window is" ) );
+	static FAutoConsoleVariable DockUIZBarOffset( TEXT( "VREd.DockUIZBarOffset" ), 2.0f, TEXT( "Z Distance between the selectionbar and the UI" ) );
 }
 
 AVREditorDockableWindow::AVREditorDockableWindow() 
-	: Super(),
-	SelectionMeshScale( FVector(0.6f, 2.0f, 1.2f) )
+	: Super()
 {
 
 	{
 		UStaticMesh* WindowMesh = nullptr;
 		{
-			static ConstructorHelpers::FObjectFinder<UStaticMesh> ObjectFinder( TEXT( "/Engine/VREditor/SM_ContentWindow_01" ) );
+			static ConstructorHelpers::FObjectFinder<UStaticMesh> ObjectFinder( TEXT( "/Engine/VREditor/UI/SM_ContentWindow_01" ) );
 			WindowMesh = ObjectFinder.Object;
 			check( WindowMesh != nullptr );
 		}
@@ -45,12 +44,12 @@ AVREditorDockableWindow::AVREditorDockableWindow()
 	{
 		UStaticMesh* SelectionMesh = nullptr;
 		{
-			static ConstructorHelpers::FObjectFinder<UStaticMesh> ObjectFinder( TEXT( "/Engine/VREditor/TransformGizmo/UniformScaleHandle" ) );
+			static ConstructorHelpers::FObjectFinder<UStaticMesh> ObjectFinder( TEXT( "/Engine/VREditor/UI/SelectionBarMesh" ) );
 			SelectionMesh = ObjectFinder.Object;
 			check( SelectionMesh != nullptr );
 		}
 
-		SelectionMeshComponent = CreateDefaultSubobject<UStaticMeshComponent>( TEXT( "SelectionMesh" ) );
+		SelectionMeshComponent = CreateDefaultSubobject<UStaticMeshComponent>( TEXT( "SelectionBarMesh" ) );
 		SelectionMeshComponent->SetStaticMesh( SelectionMesh );
 		SelectionMeshComponent->SetMobility( EComponentMobility::Movable );
 		SelectionMeshComponent->AttachTo( RootComponent );
@@ -84,29 +83,44 @@ AVREditorDockableWindow::AVREditorDockableWindow()
 		TranslucentHoverMID = UMaterialInstanceDynamic::Create( TranslucentHoverMaterial, GetTransientPackage() );
 		check( TranslucentHoverMID != nullptr );
 		SelectionMeshComponent->SetMaterial( 1, TranslucentHoverMID );
-
-		// Can't get the owner here to set the default color
-		SetSelectionBarColor( FLinearColor( 0.7f, 0.7f, 0.7f, 1.0f ) );
 	}
+}
+
+
+void AVREditorDockableWindow::SetupWidgetComponent()
+{
+	Super::SetupWidgetComponent();
 
+	SetSelectionBarColor( GetOwner().GetOwner().GetColor( FVREditorMode::EColors::UISelectionBarColor ) );
 }
 
+
 void AVREditorDockableWindow::TickManually( float DeltaTime )
 {
 	Super::TickManually( DeltaTime );
 
 	if( WidgetComponent->IsVisible() )
 	{
+		// How big the selection bar should be
+		const FVector SelectionBarSize( 1.0f, GetSize().X * 0.8f, 6.0f );	// @todo vreditor tweak
+
+		const float WorldScaleFactor = GetOwner().GetOwner().GetWorldScaleFactor();
 		const FVector AnimatedScale = CalculateAnimatedScale();
-		const float WordScaleFactor = GetOwner().GetOwner().GetWorldScaleFactor();
-		const FVector SelectionMeshWorldScale = SelectionMeshScale * WordScaleFactor;
-		SelectionMeshComponent->SetWorldScale3D( SelectionMeshWorldScale * AnimatedScale );
-
-		const FVector WindowMeshWorldScale = FVector( VREd::DockWindowTickness->GetFloat() , 1.0 / (float)Resolution.X, 1.0f / (float)Resolution.Y  ) * ( Size *  GetOwner().GetOwner().GetWorldScaleFactor() ) * AnimatedScale;
-		WindowMeshComponent->SetWorldScale3D( WindowMeshWorldScale * 10 );	
-		
-		const FVector NewRelativeLocation = FVector( 0.f, 0.f, -(((WindowMeshWorldScale.Z * 0.5f) * 1000) + ((SelectionMeshWorldScale.Z * 0.5) * 10) + (VREd::DockUIZBarOffset->GetFloat() * WordScaleFactor)) );
-		SelectionMeshComponent->SetRelativeLocation( NewRelativeLocation );
+
+		const float WindowMeshSize = 100.0f;	// Size of imported mesh, we need to inverse compensate for
+		const FVector WindowMeshScale = FVector(
+			VREd::DockWindowThickness->GetFloat(),
+			GetSize().X / WindowMeshSize,
+			GetSize().Y / WindowMeshSize ) * AnimatedScale * WorldScaleFactor;
+		WindowMeshComponent->SetRelativeScale3D( WindowMeshScale );
+
+		const FVector SelectionMeshScale = SelectionBarSize * AnimatedScale * WorldScaleFactor;
+		SelectionMeshComponent->SetRelativeScale3D( SelectionMeshScale );
+		const FVector SelectionMeshRelativeLocation = FVector(
+			0.0f,
+			0.0f,
+			-( GetSize().Y * 0.5f + SelectionBarSize.Z * 0.5f + VREd::DockUIZBarOffset->GetFloat() ) ) * AnimatedScale * WorldScaleFactor;
+		SelectionMeshComponent->SetRelativeLocation( SelectionMeshRelativeLocation );
 	}
 } 
 
@@ -135,7 +149,7 @@ void AVREditorDockableWindow::OnEnterHover( const FHitResult& HitResult )
 
 void AVREditorDockableWindow::OnLeaveHover()
 {
-	SetSelectionBarColor( GetOwner().GetOwner().GetColor( FVREditorMode::EColors::WhiteGizmoColor ) );
+	SetSelectionBarColor( GetOwner().GetOwner().GetColor( FVREditorMode::EColors::UISelectionBarColor ) );
 }
 
 void AVREditorDockableWindow::SetSelectionBarColor( const FLinearColor& LinearColor )
diff --git a/Engine/Source/Editor/VREditor/VREditorDockableWindow.h b/Engine/Source/Editor/VREditor/VREditorDockableWindow.h
--- a/Engine/Source/Editor/VREditor/VREditorDockableWindow.h
+++ b/Engine/Source/Editor/VREditor/VREditorDockableWindow.h
@@ -25,7 +25,7 @@ public:
 
 	/** Default constructor */
 	AVREditorDockableWindow();
-
+	
 	/** Updates the meshes for the UI */
 	virtual void TickManually( float DeltaTime ) override;
 
@@ -40,12 +40,12 @@ public:
 
 private:
 
+	// AVREditorFloatingUI overrides
+	virtual void SetupWidgetComponent() override;
+
 	/** Setting the color on the dynamic materials of the SelectionMesh */
 	void SetSelectionBarColor( const FLinearColor& LinearColor );
 
-	/** The room space scale for the SelectionMeshComponent */
-	FVector SelectionMeshScale;
-
 	/** Dynamic material  (opaque) */
 	UPROPERTY()
 	class UMaterialInstanceDynamic* HoverMaterialMID;
diff --git a/Engine/Source/Editor/VREditor/VREditorFloatingUI.cpp b/Engine/Source/Editor/VREditor/VREditorFloatingUI.cpp
--- a/Engine/Source/Editor/VREditor/VREditorFloatingUI.cpp
+++ b/Engine/Source/Editor/VREditor/VREditorFloatingUI.cpp
@@ -16,7 +16,7 @@ namespace VREd
 
 
 AVREditorFloatingUI::AVREditorFloatingUI()
-	: Size( 1.0f ),
+	: Scale( 1.0f ),
 	  DockedTo( EDockedTo::Nothing ),
 	  PreviousDockedTo( EDockedTo::Nothing ),
 	  SlateWidget( nullptr ),
@@ -94,7 +94,7 @@ void AVREditorFloatingUI::SetupWidgetComponent()
 	UpdateTransformIfDocked();
 }
 
-void AVREditorFloatingUI::SetSlateWidget( FVREditorUISystem& InitOwner, const TSharedRef<SWidget>& InitSlateWidget, const FIntPoint InitResolution, const float InitSize, const EDockedTo InitDockedTo )
+void AVREditorFloatingUI::SetSlateWidget( FVREditorUISystem& InitOwner, const TSharedRef<SWidget>& InitSlateWidget, const FIntPoint InitResolution, const float InitScale, const EDockedTo InitDockedTo )
 {
 	Owner = &InitOwner;
 
@@ -103,14 +103,14 @@ void AVREditorFloatingUI::SetSlateWidget( FVREditorUISystem& InitOwner, const TS
 	Resolution = InitResolution;
 	check( Resolution.X > 0 && Resolution.Y > 0 );
 
-	Size = InitSize;
+	Scale = InitScale;
 	DockedTo = InitDockedTo;
 
 	SetupWidgetComponent();
 }
 
 
-void AVREditorFloatingUI::SetUMGWidget( FVREditorUISystem& InitOwner, TSubclassOf<UVREditorBaseUserWidget> InitUserWidgetClass, const FIntPoint InitResolution, const float InitSize, const EDockedTo InitDockedTo )
+void AVREditorFloatingUI::SetUMGWidget( FVREditorUISystem& InitOwner, TSubclassOf<UVREditorBaseUserWidget> InitUserWidgetClass, const FIntPoint InitResolution, const float InitScale, const EDockedTo InitDockedTo )
 {
 	Owner = &InitOwner;
 
@@ -120,7 +120,7 @@ void AVREditorFloatingUI::SetUMGWidget( FVREditorUISystem& InitOwner, TSubclassO
 	Resolution = InitResolution;
 	check( Resolution.X > 0 && Resolution.Y > 0 );
 
-	Size = InitSize;
+	Scale = InitScale;
 	DockedTo = InitDockedTo;
 
 	SetupWidgetComponent();
@@ -365,7 +365,7 @@ FTransform AVREditorFloatingUI::MakeUITransformLockedToHand( const int32 HandInd
 		FTransform(
 			UIRotation,
 			UILocation,
-			FVector( Size * WorldScaleFactor ) );
+			FVector( Scale * WorldScaleFactor ) );
 
 	return UITransform;
 }
@@ -399,7 +399,7 @@ FTransform AVREditorFloatingUI::MakeUITransformLockedToRoom()
 	return FTransform(
 		LocalRotation + RoomTransform.GetRotation().Rotator(),
 		UILocation,
-		FVector( Size * WorldScaleFactor ) );
+		FVector( Scale * WorldScaleFactor ) );
 }
 
 FTransform AVREditorFloatingUI::MakeUITransformLockedToHead()
@@ -412,7 +412,7 @@ FTransform AVREditorFloatingUI::MakeUITransformLockedToHead()
 	FVector NewLocation = FVector( HeadWorldLocation.X, HeadWorldLocation.Y, RoomWorldLocation.Z );
 	
 	Result.SetLocation( NewLocation + ( RelativeOffset * WorldScaleFactor ) );
-	Result.SetScale3D( FVector( Size *  WorldScaleFactor ) );
+	Result.SetScale3D( FVector( Scale * WorldScaleFactor ) );
 	
 	return Result;
 }
@@ -445,4 +445,12 @@ void AVREditorFloatingUI::ShowUI( const bool bShow, const bool bAllowFading, con
 
 		FadeDelay = InitFadeDelay;
 	}
-}
\ No newline at end of file
+}
+
+
+FVector2D AVREditorFloatingUI::GetSize() const
+{
+	const float Aspect = (float)Resolution.X / (float)Resolution.Y;
+	return FVector2D( Scale, Scale / Aspect );
+}
+
diff --git a/Engine/Source/Editor/VREditor/VREditorFloatingUI.h b/Engine/Source/Editor/VREditor/VREditorFloatingUI.h
--- a/Engine/Source/Editor/VREditor/VREditorFloatingUI.h
+++ b/Engine/Source/Editor/VREditor/VREditorFloatingUI.h
@@ -33,10 +33,10 @@ public:
 	};
 
 	/** Creates a FVREditorFloatingUI using a Slate widget, and sets up safe defaults */
-	void SetSlateWidget( class FVREditorUISystem& InitOwner, const TSharedRef<SWidget>& InitSlateWidget, const FIntPoint InitResolution, const float InitSize, const EDockedTo InitDockedTo );
+	void SetSlateWidget( class FVREditorUISystem& InitOwner, const TSharedRef<SWidget>& InitSlateWidget, const FIntPoint InitResolution, const float InitScale, const EDockedTo InitDockedTo );
 
 	/** Creates a FVREditorFloatingUI using a UMG user widget, and sets up safe defaults */
-	void SetUMGWidget( class FVREditorUISystem& InitOwner, class TSubclassOf<class UVREditorBaseUserWidget> InitUserWidgetClass, const FIntPoint InitResolution, const float InitSize, const EDockedTo InitDockedTo );
+	void SetUMGWidget( class FVREditorUISystem& InitOwner, class TSubclassOf<class UVREditorBaseUserWidget> InitUserWidgetClass, const FIntPoint InitResolution, const float InitScale, const EDockedTo InitDockedTo );
 
 	/** @return Returns true if the UI is visible (or wants to be visible -- it might be transitioning */
 	bool IsUIVisible() const
@@ -68,6 +68,9 @@ public:
 		return WidgetComponent;
 	}
 
+	/** Returns the actual size of the UI in either axis, after scaling has been applied.  Does not take into animation or world scaling */
+	FVector2D GetSize() const;
+
 	/** AActor overrides */
 	virtual void Destroyed() override;
 
@@ -120,7 +123,7 @@ protected:
 	virtual FTransform MakeCustomUITransform() { return FTransform(); };
 
 	/** Called to finish setting everything up, after a widget has been assigned */
-	void SetupWidgetComponent();
+	virtual void SetupWidgetComponent();
 
 private:
 
@@ -137,10 +140,10 @@ private:
 	void RotateToHead( FTransform& Transform );
 
 	/** Create a room transform with the RelativeOffset and LocalRotation */
-	FTransform MakeUITransformLockedToRoom( );
+	FTransform MakeUITransformLockedToRoom();
 
 	/** Create a transform with the RelativeOffset and LocalRotation that moves with the head */
-	FTransform MakeUITransformLockedToHead( );
+	FTransform MakeUITransformLockedToHead();
 
 protected:
 
@@ -162,7 +165,7 @@ protected:
 	class UWidgetComponent* WidgetComponent;
 	
 	/** How big the UI should be */
-	float Size;
+	float Scale;
 
 	/** Resolution we should draw this UI at, regardless of scale */
 	FIntPoint Resolution;
diff --git a/Engine/Source/Editor/VREditor/VREditorMode.cpp b/Engine/Source/Editor/VREditor/VREditorMode.cpp
--- a/Engine/Source/Editor/VREditor/VREditorMode.cpp
+++ b/Engine/Source/Editor/VREditor/VREditorMode.cpp
@@ -211,6 +211,7 @@ void FVREditorMode::Enter()
 		Colors[ (int32)EColors::WhiteGizmoColor ] = FLinearColor( 0.7f, 0.7f, 0.7f, 1.0f );
 		Colors[ (int32)EColors::HoverGizmoColor ] = FLinearColor( FLinearColor::Yellow );
 		Colors[ (int32)EColors::DraggingGizmoColor ] = FLinearColor( FLinearColor::White );
+		Colors[ (int32)EColors::UISelectionBarColor ] = FLinearColor( 0.0f, 0.7f, 0.7f, 1.0f );
 	}
 
 	// @todo vreditor: We need to make sure the user can never switch to orthographic mode, or activate settings that
