VR Editor: Fixed wrong Content Browser size.  Tweaked UI offsets.

diff --git a/Engine/Source/Editor/VREditor/VREditorMode.h b/Engine/Source/Editor/VREditor/VREditorMode.h
--- a/Engine/Source/Editor/VREditor/VREditorMode.h
+++ b/Engine/Source/Editor/VREditor/VREditorMode.h
@@ -509,6 +509,7 @@ public:
 		WhiteGizmoColor,
 		HoverGizmoColor,
 		DraggingGizmoColor,
+		UISelectionBarColor,
 		TotalCount
 	};
 
diff --git a/Engine/Source/Editor/VREditor/VREditorUISystem.cpp b/Engine/Source/Editor/VREditor/VREditorUISystem.cpp
--- a/Engine/Source/Editor/VREditor/VREditorUISystem.cpp
+++ b/Engine/Source/Editor/VREditor/VREditorUISystem.cpp
@@ -41,7 +41,7 @@ namespace VREd
 	static FAutoConsoleVariable EditorUISize( TEXT( "VREd.EditorUISize" ), 50.0f, TEXT( "How big editor UIs should be" ) );
 	static FAutoConsoleVariable ContentBrowserUIScale( TEXT( "VREd.ContentBrowserUIScale" ), 2.5f, TEXT( "How much to scale up (or down) the content browser for VR" ) );
 	static FAutoConsoleVariable EditorUIScale( TEXT( "VREd.EditorUIScale" ), 2.5f, TEXT( "How much to scale up (or down) editor UIs for VR" ) );
-	static FAutoConsoleVariable EditorUIOffsetFromHand( TEXT( "VREd.EditorUIOffsetFromHand" ), 40.0f, TEXT( "How far to offset editor UIs from the origin of the hand mesh" ) );
+	static FAutoConsoleVariable EditorUIOffsetFromHand( TEXT( "VREd.EditorUIOffsetFromHand" ), 12.0f, TEXT( "How far to offset editor UIs from the origin of the hand mesh" ) );
 	static FAutoConsoleVariable RadialMenuFadeDelay( TEXT( "VREd.RadialMenuFadeDelay" ), 0.2f, TEXT( "The delay for the radial menu after selecting a button" ) );
 	static FAutoConsoleVariable UIAbsoluteScrollSpeed( TEXT( "VREd.UIAbsoluteScrollSpeed" ), 8.0f, TEXT( "How fast the UI scrolls when dragging the touchpad" ) );
 	static FAutoConsoleVariable UIRelativeScrollSpeed( TEXT( "VREd.UIRelativeScrollSpeed" ), 0.75f, TEXT( "How fast the UI scrolls when holding an analog stick" ) );
@@ -595,7 +595,7 @@ void FVREditorUISystem::CreateUIs()
 		QuickMenuUI = GetOwner().SpawnTransientSceneActor< AVREditorFloatingUI >(TEXT("QuickMenu"), bWithSceneComponent);
 		QuickMenuUI->SetUMGWidget( *this, QuickMenuWidgetClass, Resolution, 22.0f, AVREditorFloatingUI::EDockedTo::Nothing );
 		QuickMenuUI->ShowUI( false );
-		QuickMenuUI->SetRelativeOffset( FVector( -10.0f, 0.0f, 4.0f ) );
+		QuickMenuUI->SetRelativeOffset( FVector( -10.0f, 0.0f, 3.0f ) );
 		FloatingUIs.Add( QuickMenuUI );
 	}
 
@@ -616,8 +616,6 @@ void FVREditorUISystem::CreateUIs()
 
 	// Make some editor UIs!
 	{
-		const FVector EditorUIRelativeOffset( VREd::EditorUIOffsetFromHand->GetFloat(), 0.0f, 0.0f );
-
 		{
 			const FIntPoint Resolution( VREd::ContentBrowserUIResolutionX->GetInt(), VREd::ContentBrowserUIResolutionY->GetInt() );
 
@@ -668,9 +666,8 @@ void FVREditorUISystem::CreateUIs()
 
 			const bool bWithSceneComponent = false;
 			AVREditorFloatingUI* ContentBrowserUI = GetOwner().SpawnTransientSceneActor< AVREditorDockableWindow >( TEXT( "ContentBrowserUI" ), bWithSceneComponent );
-			ContentBrowserUI->SetSlateWidget( *this, WidgetToDraw, DefaultResolution, VREd::EditorUISize->GetFloat(), AVREditorFloatingUI::EDockedTo::Nothing );
+			ContentBrowserUI->SetSlateWidget( *this, WidgetToDraw, Resolution, VREd::ContentBrowserUISize->GetFloat(), AVREditorFloatingUI::EDockedTo::Nothing );
 			ContentBrowserUI->ShowUI( false );
-			ContentBrowserUI->SetRelativeOffset( EditorUIRelativeOffset );
 			FloatingUIs.Add( ContentBrowserUI );
 
 			EditorUIPanels[ (int32)EEditorUIPanel::ContentBrowser ] = ContentBrowserUI;
@@ -695,7 +692,6 @@ void FVREditorUISystem::CreateUIs()
 			AVREditorFloatingUI* WorldOutlinerUI = GetOwner().SpawnTransientSceneActor< AVREditorDockableWindow >( TEXT( "WorldOutlinerUI" ), bWithSceneComponent );
 			WorldOutlinerUI->SetSlateWidget( *this, WidgetToDraw, DefaultResolution, VREd::EditorUISize->GetFloat(), AVREditorFloatingUI::EDockedTo::Nothing );
 			WorldOutlinerUI->ShowUI( false );
-			WorldOutlinerUI->SetRelativeOffset( EditorUIRelativeOffset );
 			FloatingUIs.Add( WorldOutlinerUI );
 
 			EditorUIPanels[ (int32)EEditorUIPanel::WorldOutliner ] = WorldOutlinerUI;
@@ -719,7 +715,6 @@ void FVREditorUISystem::CreateUIs()
 			AVREditorFloatingUI* ActorDetailsUI = GetOwner().SpawnTransientSceneActor< AVREditorDockableWindow >( TEXT( "ActorDetailsUI" ), bWithSceneComponent );
 			ActorDetailsUI->SetSlateWidget( *this, WidgetToDraw, DefaultResolution, VREd::EditorUISize->GetFloat(), AVREditorFloatingUI::EDockedTo::Nothing );
 			ActorDetailsUI->ShowUI( false );
-			ActorDetailsUI->SetRelativeOffset( EditorUIRelativeOffset );
 			FloatingUIs.Add( ActorDetailsUI );
 
 			EditorUIPanels[ (int32)EEditorUIPanel::ActorDetails ] = ActorDetailsUI;
@@ -742,7 +737,6 @@ void FVREditorUISystem::CreateUIs()
 			AVREditorFloatingUI* ModesUI = GetOwner().SpawnTransientSceneActor< AVREditorDockableWindow >( TEXT( "ModesUI" ), bWithSceneComponent );
 			ModesUI->SetSlateWidget( *this, WidgetToDraw, DefaultResolution, VREd::EditorUISize->GetFloat(), AVREditorFloatingUI::EDockedTo::Nothing );
 			ModesUI->ShowUI( false );
-			ModesUI->SetRelativeOffset( EditorUIRelativeOffset );
 			FloatingUIs.Add( ModesUI );
 
 			// @todo vreditor placement: This is required to force the modes UI to refresh -- otherwise it looks empty
@@ -853,7 +847,7 @@ void FVREditorUISystem::ShowEditorUIPanel( AVREditorFloatingUI* Panel, const int
 			const AVREditorFloatingUI::EDockedTo DockedTo = HandIndex == VREditorConstants::LeftHandIndex ? AVREditorFloatingUI::EDockedTo::LeftHand : AVREditorFloatingUI::EDockedTo::RightHand;
 			Panel->SetDockedTo( DockedTo );
 
-			const FVector EditorUIRelativeOffset( VREd::EditorUIOffsetFromHand->GetFloat(), 0.0f, 0.0f );
+			const FVector EditorUIRelativeOffset( Panel->GetSize().Y * 0.5f + VREd::EditorUIOffsetFromHand->GetFloat(), 0.0f, 0.0f );
 			Panel->SetRelativeOffset( EditorUIRelativeOffset );
 			Panel->SetLocalRotation( FRotator( 90.0f, 180.0f, 0.0f ) );
 		}
