VR Editor: Added close button to all dockable UI panels

diff --git a/Engine/Source/Editor/VREditor/VREditorDockableWindow.cpp b/Engine/Source/Editor/VREditor/VREditorDockableWindow.cpp
--- a/Engine/Source/Editor/VREditor/VREditorDockableWindow.cpp
+++ b/Engine/Source/Editor/VREditor/VREditorDockableWindow.cpp
@@ -10,7 +10,8 @@
 namespace VREd
 {
 	static FAutoConsoleVariable DockWindowThickness( TEXT( "VREd.DockWindowTickness" ), 1.0f, TEXT( "How thick the window is" ) );
-	static FAutoConsoleVariable DockUIZBarOffset( TEXT( "VREd.DockUIZBarOffset" ), 2.0f, TEXT( "Z Distance between the selectionbar and the UI" ) );
+	static FAutoConsoleVariable DockUISelectionBarVerticalOffset( TEXT( "VREd.DockUISelectionBarVerticalOffset" ), 2.0f, TEXT( "Z Distance between the selectionbar and the UI" ) );
+	static FAutoConsoleVariable DockUICloseButtonOffsetFromCorner( TEXT( "VREd.DockUICloseButtonOffsetFromCorner" ), 2.0f, TEXT( "How far away from the corner (along each 2D axis) of the UI the close button draws" ) );
 }
 
 AVREditorDockableWindow::AVREditorDockableWindow() 
@@ -49,20 +50,20 @@ AVREditorDockableWindow::AVREditorDockableWindow()
 			check( SelectionMesh != nullptr );
 		}
 
-		SelectionMeshComponent = CreateDefaultSubobject<UStaticMeshComponent>( TEXT( "SelectionBarMesh" ) );
-		SelectionMeshComponent->SetStaticMesh( SelectionMesh );
-		SelectionMeshComponent->SetMobility( EComponentMobility::Movable );
-		SelectionMeshComponent->AttachTo( RootComponent );
+		SelectionBarMeshComponent = CreateDefaultSubobject<UStaticMeshComponent>( TEXT( "SelectionBarMesh" ) );
+		SelectionBarMeshComponent->SetStaticMesh( SelectionMesh );
+		SelectionBarMeshComponent->SetMobility( EComponentMobility::Movable );
+		SelectionBarMeshComponent->AttachTo( RootComponent );
 
-		SelectionMeshComponent->SetCollisionEnabled( ECollisionEnabled::QueryOnly );
-		SelectionMeshComponent->SetCollisionResponseToAllChannels( ECollisionResponse::ECR_Block );
-		SelectionMeshComponent->SetCollisionResponseToChannel( ECC_EditorGizmo, ECollisionResponse::ECR_Block );
+		SelectionBarMeshComponent->SetCollisionEnabled( ECollisionEnabled::QueryOnly );
+		SelectionBarMeshComponent->SetCollisionResponseToAllChannels( ECollisionResponse::ECR_Block );
+		SelectionBarMeshComponent->SetCollisionResponseToChannel( ECC_EditorGizmo, ECollisionResponse::ECR_Block );
 
-		SelectionMeshComponent->bGenerateOverlapEvents = false;
-		SelectionMeshComponent->SetCanEverAffectNavigation( false );
-		SelectionMeshComponent->bCastDynamicShadow = false;
-		SelectionMeshComponent->bCastStaticShadow = false;
-		SelectionMeshComponent->bAffectDistanceFieldLighting = false;
+		SelectionBarMeshComponent->bGenerateOverlapEvents = false;
+		SelectionBarMeshComponent->SetCanEverAffectNavigation( false );
+		SelectionBarMeshComponent->bCastDynamicShadow = false;
+		SelectionBarMeshComponent->bCastStaticShadow = false;
+		SelectionBarMeshComponent->bAffectDistanceFieldLighting = false;
 
 		UMaterial* HoverMaterial = nullptr;
 		{
@@ -70,9 +71,9 @@ AVREditorDockableWindow::AVREditorDockableWindow()
 			HoverMaterial = ObjectFinder.Object;
 			check( HoverMaterial != nullptr );
 		}
-		HoverMaterialMID = UMaterialInstanceDynamic::Create( HoverMaterial, GetTransientPackage() );
-		check( HoverMaterialMID != nullptr );
-		SelectionMeshComponent->SetMaterial( 0, HoverMaterialMID );
+		SelectionBarMID = UMaterialInstanceDynamic::Create( HoverMaterial, GetTransientPackage() );
+		check( SelectionBarMID != nullptr );
+		SelectionBarMeshComponent->SetMaterial( 0, SelectionBarMID );
 
 		UMaterial* TranslucentHoverMaterial = nullptr;
 		{
@@ -80,9 +81,53 @@ AVREditorDockableWindow::AVREditorDockableWindow()
 			TranslucentHoverMaterial = ObjectFinder.Object;
 			check( TranslucentHoverMaterial != nullptr );
 		}
-		TranslucentHoverMID = UMaterialInstanceDynamic::Create( TranslucentHoverMaterial, GetTransientPackage() );
-		check( TranslucentHoverMID != nullptr );
-		SelectionMeshComponent->SetMaterial( 1, TranslucentHoverMID );
+		SelectionBarTranslucentMID = UMaterialInstanceDynamic::Create( TranslucentHoverMaterial, GetTransientPackage() );
+		check( SelectionBarTranslucentMID != nullptr );
+		SelectionBarMeshComponent->SetMaterial( 1, SelectionBarTranslucentMID );
+	}
+
+	{
+		UStaticMesh* CloseButtonMesh = nullptr;
+		{
+			static ConstructorHelpers::FObjectFinder<UStaticMesh> ObjectFinder( TEXT( "/Engine/VREditor/UI/CloseButtonMesh" ) );
+			CloseButtonMesh = ObjectFinder.Object;
+			check( CloseButtonMesh != nullptr );
+		}
+
+		CloseButtonMeshComponent= CreateDefaultSubobject<UStaticMeshComponent>( TEXT( "CloseButtonMesh" ) );
+		CloseButtonMeshComponent->SetStaticMesh( CloseButtonMesh );
+		CloseButtonMeshComponent->SetMobility( EComponentMobility::Movable );
+		CloseButtonMeshComponent->AttachTo( RootComponent );
+
+		CloseButtonMeshComponent->SetCollisionEnabled( ECollisionEnabled::QueryOnly );
+		CloseButtonMeshComponent->SetCollisionResponseToAllChannels( ECollisionResponse::ECR_Block );
+		CloseButtonMeshComponent->SetCollisionResponseToChannel( ECC_EditorGizmo, ECollisionResponse::ECR_Block );
+
+		CloseButtonMeshComponent->bGenerateOverlapEvents = false;
+		CloseButtonMeshComponent->SetCanEverAffectNavigation( false );
+		CloseButtonMeshComponent->bCastDynamicShadow = false;
+		CloseButtonMeshComponent->bCastStaticShadow = false;
+		CloseButtonMeshComponent->bAffectDistanceFieldLighting = false;
+
+		UMaterial* HoverMaterial = nullptr;
+		{
+			static ConstructorHelpers::FObjectFinder<UMaterial> ObjectFinder( TEXT( "/Engine/VREditor/TransformGizmo/TransformGizmoMaterial" ) );
+			HoverMaterial = ObjectFinder.Object;
+			check( HoverMaterial != nullptr );
+		}
+		CloseButtonMID = UMaterialInstanceDynamic::Create( HoverMaterial, GetTransientPackage() );
+		check( CloseButtonMID != nullptr );
+		CloseButtonMeshComponent->SetMaterial( 0, CloseButtonMID );
+
+		UMaterial* TranslucentHoverMaterial = nullptr;
+		{
+			static ConstructorHelpers::FObjectFinder<UMaterial> ObjectFinder( TEXT( "/Engine/VREditor/TransformGizmo/TranslucentTransformGizmoMaterial" ) );
+			TranslucentHoverMaterial = ObjectFinder.Object;
+			check( TranslucentHoverMaterial != nullptr );
+		}
+		CloseButtonTranslucentMID = UMaterialInstanceDynamic::Create( TranslucentHoverMaterial, GetTransientPackage() );
+		check( CloseButtonTranslucentMID != nullptr );
+		CloseButtonMeshComponent->SetMaterial( 1, CloseButtonTranslucentMID );
 	}
 }
 
@@ -92,6 +137,7 @@ void AVREditorDockableWindow::SetupWidgetComponent()
 	Super::SetupWidgetComponent();
 
 	SetSelectionBarColor( GetOwner().GetOwner().GetColor( FVREditorMode::EColors::UISelectionBarColor ) );
+	SetCloseButtonColor( GetOwner().GetOwner().GetColor( FVREditorMode::EColors::UICloseButtonColor ) );
 }
 
 
@@ -101,26 +147,47 @@ void AVREditorDockableWindow::TickManually( float DeltaTime )
 
 	if( WidgetComponent->IsVisible() )
 	{
-		// How big the selection bar should be
-		const FVector SelectionBarSize( 1.0f, GetSize().X * 0.8f, 6.0f );	// @todo vreditor tweak
-
 		const float WorldScaleFactor = GetOwner().GetOwner().GetWorldScaleFactor();
 		const FVector AnimatedScale = CalculateAnimatedScale();
 
-		const float WindowMeshSize = 100.0f;	// Size of imported mesh, we need to inverse compensate for
-		const FVector WindowMeshScale = FVector(
-			VREd::DockWindowThickness->GetFloat(),
-			GetSize().X / WindowMeshSize,
-			GetSize().Y / WindowMeshSize ) * AnimatedScale * WorldScaleFactor;
-		WindowMeshComponent->SetRelativeScale3D( WindowMeshScale );
-
-		const FVector SelectionMeshScale = SelectionBarSize * AnimatedScale * WorldScaleFactor;
-		SelectionMeshComponent->SetRelativeScale3D( SelectionMeshScale );
-		const FVector SelectionMeshRelativeLocation = FVector(
-			0.0f,
-			0.0f,
-			-( GetSize().Y * 0.5f + SelectionBarSize.Z * 0.5f + VREd::DockUIZBarOffset->GetFloat() ) ) * AnimatedScale * WorldScaleFactor;
-		SelectionMeshComponent->SetRelativeLocation( SelectionMeshRelativeLocation );
+		// Update the window border mesh
+		{
+			const float WindowMeshSize = 100.0f;	// Size of imported mesh, we need to inverse compensate for
+
+			const FVector WindowMeshScale = FVector(
+				VREd::DockWindowThickness->GetFloat(),
+				GetSize().X / WindowMeshSize,
+				GetSize().Y / WindowMeshSize ) * AnimatedScale * WorldScaleFactor;
+			WindowMeshComponent->SetRelativeScale3D( WindowMeshScale );			
+		}
+
+		// Update the selection bar
+		{
+			// How big the selection bar should be
+			const FVector SelectionBarSize( 1.0f, GetSize().X * 0.8f, 6.0f );	// @todo vreditor tweak
+
+			const FVector SelectionBarScale = SelectionBarSize * AnimatedScale * WorldScaleFactor;
+			SelectionBarMeshComponent->SetRelativeScale3D( SelectionBarScale );
+			const FVector SelectionBarRelativeLocation = FVector(
+				0.0f,
+				0.0f,
+				-( GetSize().Y * 0.5f + SelectionBarSize.Z * 0.5f + VREd::DockUISelectionBarVerticalOffset->GetFloat() ) ) * AnimatedScale * WorldScaleFactor;
+			SelectionBarMeshComponent->SetRelativeLocation( SelectionBarRelativeLocation );
+		}
+
+		// Update the close button
+		{
+			// How big the close button should be
+			const FVector CloseButtonSize( 1.0f, 3.0f, 3.0f );	// @todo vreditor tweak
+
+			const FVector CloseButtonScale = CloseButtonSize * AnimatedScale * WorldScaleFactor;
+			CloseButtonMeshComponent->SetRelativeScale3D( CloseButtonScale );
+			const FVector CloseButtonRelativeLocation = FVector(
+				0.0f,
+				-( GetSize().X * 0.5f + CloseButtonSize.Y * 0.5f + VREd::DockUICloseButtonOffsetFromCorner->GetFloat() ),
+				( GetSize().Y * 0.5f + CloseButtonSize.Z * 0.5f + VREd::DockUICloseButtonOffsetFromCorner->GetFloat() ) ) * AnimatedScale * WorldScaleFactor;
+			CloseButtonMeshComponent->SetRelativeLocation( CloseButtonRelativeLocation );
+		}
 	}
 } 
 
@@ -141,20 +208,35 @@ void AVREditorDockableWindow::UpdateRelativeRoomTransform()
 
 void AVREditorDockableWindow::OnEnterHover( const FHitResult& HitResult )
 {
-	if( SelectionMeshComponent == HitResult.GetComponent() )
+	if( SelectionBarMeshComponent == HitResult.GetComponent() )
 	{
 		SetSelectionBarColor( GetOwner().GetOwner().GetColor( FVREditorMode::EColors::HoverGizmoColor ) );
 	}
+
+	if( CloseButtonMeshComponent == HitResult.GetComponent() )
+	{
+		SetCloseButtonColor( GetOwner().GetOwner().GetColor( FVREditorMode::EColors::HoverGizmoColor ) );
+	}
 }
 
 void AVREditorDockableWindow::OnLeaveHover()
 {
 	SetSelectionBarColor( GetOwner().GetOwner().GetColor( FVREditorMode::EColors::UISelectionBarColor ) );
+	SetCloseButtonColor( GetOwner().GetOwner().GetColor( FVREditorMode::EColors::UICloseButtonColor ) );
 }
 
+
 void AVREditorDockableWindow::SetSelectionBarColor( const FLinearColor& LinearColor )
 {
 	static FName StaticColorParameterName( "Color" );
-	HoverMaterialMID->SetVectorParameterValue( StaticColorParameterName, LinearColor );
-	TranslucentHoverMID->SetVectorParameterValue( StaticColorParameterName, LinearColor );
-}
\ No newline at end of file
+	SelectionBarMID->SetVectorParameterValue( StaticColorParameterName, LinearColor );
+	SelectionBarTranslucentMID->SetVectorParameterValue( StaticColorParameterName, LinearColor );
+}
+
+
+void AVREditorDockableWindow::SetCloseButtonColor( const FLinearColor& LinearColor )
+{
+	static FName StaticColorParameterName( "Color" );
+	CloseButtonMID->SetVectorParameterValue( StaticColorParameterName, LinearColor );
+	CloseButtonTranslucentMID->SetVectorParameterValue( StaticColorParameterName, LinearColor );
+}
diff --git a/Engine/Source/Editor/VREditor/VREditorDockableWindow.h b/Engine/Source/Editor/VREditor/VREditorDockableWindow.h
--- a/Engine/Source/Editor/VREditor/VREditorDockableWindow.h
+++ b/Engine/Source/Editor/VREditor/VREditorDockableWindow.h
@@ -13,14 +13,6 @@ class AVREditorDockableWindow : public AVREditorFloatingUI
 {
 	GENERATED_BODY()
 	
-	/** The dockable window mesh */
-	UPROPERTY()
-	class UStaticMeshComponent* WindowMeshComponent;
-
-	/** Mesh underneath the window for easy selecting and dragging */
-	UPROPERTY()
-	class UStaticMeshComponent* SelectionMeshComponent;
-
 public:
 
 	/** Default constructor */
@@ -37,20 +29,49 @@ public:
 
 	/** Leaving hover with laser changes the color of SelectionMesh */
 	void OnLeaveHover();
+	
+	/** Gets the close button component */
+	UStaticMeshComponent* GetCloseButtonMeshComponent()
+	{
+		return CloseButtonMeshComponent;
+	}
 
 private:
 
 	// AVREditorFloatingUI overrides
 	virtual void SetupWidgetComponent() override;
 
-	/** Setting the color on the dynamic materials of the SelectionMesh */
+	/** Set the color on the dynamic materials of the selection bar */
 	void SetSelectionBarColor( const FLinearColor& LinearColor );
 
-	/** Dynamic material  (opaque) */
+	/** Set the color on the dynamic materials of the close button */
+	void SetCloseButtonColor( const FLinearColor& LinearColor );
+
+	/** The dockable window mesh */
+	UPROPERTY()
+	class UStaticMeshComponent* WindowMeshComponent;
+
+	/** Mesh underneath the window for easy selecting and dragging */
+	UPROPERTY()
+	class UStaticMeshComponent* SelectionBarMeshComponent;
+
+	/** Mesh that represents the close button for this UI */
+	UPROPERTY()
+	class UStaticMeshComponent* CloseButtonMeshComponent;
+
+	/** Selection bar dynamic material  (opaque) */
+	UPROPERTY()
+	class UMaterialInstanceDynamic* SelectionBarMID;
+
+	/** Select bar dynamic material (translucent) */
+	UPROPERTY()
+	class UMaterialInstanceDynamic* SelectionBarTranslucentMID;
+
+	/** Close button dynamic material  (opaque) */
 	UPROPERTY()
-	class UMaterialInstanceDynamic* HoverMaterialMID;
+	class UMaterialInstanceDynamic* CloseButtonMID;
 
-	/** Dynamic material (translucent) */
+	/** Close button dynamic material (translucent) */
 	UPROPERTY()
-	class UMaterialInstanceDynamic* TranslucentHoverMID;
-};	
\ No newline at end of file
+	class UMaterialInstanceDynamic* CloseButtonTranslucentMID;
+};
\ No newline at end of file
diff --git a/Engine/Source/Editor/VREditor/VREditorMode.cpp b/Engine/Source/Editor/VREditor/VREditorMode.cpp
--- a/Engine/Source/Editor/VREditor/VREditorMode.cpp
+++ b/Engine/Source/Editor/VREditor/VREditorMode.cpp
@@ -212,6 +212,7 @@ void FVREditorMode::Enter()
 		Colors[ (int32)EColors::HoverGizmoColor ] = FLinearColor( FLinearColor::Yellow );
 		Colors[ (int32)EColors::DraggingGizmoColor ] = FLinearColor( FLinearColor::White );
 		Colors[ (int32)EColors::UISelectionBarColor ] = FLinearColor( 0.0f, 0.7f, 0.7f, 1.0f );
+		Colors[ (int32)EColors::UICloseButtonColor ] = FLinearColor( 0.8f, 0.0f, 0.0f, 1.0f );
 	}
 
 	// @todo vreditor: We need to make sure the user can never switch to orthographic mode, or activate settings that
diff --git a/Engine/Source/Editor/VREditor/VREditorMode.h b/Engine/Source/Editor/VREditor/VREditorMode.h
--- a/Engine/Source/Editor/VREditor/VREditorMode.h
+++ b/Engine/Source/Editor/VREditor/VREditorMode.h
@@ -510,6 +510,7 @@ public:
 		HoverGizmoColor,
 		DraggingGizmoColor,
 		UISelectionBarColor,
+		UICloseButtonColor,
 		TotalCount
 	};
 
diff --git a/Engine/Source/Editor/VREditor/VREditorWorldInteraction.cpp b/Engine/Source/Editor/VREditor/VREditorWorldInteraction.cpp
--- a/Engine/Source/Editor/VREditor/VREditorWorldInteraction.cpp
+++ b/Engine/Source/Editor/VREditor/VREditorWorldInteraction.cpp
@@ -388,9 +388,20 @@ void FVREditorWorldInteraction::OnVRAction( FEditorViewportClient& ViewportClien
 								AVREditorDockableWindow* DockableWindow = Cast< AVREditorDockableWindow >( Actor );
 								if (DockableWindow)
 								{
-									Hand.DraggingMode = EVREditorDraggingMode::DockableWindow;
-									DockSelectDistance = ( LaserPointerStart - HitResult.Location ).Size();
-									Owner.GetUISystem().StartDraggingDockUI( DockableWindow, VRAction.HandIndex, DockSelectDistance );
+									if( HitResult.Component == DockableWindow->GetCloseButtonMeshComponent() )
+									{
+										// Close the window
+										const bool bShouldShow = false;
+										const bool bShowOnHand = false;
+										const bool bRefreshQuickMenu = true;
+										Owner.GetUISystem().ShowEditorUIPanel( DockableWindow, VRAction.HandIndex, bShouldShow, bShowOnHand, bRefreshQuickMenu );
+									}
+									else
+									{
+										Hand.DraggingMode = EVREditorDraggingMode::DockableWindow;
+										DockSelectDistance = ( LaserPointerStart - HitResult.Location ).Size();
+										Owner.GetUISystem().StartDraggingDockUI( DockableWindow, VRAction.HandIndex, DockSelectDistance );
+									}
 									bWasHandled = true;
 								}
 							}
