VR Editor :
- Scale indicator matches color of hands
- Bigger radius for trackpad hiding all panels

diff --git a/Engine/Source/Editor/VREditor/VREditorMode.cpp b/Engine/Source/Editor/VREditor/VREditorMode.cpp
--- a/Engine/Source/Editor/VREditor/VREditorMode.cpp
+++ b/Engine/Source/Editor/VREditor/VREditorMode.cpp
@@ -204,7 +204,7 @@ void FVREditorMode::Enter()
 		Colors[ (int32)EColors::SelectionColor ] = FLinearColor::Green;
 		Colors[ (int32)EColors::TeleportColor ] = FLinearColor( 1.0f, 0.0f, 0.75f, 1.0f );
 		Colors[ (int32)EColors::WorldDraggingColor_OneHand ] = FLinearColor::Yellow;
-		Colors[ (int32)EColors::WorldDraggingColor_TwoHands ] = FLinearColor::Blue;
+		Colors[ (int32)EColors::WorldDraggingColor_TwoHands ] = FLinearColor( 0.05f, 0.05f, 0.4f, 1.0f );
 		Colors[ (int32)EColors::RedGizmoColor ] = FLinearColor( 0.4f, 0.05f, 0.05f, 1.0f );
 		Colors[ (int32)EColors::GreenGizmoColor ] = FLinearColor( 0.05f, 0.4f, 0.05f, 1.0f );
 		Colors[ (int32)EColors::BlueGizmoColor ] = FLinearColor( 0.05f, 0.05f, 0.4f, 1.0f );
@@ -712,45 +712,69 @@ void FVREditorMode::SpawnAvatarMeshActor()
 			{
 				UMaterialInterface* UserScaleIndicatorMaterial = LoadObject<UMaterialInterface>( nullptr, TEXT( "/Engine/VREditor/LaserPointer/LaserPointerMaterialInst" ) );
 				check( UserScaleIndicatorMaterial != nullptr );
-			
+
+				UMaterialInstance* TranslucentUserScaleIndicatorMaterial = LoadObject<UMaterialInstance>( nullptr, TEXT( "/Engine/VREditor/LaserPointer/TranslucentLaserPointerMaterialInst" ) );
+				check( TranslucentUserScaleIndicatorMaterial != nullptr );
+
 				UStaticMesh* ScaleLineMesh = LoadObject<UStaticMesh>( nullptr, TEXT( "/Engine/VREditor/LaserPointer/LaserPointerMesh" ) );
 				check( ScaleLineMesh != nullptr );
-			
+
 				// Creating the background bar progress of the scale 
-				ScaleProgressMeshComponent = NewObject<UStaticMeshComponent>( AvatarMeshActor );
-				AvatarMeshActor->AddOwnedComponent( ScaleProgressMeshComponent );
-				ScaleProgressMeshComponent->AttachTo( AvatarMeshActor->GetRootComponent() );
-				ScaleProgressMeshComponent->RegisterComponent();
+				{
+					ScaleProgressMeshComponent = NewObject<UStaticMeshComponent>( AvatarMeshActor );
+					AvatarMeshActor->AddOwnedComponent( ScaleProgressMeshComponent );
+					ScaleProgressMeshComponent->AttachTo( AvatarMeshActor->GetRootComponent() );
+					ScaleProgressMeshComponent->RegisterComponent();
 				
-				ScaleProgressMeshComponent->SetStaticMesh( ScaleLineMesh );
-				ScaleProgressMeshComponent->SetMobility( EComponentMobility::Movable );
-				ScaleProgressMeshComponent->SetCollisionEnabled( ECollisionEnabled::NoCollision );
+					ScaleProgressMeshComponent->SetStaticMesh( ScaleLineMesh );
+					ScaleProgressMeshComponent->SetMobility( EComponentMobility::Movable );
+					ScaleProgressMeshComponent->SetCollisionEnabled( ECollisionEnabled::NoCollision );
 
-				ScaleProgressMeshComponent->SetMaterial( 0, nullptr );
-				ScaleProgressMeshComponent->SetMaterial( 1, UserScaleIndicatorMaterial );
+					UMaterialInstanceDynamic* UserScaleMID = UMaterialInstanceDynamic::Create( UserScaleIndicatorMaterial, GetTransientPackage() );
+					check( UserScaleMID != nullptr );
+					ScaleProgressMeshComponent->SetMaterial( 0, UserScaleMID );
+
+					UMaterialInstanceDynamic* TranslucentUserScaleMID = UMaterialInstanceDynamic::Create( TranslucentUserScaleIndicatorMaterial, GetTransientPackage() );
+					check( TranslucentUserScaleMID != nullptr );
+					ScaleProgressMeshComponent->SetMaterial( 1, TranslucentUserScaleMID );
+
+					static FName StaticLaserColorName( "LaserColor" );
+					const FLinearColor Color = GetColor( EColors::WorldDraggingColor_TwoHands );
+					UserScaleMID->SetVectorParameterValue( StaticLaserColorName, Color );
+					TranslucentUserScaleMID->SetVectorParameterValue( StaticLaserColorName, Color );
+
+					// The user scale indicator starts invisible
+					ScaleProgressMeshComponent->SetVisibility( false );
+				}
 
-				// The user scale indicator starts invisible
-				ScaleProgressMeshComponent->SetVisibility( false );
 
 				// Creating the current progress of the scale 
-				CurrentScaleProgressMeshComponent = NewObject<UStaticMeshComponent>( AvatarMeshActor );
-				AvatarMeshActor->AddOwnedComponent( CurrentScaleProgressMeshComponent );
-				CurrentScaleProgressMeshComponent->AttachTo( AvatarMeshActor->GetRootComponent() );
-				CurrentScaleProgressMeshComponent->RegisterComponent();
-
-				CurrentScaleProgressMeshComponent->SetStaticMesh( ScaleLineMesh );
-				CurrentScaleProgressMeshComponent->SetMobility( EComponentMobility::Movable );
-				CurrentScaleProgressMeshComponent->SetCollisionEnabled( ECollisionEnabled::NoCollision );
-
-				UMaterialInstanceDynamic* CurrenterUserScaleMID = UMaterialInstanceDynamic::Create( UserScaleIndicatorMaterial, GetTransientPackage() );
-				check( CurrenterUserScaleMID != nullptr );
-				CurrentScaleProgressMeshComponent->SetMaterial( 0, nullptr );
-				CurrentScaleProgressMeshComponent->SetMaterial( 1, CurrenterUserScaleMID );
-				static FName StaticLaserColorName( "LaserColor" );
-				CurrenterUserScaleMID->SetVectorParameterValue( StaticLaserColorName, FLinearColor::Green );
-
-				// The user scale indicator starts invisible
-				CurrentScaleProgressMeshComponent->SetVisibility( false );
+				{
+					CurrentScaleProgressMeshComponent = NewObject<UStaticMeshComponent>( AvatarMeshActor );
+					AvatarMeshActor->AddOwnedComponent( CurrentScaleProgressMeshComponent );
+					CurrentScaleProgressMeshComponent->AttachTo( AvatarMeshActor->GetRootComponent() );
+					CurrentScaleProgressMeshComponent->RegisterComponent();
+
+					CurrentScaleProgressMeshComponent->SetStaticMesh( ScaleLineMesh );
+					CurrentScaleProgressMeshComponent->SetMobility( EComponentMobility::Movable );
+					CurrentScaleProgressMeshComponent->SetCollisionEnabled( ECollisionEnabled::NoCollision );
+
+					UMaterialInstanceDynamic* UserScaleMID = UMaterialInstanceDynamic::Create( UserScaleIndicatorMaterial, GetTransientPackage() );
+					check( UserScaleMID != nullptr );
+					CurrentScaleProgressMeshComponent->SetMaterial( 0, UserScaleMID );
+
+					UMaterialInstanceDynamic* TranslucentUserScaleMID = UMaterialInstanceDynamic::Create( TranslucentUserScaleIndicatorMaterial, GetTransientPackage() );
+					check( TranslucentUserScaleMID != nullptr );
+					CurrentScaleProgressMeshComponent->SetMaterial( 1, TranslucentUserScaleMID );
+
+					static FName StaticLaserColorName( "LaserColor" );
+					const FLinearColor Color = GetColor( EColors::GreenGizmoColor );
+					UserScaleMID->SetVectorParameterValue( StaticLaserColorName, Color );
+					TranslucentUserScaleMID->SetVectorParameterValue( StaticLaserColorName, Color );
+
+					// The user scale indicator starts invisible
+					CurrentScaleProgressMeshComponent->SetVisibility( false );
+				}
 			}
 
 			// Creating the text for scaling
@@ -781,7 +805,7 @@ void FVREditorMode::SpawnAvatarMeshActor()
 				UserScaleIndicatorText->SetFont( TextFont );
 				UserScaleIndicatorText->SetWorldSize( 8.0f );
 				UserScaleIndicatorText->SetTextMaterial( UserScaleIndicatorMaterial );
-				UserScaleIndicatorText->SetTextRenderColor( FLinearColor::Red.ToFColor( false ) );
+				UserScaleIndicatorText->SetTextRenderColor( GetColor( EColors::WorldDraggingColor_TwoHands ).ToFColor( false ) );
 
 				// Center the text horizontally
 				UserScaleIndicatorText->SetHorizontalAlignment( EHTA_Center );
diff --git a/Engine/Source/Editor/VREditor/VREditorRadialMenu.cpp b/Engine/Source/Editor/VREditor/VREditorRadialMenu.cpp
--- a/Engine/Source/Editor/VREditor/VREditorRadialMenu.cpp
+++ b/Engine/Source/Editor/VREditor/VREditorRadialMenu.cpp
@@ -14,7 +14,7 @@
 
 namespace VREd
 {
-	static FAutoConsoleVariable RadialMenuSelectionRadius( TEXT( "VREd.RadialMenuSelectionRadius" ), 0.2f, TEXT( "The radius that is used for selecting buttons on the radial menu" ) );
+	static FAutoConsoleVariable RadialMenuSelectionRadius( TEXT( "VREd.RadialMenuSelectionRadius" ), 0.5f, TEXT( "The radius that is used for selecting buttons on the radial menu" ) );
 }
 
 UVREditorRadialMenu::UVREditorRadialMenu( const FObjectInitializer& ObjectInitializer )
