VR Editor : Added tutorial panel with commands for resolution, room location, size, yaw and pitch.

diff --git a/Engine/Source/Editor/VREditor/VREditorUISystem.cpp b/Engine/Source/Editor/VREditor/VREditorUISystem.cpp
--- a/Engine/Source/Editor/VREditor/VREditorUISystem.cpp
+++ b/Engine/Source/Editor/VREditor/VREditorUISystem.cpp
@@ -48,6 +48,16 @@ namespace VREd
 	static FAutoConsoleVariable MinUIScrollDeltaForInertia( TEXT( "VREd.MinUIScrollDeltaForInertia" ), 0.25f, TEXT( "Minimum amount of touch pad input before inertial UI scrolling kicks in" ) );
 	static FAutoConsoleVariable MinDockDragDistance( TEXT( "VREd.MinDockDragDistance" ), 10.0f, TEXT( "Minimum amount of distance needed to drag dock from hand" ) );
 	static FAutoConsoleVariable DoubleClickTime( TEXT( "VREd.DoubleClickTime" ), 0.25f, TEXT( "Minumum duration between clicks for a double click event to fire" ) );
+
+	// Tutorial UI commands
+	static FAutoConsoleVariable TutorialUIResolutionX( TEXT( "VREd.TutorialUI.Resolution.X" ), 1920, TEXT( "Minumum duration between clicks for a double click event to fire" ) );
+	static FAutoConsoleVariable TutorialUIResolutionY( TEXT( "VREd.TutorialUI.Resolution.Y" ), 1080, TEXT( "Minumum duration between clicks for a double click event to fire" ) );
+	static FAutoConsoleVariable TutorialUISize( TEXT( "VREd.TutorialUI.Size" ), 200, TEXT( "Minumum duration between clicks for a double click event to fire" ) );
+	static FAutoConsoleVariable TutorialUIYaw( TEXT( "VREd.TutorialUI.Yaw" ), 270, TEXT( "Minumum duration between clicks for a double click event to fire" ) );
+	static FAutoConsoleVariable TutorialUIPitch( TEXT( "VREd.TutorialUI.Pitch" ), 45, TEXT( "Minumum duration between clicks for a double click event to fire" ) );
+	static FAutoConsoleVariable TutorialUILocationX( TEXT( "VREd.TutorialUI.Location.X" ), 0, TEXT( "Minumum duration between clicks for a double click event to fire" ) );
+	static FAutoConsoleVariable TutorialUILocationY( TEXT( "VREd.TutorialUI.Location.Y" ), 200, TEXT( "Minumum duration between clicks for a double click event to fire" ) );
+	static FAutoConsoleVariable TutorialUILocationZ( TEXT( "VREd.TutorialUI.Location.Z" ), 40, TEXT( "Minumum duration between clicks for a double click event to fire" ) );
 }
 
 
@@ -76,6 +86,9 @@ FVREditorUISystem::FVREditorUISystem( FVREditorMode& InitOwner )
 	QuickRadialWidgetClass = LoadClass<UVREditorRadialMenu>(nullptr, TEXT("/Engine/VREditor/UI/VRRadialQuickMenu.VRRadialQuickMenu_C"));
 	check(QuickRadialWidgetClass != nullptr);
 
+	TutorialWidgetClass = LoadClass<UVREditorBaseUserWidget>( nullptr, TEXT( "/Engine/VREditor/Tutorial/UI_VR_Tutorial_00.UI_VR_Tutorial_00_C" ) );
+	check( TutorialWidgetClass != nullptr );
+
 	// Create all of our UI panels
 	CreateUIs();
 }
@@ -118,6 +131,7 @@ void FVREditorUISystem::AddReferencedObjects( FReferenceCollector& Collector )
 	}
 	Collector.AddReferencedObject( QuickMenuWidgetClass );
 	Collector.AddReferencedObject( QuickRadialWidgetClass );
+	Collector.AddReferencedObject( TutorialWidgetClass );
 }
 
 
@@ -744,6 +758,18 @@ void FVREditorUISystem::CreateUIs()
 
 			EditorUIPanels[ (int32)EEditorUIPanel::Modes ] = ModesUI;			
 		}
+
+		// Create the tutorial dockable window
+		{
+			AVREditorFloatingUI* TutorialUI = GetOwner().SpawnTransientSceneActor< AVREditorDockableWindow >( TEXT( "TutorialUI" ), false );
+			TutorialUI->SetUMGWidget( *this, TutorialWidgetClass, FIntPoint( VREd::TutorialUIResolutionX->GetFloat(), VREd::TutorialUIResolutionY->GetFloat() ), VREd::TutorialUISize->GetFloat(), AVREditorFloatingUI::EDockedTo::Room );
+			TutorialUI->ShowUI( true );
+			TutorialUI->SetRelativeOffset( FVector( VREd::TutorialUILocationX->GetFloat(), VREd::TutorialUILocationY->GetFloat(), VREd::TutorialUILocationZ->GetFloat() ) );
+			TutorialUI->SetLocalRotation( FRotator( VREd::TutorialUIPitch->GetFloat(), VREd::TutorialUIYaw->GetFloat(), 0 ) );
+			FloatingUIs.Add( TutorialUI );
+
+			EditorUIPanels[ (int32)EEditorUIPanel::Tutorial ] = TutorialUI;		
+		}
 	}
 }
 
diff --git a/Engine/Source/Editor/VREditor/VREditorUISystem.h b/Engine/Source/Editor/VREditor/VREditorUISystem.h
--- a/Engine/Source/Editor/VREditor/VREditorUISystem.h
+++ b/Engine/Source/Editor/VREditor/VREditorUISystem.h
@@ -15,6 +15,7 @@ public:
 		WorldOutliner,
 		ActorDetails,
 		Modes,
+		Tutorial,
 
 		// ...
 
@@ -131,8 +132,11 @@ protected:
 	/** Quick menu widget class */
 	UClass* QuickMenuWidgetClass;
 
-	/** Quick radial menu qidget class */
+	/** Quick radial menu widget class */
 	UClass* QuickRadialWidgetClass;
+
+	/** Tutorial menu widget class */
+	UClass* TutorialWidgetClass;
 	
 	//
 	// Dragging UI
