VR Editor: Window controls are now hidden unless you aim your laser at them

diff --git a/Engine/Source/Editor/VREditor/Public/IVREditorMode.h b/Engine/Source/Editor/VREditor/Public/IVREditorMode.h
--- a/Engine/Source/Editor/VREditor/Public/IVREditorMode.h
+++ b/Engine/Source/Editor/VREditor/Public/IVREditorMode.h
@@ -73,7 +73,7 @@ public:
 	 *
 	 * @return	True if we have motion controller data for this hand and could return a valid result
 	 */
-	virtual bool GetLaserPointer( const int32 HandIndex, FVector& LaserPointerStart, FVector& LaserPointerEnd, const bool bEvenIfUIIsInFront = false, const float LaserLengthOverride = 0.0f ) = 0;
+	virtual bool GetLaserPointer( const int32 HandIndex, FVector& LaserPointerStart, FVector& LaserPointerEnd, const bool bEvenIfUIIsInFront = false, const float LaserLengthOverride = 0.0f ) const = 0;
 
 	/**
 	 * Gets our avatar's mesh actor
diff --git a/Engine/Source/Editor/VREditor/VREditorDockableWindow.cpp b/Engine/Source/Editor/VREditor/VREditorDockableWindow.cpp
--- a/Engine/Source/Editor/VREditor/VREditorDockableWindow.cpp
+++ b/Engine/Source/Editor/VREditor/VREditorDockableWindow.cpp
@@ -129,6 +129,11 @@ AVREditorDockableWindow::AVREditorDockableWindow()
 		check( CloseButtonTranslucentMID != nullptr );
 		CloseButtonMeshComponent->SetMaterial( 1, CloseButtonTranslucentMID );
 	}
+
+	// The selection bar and close button will not be initially visible.  They'll appear when the user aims
+	// their laser toward the UI
+	SelectionBarMeshComponent->SetVisibility( false );
+	CloseButtonMeshComponent->SetVisibility( false );
 }
 
 
@@ -150,6 +155,31 @@ void AVREditorDockableWindow::TickManually( float DeltaTime )
 		const float WorldScaleFactor = GetOwner().GetOwner().GetWorldScaleFactor();
 		const FVector AnimatedScale = CalculateAnimatedScale();
 
+		
+		bool bIsLaserAimingTowardUI = false;
+		{
+			const FTransform UICapsuleTransform = this->GetActorTransform();
+
+			const FVector UICapsuleStart = FVector( 0.0f, 0.0f, -GetSize().Y * 0.4f ) * WorldScaleFactor * AnimatedScale;
+			const FVector UICapsuleEnd = FVector( 0.0f, 0.0f, GetSize().Y * 0.4f ) * WorldScaleFactor * AnimatedScale;
+			const float UICapsuleLocalRadius = GetSize().X * 0.5f * WorldScaleFactor * AnimatedScale.Y;
+			const float MinDistanceToUICapsule = 10.0f * WorldScaleFactor * AnimatedScale.Y;	// @todo vreditor tweak
+			const FVector UIForwardVector = FVector::ForwardVector;
+			const float MinDotForAimingAtOtherHand = -1.1f;	// @todo vreditor tweak
+
+			for( int32 HandIndex = 0; HandIndex < VREditorConstants::NumVirtualHands; ++HandIndex )
+			{
+				if( GetOwner().GetOwner().IsHandAimingTowardsCapsule( HandIndex, UICapsuleTransform, UICapsuleStart, UICapsuleEnd, UICapsuleLocalRadius, MinDistanceToUICapsule, UIForwardVector, MinDotForAimingAtOtherHand ) )
+				{
+					bIsLaserAimingTowardUI = true;
+				}
+			}
+		}
+
+		// Only show our extra buttons and controls if the user is roughly aiming toward us.  This just reduces clutter.
+		SelectionBarMeshComponent->SetVisibility( bIsLaserAimingTowardUI );
+		CloseButtonMeshComponent->SetVisibility( bIsLaserAimingTowardUI );
+
 		// Update the window border mesh
 		{
 			const float WindowMeshSize = 100.0f;	// Size of imported mesh, we need to inverse compensate for
diff --git a/Engine/Source/Editor/VREditor/VREditorMode.cpp b/Engine/Source/Editor/VREditor/VREditorMode.cpp
--- a/Engine/Source/Editor/VREditor/VREditorMode.cpp
+++ b/Engine/Source/Editor/VREditor/VREditorMode.cpp
@@ -1617,9 +1617,10 @@ void FVREditorMode::SetRoomTransform( const FTransform& NewRoomTransform )
 }
 
 
-bool FVREditorMode::GetHandTransformAndForwardVector( int32 HandIndex, FTransform& OutHandTransform, FVector& OutForwardVector )
+bool FVREditorMode::GetHandTransformAndForwardVector( int32 HandIndex, FTransform& OutHandTransform, FVector& OutForwardVector ) const
 {
-	PollInputIfNeeded();
+	FVREditorMode* MutableThis = const_cast<FVREditorMode*>( this );
+	MutableThis->PollInputIfNeeded();
 
 	const FVirtualHand& Hand = VirtualHands[ HandIndex ];
 	if( Hand.bHaveMotionController )
@@ -1642,7 +1643,7 @@ bool FVREditorMode::GetHandTransformAndForwardVector( int32 HandIndex, FTransfor
 }
 
 
-bool FVREditorMode::GetLaserPointer( const int32 HandIndex, FVector& LaserPointerStart, FVector& LaserPointerEnd, const bool bEvenIfUIIsInFront, const float LaserLengthOverride )
+bool FVREditorMode::GetLaserPointer( const int32 HandIndex, FVector& LaserPointerStart, FVector& LaserPointerEnd, const bool bEvenIfUIIsInFront, const float LaserLengthOverride ) const
 {
 	const FVirtualHand& Hand = VirtualHands[ HandIndex ];
 
@@ -2542,22 +2543,22 @@ void FVREditorMode::ApplyButtonPressColors( FVRAction VRAction, EInputEvent Even
 		{
 			Hand.HandMeshMID->SetScalarParameterValue( StaticShoulderParameter, PressStrength );
 		}
-		else if (Event == IE_Released)
+		else if( Event == IE_Released )
 		{
 			Hand.HandMeshMID->SetScalarParameterValue( StaticShoulderParameter, 0.0f );
 		}
 	}
 
 	//Trackpad
-	if ( VRAction.ActionType == EVRActionType::ConfirmRadialSelection )
+	if( VRAction.ActionType == EVRActionType::ConfirmRadialSelection )
 	{
 		static FName StaticTrackpadParameter( "B3" );
 
-		if (Event == IE_Pressed)
+		if( Event == IE_Pressed )
 		{
 			Hand.HandMeshMID->SetScalarParameterValue( StaticTrackpadParameter, PressStrength );
 		}
-		else if (Event == IE_Released)
+		else if( Event == IE_Released )
 		{
 			Hand.HandMeshMID->SetScalarParameterValue( StaticTrackpadParameter, 0.0f );
 		}
@@ -2580,8 +2581,8 @@ void FVREditorMode::ApplyButtonPressColors( FVRAction VRAction, EInputEvent Even
 
 void FVREditorMode::CycleTransformGizmoHandleType()
 {
-	EGizmoHandleTypes NewGizmoType = (EGizmoHandleTypes)((uint8)CurrentGizmoType + 1);
-	if (NewGizmoType > EGizmoHandleTypes::Scale)
+	EGizmoHandleTypes NewGizmoType = (EGizmoHandleTypes)( (uint8)CurrentGizmoType + 1 );
+	if( NewGizmoType > EGizmoHandleTypes::Scale )
 	{
 		NewGizmoType = EGizmoHandleTypes::All;
 	}
@@ -2621,4 +2622,65 @@ FLinearColor FVREditorMode::GetColor( const EColors Color ) const
 }
 
 
+bool FVREditorMode::IsHandAimingTowardsCapsule( const int32 HandIndex, const FTransform& CapsuleTransform, FVector CapsuleStart, const FVector CapsuleEnd, const float CapsuleRadius, const float MinDistanceToCapsule, const FVector CapsuleFrontDirection, const float MinDotForAimingAtCapsule ) const
+{
+	bool bIsAimingTowards = false;
+	const float WorldScaleFactor = GetWorldScaleFactor();
+
+	FVector LaserPointerStart, LaserPointerEnd;
+	if( GetLaserPointer( HandIndex, /* Out */ LaserPointerStart, /* Out */ LaserPointerEnd ) )
+	{
+		const FVector LaserPointerStartInCapsuleSpace = CapsuleTransform.InverseTransformPosition( LaserPointerStart );
+		const FVector LaserPointerEndInCapsuleSpace = CapsuleTransform.InverseTransformPosition( LaserPointerEnd );
+
+		FVector ClosestPointOnLaserPointer, ClosestPointOnUICapsule;
+		FMath::SegmentDistToSegment(
+			LaserPointerStartInCapsuleSpace, LaserPointerEndInCapsuleSpace,
+			CapsuleStart, CapsuleEnd,
+			/* Out */ ClosestPointOnLaserPointer,
+			/* Out */ ClosestPointOnUICapsule );
+
+		const bool bIsClosestPointInsideCapsule = ( ClosestPointOnLaserPointer - ClosestPointOnUICapsule ).Size() <= CapsuleRadius;
+
+		const FVector TowardLaserPointerVector = ( ClosestPointOnLaserPointer - ClosestPointOnUICapsule ).GetSafeNormal();
+
+		// Apply capsule radius
+		ClosestPointOnUICapsule += TowardLaserPointerVector * CapsuleRadius;
+
+		if( false )	// @todo vreditor debug
+		{
+			const float RenderCapsuleLength = ( CapsuleEnd - CapsuleStart ).Size() + CapsuleRadius * 2.0f;
+			// @todo vreditor:  This capsule draws with the wrong orientation
+			if( false )
+			{
+				DrawDebugCapsule( GetWorld(), CapsuleTransform.TransformPosition( CapsuleStart + ( CapsuleEnd - CapsuleStart ) * 0.5f ), RenderCapsuleLength * 0.5f, CapsuleRadius, CapsuleTransform.GetRotation() * FRotator( 90.0f, 0, 0 ).Quaternion(), FColor::Green, false, 0.0f );
+			}
+			DrawDebugLine( GetWorld(), CapsuleTransform.TransformPosition( ClosestPointOnLaserPointer ), CapsuleTransform.TransformPosition( ClosestPointOnUICapsule ), FColor::Green, false, 0.0f );
+			DrawDebugSphere( GetWorld(), CapsuleTransform.TransformPosition( ClosestPointOnLaserPointer ), 1.5f * WorldScaleFactor, 32, FColor::Red, false, 0.0f );
+			DrawDebugSphere( GetWorld(), CapsuleTransform.TransformPosition( ClosestPointOnUICapsule ), 1.5f * WorldScaleFactor, 32, FColor::Green, false, 0.0f );
+		}
+
+		// If we're really close to the capsule
+		if( bIsClosestPointInsideCapsule ||
+			( ClosestPointOnUICapsule - ClosestPointOnLaserPointer ).Size() <= MinDistanceToCapsule )
+		{
+			const FVector LaserPointerDirectionInCapsuleSpace = ( LaserPointerEndInCapsuleSpace - LaserPointerStartInCapsuleSpace ).GetSafeNormal();
+
+			if( false )	// @todo vreditor debug
+			{
+				DrawDebugLine( GetWorld(), CapsuleTransform.TransformPosition( FVector::ZeroVector ), CapsuleTransform.TransformPosition( CapsuleFrontDirection * 5.0f ), FColor::Yellow, false, 0.0f );
+				DrawDebugLine( GetWorld(), CapsuleTransform.TransformPosition( FVector::ZeroVector ), CapsuleTransform.TransformPosition( -LaserPointerDirectionInCapsuleSpace * 5.0f ), FColor::Purple, false, 0.0f );
+			}
+
+			const float Dot = FVector::DotProduct( CapsuleFrontDirection, -LaserPointerDirectionInCapsuleSpace );
+			if( Dot >= MinDotForAimingAtCapsule )
+			{
+				bIsAimingTowards = true;
+			}
+		}
+	}
+
+	return bIsAimingTowards;
+}
+
 #undef LOCTEXT_NAMESPACE
\ No newline at end of file
diff --git a/Engine/Source/Editor/VREditor/VREditorMode.h b/Engine/Source/Editor/VREditor/VREditorMode.h
--- a/Engine/Source/Editor/VREditor/VREditorMode.h
+++ b/Engine/Source/Editor/VREditor/VREditorMode.h
@@ -72,7 +72,7 @@ public:
 	virtual void Render( const FSceneView* SceneView, FViewport* Viewport, FPrimitiveDrawInterface* PDI ) override;	
 
 	// IVREditorMode interface
-	virtual bool GetLaserPointer( const int32 HandIndex, FVector& LaserPointerStart, FVector& LaserPointerEnd, const bool bEvenIfUIIsInFront = false, const float LaserLengthOverride = 0.0f) override;
+	virtual bool GetLaserPointer( const int32 HandIndex, FVector& LaserPointerStart, FVector& LaserPointerEnd, const bool bEvenIfUIIsInFront = false, const float LaserLengthOverride = 0.0f) const override;
 	virtual bool IsInputCaptured( const FVRAction VRAction ) const;
 	virtual AActor* GetAvatarMeshActor() override
 	{
@@ -109,6 +109,13 @@ public:
 		return VirtualHands[ VirtualHandIndex ];
 	}
 
+	/** Gets a virtual hand (const) */
+	const FVirtualHand& GetVirtualHand( const int32 VirtualHandIndex ) const
+	{
+		check( VirtualHandIndex >= 0 && VirtualHandIndex < VREditorConstants::NumVirtualHands );
+		return VirtualHands[ VirtualHandIndex ];
+	}
+
 	/** Given an index of a hand, gets the index of the opposite hand */
 	int32 GetOtherHandIndex( const int32 HandIndex ) const
 	{
@@ -123,6 +130,12 @@ public:
 		return VirtualHands[ GetOtherHandIndex( HandIndex ) ];
 	}
 
+	/** Gets the other hand, given a hand index (const) */
+	const FVirtualHand& GetOtherHand( const int32 HandIndex ) const
+	{
+		return VirtualHands[ GetOtherHandIndex( HandIndex ) ];
+	}
+
 	/** Gets the world space transform of the calibrated VR room origin.  When using a seated VR device, this will feel like the
 	    camera's world transform (before any HMD positional or rotation adjustments are applied.) */
 	FTransform GetRoomTransform() const;
@@ -262,6 +275,8 @@ public:
 	/** @return Returns the type of HMD we're dealing with */
 	EHMDDeviceType::Type GetHMDDeviceType() const;
 
+	/** @return Checks to see if the specified hand is aiming roughly toward the specified capsule */
+	bool IsHandAimingTowardsCapsule( const int32 HandIndex, const FTransform& CapsuleTransform, const FVector CapsuleStart, const FVector CapsuleEnd, const float CapsuleRadius, const float MinDistanceToCapsule, const FVector CapsuleFrontDirection, const float MinDotForAimingAtCapsule ) const;
 
 private:
 
@@ -277,7 +292,7 @@ private:
 	 *
 	 * @return	True if we have motion controller data for this hand and could return a valid result
 	 */
-	bool GetHandTransformAndForwardVector( int32 HandIndex, FTransform& OutHandTransform, FVector& OutForwardVector );
+	bool GetHandTransformAndForwardVector( int32 HandIndex, FTransform& OutHandTransform, FVector& OutForwardVector ) const;
 
 	/**
 	 * Translates a raw key into a vr editing action
diff --git a/Engine/Source/Editor/VREditor/VREditorUISystem.cpp b/Engine/Source/Editor/VREditor/VREditorUISystem.cpp
--- a/Engine/Source/Editor/VREditor/VREditorUISystem.cpp
+++ b/Engine/Source/Editor/VREditor/VREditorUISystem.cpp
@@ -468,58 +468,18 @@ void FVREditorUISystem::Tick( FEditorViewportClient* ViewportClient, const float
 			const FVirtualHand& OtherHand = Owner.GetOtherHand( HandIndex );
 			const int32 OtherHandIndex = Owner.GetOtherHandIndex( HandIndex );		
 
-			FVector LaserPointerStart, LaserPointerEnd;
-			if( Owner.GetLaserPointer( HandIndex, /* Out */ LaserPointerStart, /* Out */ LaserPointerEnd ) )
+			// @todo vreditor tweak: Weird to hard code this here.  Probably should be an accessor on the hand itself, and based on the actual device type
+			const FTransform UICapsuleTransform = OtherHand.Transform;
+			const FVector UICapsuleStart = FVector( -6.0f, 0.0f, 0.0f ) * WorldScaleFactor;
+			const FVector UICapsuleEnd = FVector( -14.0f, 0.0f, 0.0f ) * WorldScaleFactor;
+			const float UICapsuleLocalRadius = 6.0f * WorldScaleFactor;
+			const float MinDistanceToUICapsule = 8.0f * WorldScaleFactor;	// @todo vreditor tweak
+			const FVector UIForwardVector = FVector::UpVector;
+			const float MinDotForAimingAtOtherHand = 0.25f;	// @todo vreditor tweak
+
+			if( Owner.IsHandAimingTowardsCapsule( HandIndex, UICapsuleTransform, UICapsuleStart, UICapsuleEnd, UICapsuleLocalRadius, MinDistanceToUICapsule, UIForwardVector, MinDotForAimingAtOtherHand ) )
 			{
-				const FVector LaserPointerStartInOtherHandSpace = OtherHand.Transform.InverseTransformPosition( LaserPointerStart );
-				const FVector LaserPointerEndInOtherHandSpace = OtherHand.Transform.InverseTransformPosition( LaserPointerEnd );
-
-				// @todo vreditor tweak: Weird to hard code this here.  Probably should be an accessor on the hand itself, and based on the actual device type
-				const float UICapsuleLocalRadius = 6.0f * WorldScaleFactor;
-				const FVector UICapsuleStart = FVector( -6.0f, 0.0f, 0.0f ) * WorldScaleFactor;
-				const FVector UICapsuleEnd = FVector( -14.0f, 0.0f, 0.0f ) * WorldScaleFactor;
-				FVector ClosestPointOnLaserPointer, ClosestPointOnUICapsule;
-				FMath::SegmentDistToSegment(
-					LaserPointerStartInOtherHandSpace, LaserPointerEndInOtherHandSpace,
-					UICapsuleStart, UICapsuleEnd,
-					/* Out */ ClosestPointOnLaserPointer,
-					/* Out */ ClosestPointOnUICapsule );
-
-				const FVector TowardLaserPointerVector = ( ClosestPointOnLaserPointer - ClosestPointOnUICapsule ).GetSafeNormal();
-
-				// Apply capsule radius
-				ClosestPointOnUICapsule += TowardLaserPointerVector * UICapsuleLocalRadius;
-
-				if( false )	// @todo vreditor debug
-				{
-					const float RenderCapsuleLength = ( UICapsuleEnd - UICapsuleStart ).Size() + UICapsuleLocalRadius * 2.0f;
-					DrawDebugCapsule( Owner.GetWorld(), OtherHand.Transform.TransformPosition( UICapsuleStart + ( UICapsuleEnd - UICapsuleStart ) * 0.5f ), RenderCapsuleLength * 0.5f, UICapsuleLocalRadius, OtherHand.Transform.GetRotation() * FRotator( 90.0f, 0, 0 ).Quaternion(), FColor::Green, false, 0.0f );
-					DrawDebugLine( Owner.GetWorld(), OtherHand.Transform.TransformPosition( ClosestPointOnLaserPointer ), OtherHand.Transform.TransformPosition( ClosestPointOnUICapsule ), FColor::Green, false, 0.0f );
-					DrawDebugSphere( Owner.GetWorld(), OtherHand.Transform.TransformPosition( ClosestPointOnLaserPointer ), 1.5f * WorldScaleFactor, 32, FColor::Red, false, 0.0f );
-					DrawDebugSphere( Owner.GetWorld(), OtherHand.Transform.TransformPosition( ClosestPointOnUICapsule ), 1.5f * WorldScaleFactor, 32, FColor::Green, false, 0.0f );
-				}
-				
-				// If we're really close to the capsule
-				const float MinDistanceToUICapsule = 8.0f * WorldScaleFactor;	// @todo vreditor tweak
-				if( ( ClosestPointOnUICapsule - ClosestPointOnLaserPointer ).Size() < MinDistanceToUICapsule )
-				{
-					// Make the quick menu UI visible
-					const FVector LaserPointerDirectionInOtherHandSpace = ( LaserPointerEndInOtherHandSpace - LaserPointerStartInOtherHandSpace ).GetSafeNormal();
-					const FVector UIForwardVector = FVector::UpVector;
-
-					if( false )	// @todo vreditor debug
-					{
-						DrawDebugLine( Owner.GetWorld(), OtherHand.Transform.TransformPosition( FVector::ZeroVector ), OtherHand.Transform.TransformPosition( UIForwardVector * 5.0f ), FColor::Yellow, false, 0.0f );
-						DrawDebugLine( Owner.GetWorld(), OtherHand.Transform.TransformPosition( FVector::ZeroVector ), OtherHand.Transform.TransformPosition( -LaserPointerDirectionInOtherHandSpace * 5.0f ), FColor::Purple, false, 0.0f );
-					}
-
-					const float MinDotForAimingAtOtherHand = 0.25f;	// @todo vreditor tweak
-					const float Dot = FVector::DotProduct( UIForwardVector, -LaserPointerDirectionInOtherHandSpace );
-					if( Dot >= MinDotForAimingAtOtherHand )
-					{
-						bShowQuickMenuOnArm = true;
-					}
-				}
+				bShowQuickMenuOnArm = true;
 			}
 
 			if( bShowQuickMenuOnArm )
@@ -1066,4 +1026,4 @@ void FVREditorUISystem::TogglePanelsVisibility()
 	{
 		QuickMenuUI->GetUserWidget<UVREditorQuickMenu>()->RefreshUI();
 	}
-}
+}
\ No newline at end of file
