VR Editor : Fixed bug that caused the laserpointer not to be spawned back after closing a panel on a hand

diff --git a/Engine/Source/Editor/VREditor/VREditorWorldInteraction.cpp b/Engine/Source/Editor/VREditor/VREditorWorldInteraction.cpp
--- a/Engine/Source/Editor/VREditor/VREditorWorldInteraction.cpp
+++ b/Engine/Source/Editor/VREditor/VREditorWorldInteraction.cpp
@@ -394,7 +394,7 @@ void FVREditorWorldInteraction::OnVRAction( FEditorViewportClient& ViewportClien
 										const bool bShouldShow = false;
 										const bool bShowOnHand = false;
 										const bool bRefreshQuickMenu = true;
-										Owner.GetUISystem().ShowEditorUIPanel( DockableWindow, VRAction.HandIndex, bShouldShow, bShowOnHand, bRefreshQuickMenu );
+										Owner.GetUISystem().ShowEditorUIPanel( DockableWindow, Owner.GetOtherHandIndex( VRAction.HandIndex ), bShouldShow, bShowOnHand, bRefreshQuickMenu );
 									}
 									else
 									{
