Re-disabling the filter/path popup in CB in the VREditor.

diff --git a/Engine/Source/Editor/VREditor/VREditorUISystem.cpp b/Engine/Source/Editor/VREditor/VREditorUISystem.cpp
--- a/Engine/Source/Editor/VREditor/VREditorUISystem.cpp
+++ b/Engine/Source/Editor/VREditor/VREditorUISystem.cpp
@@ -601,9 +601,9 @@ void FVREditorUISystem::CreateUIs()
 			Config.bCanShowRealTimeThumbnails = false;
 			Config.InitialAssetViewType = EAssetViewType::Tile;
 			Config.bCanShowDevelopersFolder = false;
-			Config.bCanShowFolders = true;
-			Config.bUsePathPicker = true;
-			Config.bCanShowFilters = true;
+			Config.bCanShowFolders = false;
+			Config.bUsePathPicker = false;
+			Config.bCanShowFilters = false;
 			Config.bCanShowAssetSearch = false;
 			Config.bUseSourcesView = true;
 			Config.bExpandSourcesView = true;
