VR Editor: Tweaked close button look on dockable windows

diff --git a/Engine/Source/Editor/VREditor/VREditorDockableWindow.cpp b/Engine/Source/Editor/VREditor/VREditorDockableWindow.cpp
--- a/Engine/Source/Editor/VREditor/VREditorDockableWindow.cpp
+++ b/Engine/Source/Editor/VREditor/VREditorDockableWindow.cpp
@@ -11,7 +11,7 @@ namespace VREd
 {
 	static FAutoConsoleVariable DockWindowThickness( TEXT( "VREd.DockWindowTickness" ), 1.0f, TEXT( "How thick the window is" ) );
 	static FAutoConsoleVariable DockUISelectionBarVerticalOffset( TEXT( "VREd.DockUISelectionBarVerticalOffset" ), 2.0f, TEXT( "Z Distance between the selectionbar and the UI" ) );
-	static FAutoConsoleVariable DockUICloseButtonOffsetFromCorner( TEXT( "VREd.DockUICloseButtonOffsetFromCorner" ), 2.0f, TEXT( "How far away from the corner (along each 2D axis) of the UI the close button draws" ) );
+	static FAutoConsoleVariable DockUICloseButtonOffsetFromCorner( TEXT( "VREd.DockUICloseButtonOffsetFromCorner" ), 1.5f, TEXT( "How far away from the corner (along each 2D axis) of the UI the close button draws" ) );
 }
 
 AVREditorDockableWindow::AVREditorDockableWindow() 
@@ -208,7 +208,7 @@ void AVREditorDockableWindow::TickManually( float DeltaTime )
 		// Update the close button
 		{
 			// How big the close button should be
-			const FVector CloseButtonSize( 1.0f, 3.0f, 3.0f );	// @todo vreditor tweak
+			const FVector CloseButtonSize( 1.0f, 2.0f, 2.0f );	// @todo vreditor tweak
 
 			const FVector CloseButtonScale = CloseButtonSize * AnimatedScale * WorldScaleFactor;
 			CloseButtonMeshComponent->SetRelativeScale3D( CloseButtonScale );
