VR Editor: Tweaked close button look on dockable windows (part 2)

diff --git a/Engine/Source/Editor/VREditor/VREditorMode.cpp b/Engine/Source/Editor/VREditor/VREditorMode.cpp
--- a/Engine/Source/Editor/VREditor/VREditorMode.cpp
+++ b/Engine/Source/Editor/VREditor/VREditorMode.cpp
@@ -212,7 +212,7 @@ void FVREditorMode::Enter()
 		Colors[ (int32)EColors::HoverGizmoColor ] = FLinearColor( FLinearColor::Yellow );
 		Colors[ (int32)EColors::DraggingGizmoColor ] = FLinearColor( FLinearColor::White );
 		Colors[ (int32)EColors::UISelectionBarColor ] = FLinearColor( 0.0f, 0.7f, 0.7f, 1.0f );
-		Colors[ (int32)EColors::UICloseButtonColor ] = FLinearColor( 0.8f, 0.0f, 0.0f, 1.0f );
+		Colors[ (int32)EColors::UICloseButtonColor ] = FLinearColor( 0.1f, 0.1f, 0.1f, 1.0f );
 	}
 
 	// @todo vreditor: We need to make sure the user can never switch to orthographic mode, or activate settings that
