VR Editor: Fixed a few hacks that aren't needed anymore

diff --git a/Engine/Source/Editor/VREditor/VREditorFloatingUI.cpp b/Engine/Source/Editor/VREditor/VREditorFloatingUI.cpp
--- a/Engine/Source/Editor/VREditor/VREditorFloatingUI.cpp
+++ b/Engine/Source/Editor/VREditor/VREditorFloatingUI.cpp
@@ -6,8 +6,7 @@
 #include "VREditorBaseUserWidget.h"
 #include "VREditorMode.h"
 #include "WidgetComponent.h"
-#include "LevelEditor.h"
-#include "ILevelViewport.h"
+#include "SLevelViewport.h"	// For tab manager tricks
 
 
 namespace VREd
@@ -446,28 +445,14 @@ void AVREditorFloatingUI::ShowUI( const bool bShow, const bool bAllowFading, con
 			FadeAlpha = bShow ? 1.0f : 0.0f;
 		}
 
-		// @todo vreditor Could fail a million different ways.
+		const TSharedRef<SViewport>& ViewportWidget = Owner->GetOwner().GetLevelViewportPossessedForVR().GetViewportWidget().Pin().ToSharedRef();
 		if ( bShow )
 		{
-			FLevelEditorModule& LevelEditorModule = FModuleManager::Get().GetModuleChecked<FLevelEditorModule>(TEXT("LevelEditor"));
-			TSharedPtr<ILevelViewport> ActiveLevelViewport = LevelEditorModule.GetFirstActiveViewport();
-			
-			if ( ActiveLevelViewport.IsValid() )
-			{
-				TSharedPtr<SViewport> ViewportWidget = ActiveLevelViewport->GetViewportWidget().Pin();
-				WidgetComponent->RegisterHitTesterWithViewport(ViewportWidget);
-			}
+			WidgetComponent->RegisterHitTesterWithViewport(ViewportWidget);
 		}
 		else
 		{
-			FLevelEditorModule& LevelEditorModule = FModuleManager::Get().GetModuleChecked<FLevelEditorModule>(TEXT("LevelEditor"));
-			TSharedPtr<ILevelViewport> ActiveLevelViewport = LevelEditorModule.GetFirstActiveViewport();
-
-			if ( ActiveLevelViewport.IsValid() )
-			{
-				TSharedPtr<SViewport> ViewportWidget = ActiveLevelViewport->GetViewportWidget().Pin();
-				WidgetComponent->UnregisterHitTesterWithViewport(ViewportWidget);
-			}
+			WidgetComponent->UnregisterHitTesterWithViewport(ViewportWidget);
 		}
 
 		FadeDelay = InitFadeDelay;
diff --git a/Engine/Source/Editor/VREditor/VREditorWorldInteraction.cpp b/Engine/Source/Editor/VREditor/VREditorWorldInteraction.cpp
--- a/Engine/Source/Editor/VREditor/VREditorWorldInteraction.cpp
+++ b/Engine/Source/Editor/VREditor/VREditorWorldInteraction.cpp
@@ -2683,13 +2683,6 @@ void FVREditorWorldInteraction::OnAssetDragStartedFromContentBrowser( const TArr
 		}
 	}
 
-	// We're always expecting a hand to be hovering at the time we receive this event, but in forced VR mode
-	// there may not be a hand if you're interacting with the legit CB.
-	if ( !FVREditorMode::IsActuallyUsingVR() && PlacingWithHandIndex == INDEX_NONE )
-	{
-		return;
-	}
-
 	// We're always expecting a hand to be hovering at the time we receive this event
 	if( PlacingWithHandIndex != INDEX_NONE )	// @todo vreditor: This ensure fires when in "fake VR" mode when you drag an asset (Expected)
 	{
