VREditorWidgetComponent now allows controlling when drawing can occur.  The first frame is always drawn.

diff --git a/Engine/Source/Editor/VREditor/VREditorWidgetComponent.cpp b/Engine/Source/Editor/VREditor/VREditorWidgetComponent.cpp
--- a/Engine/Source/Editor/VREditor/VREditorWidgetComponent.cpp
+++ b/Engine/Source/Editor/VREditor/VREditorWidgetComponent.cpp
@@ -6,5 +6,23 @@
 UVREditorWidgetComponent::UVREditorWidgetComponent(const FObjectInitializer& ObjectInitializer)
 	: Super( ObjectInitializer )
 {
+	bAllowDrawing = false;
+	bHasEverDrawn = false;
+}
+
+bool UVREditorWidgetComponent::ShouldDrawWidget() const
+{
+	if ( bAllowDrawing || !bHasEverDrawn )
+	{
+		return Super::ShouldDrawWidget();
+	}
+
+	return false;
+}
+
+void UVREditorWidgetComponent::DrawWidgetToRenderTarget(float DeltaTime)
+{
+	bHasEverDrawn = true;
 
+	Super::DrawWidgetToRenderTarget(DeltaTime);
 }
diff --git a/Engine/Source/Editor/VREditor/VREditorWidgetComponent.h b/Engine/Source/Editor/VREditor/VREditorWidgetComponent.h
--- a/Engine/Source/Editor/VREditor/VREditorWidgetComponent.h
+++ b/Engine/Source/Editor/VREditor/VREditorWidgetComponent.h
@@ -16,4 +16,27 @@ class UVREditorWidgetComponent : public UWidgetComponent
 public:
 	/** Default constructor that sets up CDO properties */
 	UVREditorWidgetComponent(const FObjectInitializer& ObjectInitializer);
+
+	void SetAllowDrawing(bool Value) { bAllowDrawing = Value; }
+	bool GetAllowDrawing() const { return bAllowDrawing; }
+
+protected:
+	virtual bool ShouldDrawWidget() const override;
+
+	virtual void DrawWidgetToRenderTarget(float DeltaTime) override;
+
+private:
+	/**
+	 * Controls if we draw, the VREditorWidget allows for manual enabling or 
+	 * disabling of updating the slate widget.
+	 */
+	UPROPERTY()
+	bool bAllowDrawing;
+
+	/**
+	 * Recorders if we've drawn at least once, that way we can always draw the first
+	 * frame then go manual after that.
+	 */
+	UPROPERTY()
+	bool bHasEverDrawn;
 };
\ No newline at end of file
diff --git a/Engine/Source/Runtime/UMG/Public/Components/WidgetComponent.h b/Engine/Source/Runtime/UMG/Public/Components/WidgetComponent.h
--- a/Engine/Source/Runtime/UMG/Public/Components/WidgetComponent.h
+++ b/Engine/Source/Runtime/UMG/Public/Components/WidgetComponent.h
@@ -204,14 +204,14 @@ public:
 
 	void SetEditTimeUsable(bool Value) { bEditTimeUsable = Value; }
 
-protected:
-	virtual bool ShouldDrawWidget() const;
-
 protected:
 	void RemoveWidgetFromScreen();
 
+	/** Allows subclasses to control if the widget should be drawn.  Called right before we draw the widget. */
+	virtual bool ShouldDrawWidget() const;
+
 	/** Draws the current widget to the render target if possible. */
-	void DrawWidgetToRenderTarget(float DeltaTime);
+	virtual void DrawWidgetToRenderTarget(float DeltaTime);
 
 protected:
 	/** The coordinate space in which to render the widget */
