VR Editor:
- Added scale for UI with trackpad
- Added button to open and close tutorial

diff --git a/Engine/Source/Editor/VREditor/VREditorFloatingUI.cpp b/Engine/Source/Editor/VREditor/VREditorFloatingUI.cpp
--- a/Engine/Source/Editor/VREditor/VREditorFloatingUI.cpp
+++ b/Engine/Source/Editor/VREditor/VREditorFloatingUI.cpp
@@ -40,6 +40,8 @@ AVREditorFloatingUI::AVREditorFloatingUI()
 	WidgetComponent = CreateDefaultSubobject<UWidgetComponent>( TEXT( "WidgetComponent" ), bTransient );
 	WidgetComponent->SetEditTimeUsable(true);
 	WidgetComponent->AttachTo( SceneComponent );
+
+	InitialScale = Scale;
 }
 
 
@@ -106,6 +108,8 @@ void AVREditorFloatingUI::SetSlateWidget( FVREditorUISystem& InitOwner, const TS
 	check( Resolution.X > 0 && Resolution.Y > 0 );
 
 	Scale = InitScale;
+	InitialScale = Scale;
+
 	DockedTo = InitDockedTo;
 
 	SetupWidgetComponent();
@@ -123,6 +127,8 @@ void AVREditorFloatingUI::SetUMGWidget( FVREditorUISystem& InitOwner, TSubclassO
 	check( Resolution.X > 0 && Resolution.Y > 0 );
 
 	Scale = InitScale;
+	InitialScale = Scale;
+
 	DockedTo = InitDockedTo;
 
 	SetupWidgetComponent();
@@ -331,6 +337,11 @@ void AVREditorFloatingUI::SetCollisionOnShow( const bool bInCollisionOnShow )
 	bCollisionOnShowUI = bInCollisionOnShow;
 }
 
+float AVREditorFloatingUI::GetInitialScale() const
+{
+	return InitialScale;
+}
+
 FTransform AVREditorFloatingUI::MakeUITransformLockedToHand( const int32 HandIndex, const bool bOnArm )
 {
 	const float WorldScaleFactor = Owner->GetOwner().GetWorldScaleFactor();
@@ -466,3 +477,17 @@ FVector2D AVREditorFloatingUI::GetSize() const
 	return FVector2D( Scale, Scale / Aspect );
 }
 
+float AVREditorFloatingUI::GetScale() const
+{
+	return Scale;
+}
+
+void AVREditorFloatingUI::SetScale( const float NewSize )
+{
+	Scale = NewSize;
+
+	const float WorldScaleFactor = Owner->GetOwner().GetWorldScaleFactor();
+	const FVector NewScale( Scale * WorldScaleFactor );
+	const float Aspect = (float)Resolution.X / (float)Resolution.Y;
+	WidgetComponent->SetWorldScale3D( FVector( 1.0f / NewScale.X, 1.0f / (float)Resolution.X, 1.0f / (float)Resolution.Y / Aspect ) * NewScale );
+}
\ No newline at end of file
diff --git a/Engine/Source/Editor/VREditor/VREditorFloatingUI.h b/Engine/Source/Editor/VREditor/VREditorFloatingUI.h
--- a/Engine/Source/Editor/VREditor/VREditorFloatingUI.h
+++ b/Engine/Source/Editor/VREditor/VREditorFloatingUI.h
@@ -71,6 +71,12 @@ public:
 	/** Returns the actual size of the UI in either axis, after scaling has been applied.  Does not take into animation or world scaling */
 	FVector2D GetSize() const;
 
+	/** Gets the scale */
+	float GetScale() const;
+
+	/** Sets a new size for the UI */
+	void SetScale( const float NewSize );
+
 	/** AActor overrides */
 	virtual void Destroyed() override;
 
@@ -114,6 +120,9 @@ public:
 	/** Sets if the collision is on when showing the UI*/
 	void SetCollisionOnShow( const bool bInCollisionOnShow );
 
+	/** Gets the initial size of this UI */
+	float GetInitialScale() const;
+
 protected:
 
 	/** Returns a scale to use for this widget that takes into account animation */
@@ -200,5 +209,8 @@ private:
 
 	/** Delay to fading in or out. Set on ShowUI and cleared on finish of fade in/out */
 	float FadeDelay;
+
+	/** The starting scale of this UI */
+	float InitialScale;
 };
 
diff --git a/Engine/Source/Editor/VREditor/VREditorQuickMenu.cpp b/Engine/Source/Editor/VREditor/VREditorQuickMenu.cpp
--- a/Engine/Source/Editor/VREditor/VREditorQuickMenu.cpp
+++ b/Engine/Source/Editor/VREditor/VREditorQuickMenu.cpp
@@ -317,4 +317,16 @@ bool UVREditorQuickMenu::IsGameModeEnabled() const
 }
 
 
+void UVREditorQuickMenu::OnTutorialButtonClicked( const bool bIsChecked )
+{
+	bool bShouldShow = !IsTutorialVisible();
+	const int32 MyHandIndex = GetOwner()->GetDockedTo() == AVREditorFloatingUI::EDockedTo::LeftArm ? VREditorConstants::LeftHandIndex : VREditorConstants::RightHandIndex;
+	GetOwner()->GetOwner().ShowEditorUIPanel( FVREditorUISystem::EEditorUIPanel::Tutorial, MyHandIndex, bShouldShow );
+}
+
+bool UVREditorQuickMenu::IsTutorialVisible() const
+{
+	return ( GetOwner()->GetOwner().IsShowingEditorUIPanel( FVREditorUISystem::EEditorUIPanel::Tutorial ) );
+}
+
 #undef LOCTEXT_NAMESPACE
\ No newline at end of file
diff --git a/Engine/Source/Editor/VREditor/VREditorQuickMenu.h b/Engine/Source/Editor/VREditor/VREditorQuickMenu.h
--- a/Engine/Source/Editor/VREditor/VREditorQuickMenu.h
+++ b/Engine/Source/Editor/VREditor/VREditorQuickMenu.h
@@ -126,6 +126,12 @@ public:
 	UFUNCTION( BlueprintPure, Category = "Default" )
 	bool IsGameModeEnabled() const;
 
+	UFUNCTION( BlueprintCallable, Category = "Default" )
+	void OnTutorialButtonClicked( const bool bIsChecked );
+
+	UFUNCTION( BlueprintPure, Category = "Default" )
+	bool IsTutorialVisible() const;
+
 	UFUNCTION(BlueprintImplementableEvent, meta=(DisplayName = "Refresh"))
 	void RefreshUI();
 
diff --git a/Engine/Source/Editor/VREditor/VREditorUISystem.cpp b/Engine/Source/Editor/VREditor/VREditorUISystem.cpp
--- a/Engine/Source/Editor/VREditor/VREditorUISystem.cpp
+++ b/Engine/Source/Editor/VREditor/VREditorUISystem.cpp
@@ -47,17 +47,20 @@ namespace VREd
 	static FAutoConsoleVariable UIRelativeScrollSpeed( TEXT( "VREd.UIRelativeScrollSpeed" ), 0.75f, TEXT( "How fast the UI scrolls when holding an analog stick" ) );
 	static FAutoConsoleVariable MinUIScrollDeltaForInertia( TEXT( "VREd.MinUIScrollDeltaForInertia" ), 0.25f, TEXT( "Minimum amount of touch pad input before inertial UI scrolling kicks in" ) );
 	static FAutoConsoleVariable MinDockDragDistance( TEXT( "VREd.MinDockDragDistance" ), 10.0f, TEXT( "Minimum amount of distance needed to drag dock from hand" ) );
-	static FAutoConsoleVariable DoubleClickTime( TEXT( "VREd.DoubleClickTime" ), 0.25f, TEXT( "Minumum duration between clicks for a double click event to fire" ) );
+	static FAutoConsoleVariable DoubleClickTime( TEXT( "VREd.DoubleClickTime" ), 0.25f, TEXT( "Minimum duration between clicks for a double click event to fire" ) );
+
+	static FAutoConsoleVariable MaxDockWindowSize( TEXT( "VREd.MaxDockWindowSize" ), 250, TEXT( "Maximum size for dockable windows" ) );
+	static FAutoConsoleVariable MinDockWindowSize( TEXT( "VREd.MinDockWindowSize" ), 40, TEXT( "Minimum size for dockable windows" ) );
 
 	// Tutorial UI commands
-	static FAutoConsoleVariable TutorialUIResolutionX( TEXT( "VREd.TutorialUI.Resolution.X" ), 1920, TEXT( "Minumum duration between clicks for a double click event to fire" ) );
-	static FAutoConsoleVariable TutorialUIResolutionY( TEXT( "VREd.TutorialUI.Resolution.Y" ), 1080, TEXT( "Minumum duration between clicks for a double click event to fire" ) );
-	static FAutoConsoleVariable TutorialUISize( TEXT( "VREd.TutorialUI.Size" ), 200, TEXT( "Minumum duration between clicks for a double click event to fire" ) );
-	static FAutoConsoleVariable TutorialUIYaw( TEXT( "VREd.TutorialUI.Yaw" ), 270, TEXT( "Minumum duration between clicks for a double click event to fire" ) );
-	static FAutoConsoleVariable TutorialUIPitch( TEXT( "VREd.TutorialUI.Pitch" ), 45, TEXT( "Minumum duration between clicks for a double click event to fire" ) );
-	static FAutoConsoleVariable TutorialUILocationX( TEXT( "VREd.TutorialUI.Location.X" ), 0, TEXT( "Minumum duration between clicks for a double click event to fire" ) );
-	static FAutoConsoleVariable TutorialUILocationY( TEXT( "VREd.TutorialUI.Location.Y" ), 200, TEXT( "Minumum duration between clicks for a double click event to fire" ) );
-	static FAutoConsoleVariable TutorialUILocationZ( TEXT( "VREd.TutorialUI.Location.Z" ), 40, TEXT( "Minumum duration between clicks for a double click event to fire" ) );
+	static FAutoConsoleVariable TutorialUIResolutionX( TEXT( "VREd.TutorialUI.Resolution.X" ), 1920, TEXT( "The X resolution for the tutorial UI panel" ) );
+	static FAutoConsoleVariable TutorialUIResolutionY( TEXT( "VREd.TutorialUI.Resolution.Y" ), 1080, TEXT( "The Y resolution for the tutorial UI panel" ) );
+	static FAutoConsoleVariable TutorialUISize( TEXT( "VREd.TutorialUI.Size" ), 200, TEXT( "The room space size for the tutorial UI panel" ) );
+	static FAutoConsoleVariable TutorialUIYaw( TEXT( "VREd.TutorialUI.Yaw" ), 270, TEXT( "The yaw rotation for the tutorial UI panel" ) );
+	static FAutoConsoleVariable TutorialUIPitch( TEXT( "VREd.TutorialUI.Pitch" ), 45, TEXT( "The pitch rotation for the tutorial UI panel" ) );
+	static FAutoConsoleVariable TutorialUILocationX( TEXT( "VREd.TutorialUI.Location.X" ), 0, TEXT( "The X location for the tutorial UI panel" ) );
+	static FAutoConsoleVariable TutorialUILocationY( TEXT( "VREd.TutorialUI.Location.Y" ), 200, TEXT( "The Y location for the tutorial UI panel" ) );
+	static FAutoConsoleVariable TutorialUILocationZ( TEXT( "VREd.TutorialUI.Location.Z" ), 40, TEXT( "The Z location for the tutorial UI panel" ) );
 }
 
 
@@ -292,7 +295,7 @@ void FVREditorUISystem::OnVRHoverUpdate( FEditorViewportClient& ViewportClient,
 {
 	FVirtualHand& Hand = Owner.GetVirtualHand( HandIndex );
 
-	if( !bWasHandled )
+	if( !bWasHandled && Hand.DraggingMode != EVREditorDraggingMode::DockableWindow )
 	{
 		FVector LaserPointerStart, LaserPointerEnd;
 		if( Owner.GetLaserPointer( HandIndex, LaserPointerStart, LaserPointerEnd ) )
@@ -864,6 +867,11 @@ void FVREditorUISystem::ShowEditorUIPanel( AVREditorFloatingUI* Panel, const int
 			const AVREditorFloatingUI::EDockedTo DockedTo = HandIndex == VREditorConstants::LeftHandIndex ? AVREditorFloatingUI::EDockedTo::LeftHand : AVREditorFloatingUI::EDockedTo::RightHand;
 			Panel->SetDockedTo( DockedTo );
 
+			if (bShouldShow)
+			{
+				Panel->SetScale( Panel->GetInitialScale() );
+			}
+
 			const FVector EditorUIRelativeOffset( Panel->GetSize().Y * 0.5f + VREd::EditorUIOffsetFromHand->GetFloat(), 0.0f, 0.0f );
 			Panel->SetRelativeOffset( EditorUIRelativeOffset );
 			Panel->SetLocalRotation( FRotator( 90.0f, 180.0f, 0.0f ) );
@@ -1002,6 +1010,11 @@ AVREditorDockableWindow* FVREditorUISystem::GetDraggingDockUI() const
 	return DraggingUI;
 }
 
+int32 FVREditorUISystem::GetDraggingDockUIHandIndex() const
+{
+	return DraggingUIHandIndex;
+}
+
 void FVREditorUISystem::StopDraggingDockUI()
 {
 	DraggingUI = nullptr;
@@ -1057,4 +1070,14 @@ void FVREditorUISystem::TogglePanelsVisibility()
 	{
 		QuickMenuUI->GetUserWidget<UVREditorQuickMenu>()->RefreshUI();
 	}
-}
\ No newline at end of file
+}
+
+float FVREditorUISystem::GetMaxDockWindowSize() const
+{
+	return VREd::MaxDockWindowSize->GetFloat();
+}
+
+float FVREditorUISystem::GetMinDockWindowSize() const
+{
+	return VREd::MinDockWindowSize->GetFloat();
+}
diff --git a/Engine/Source/Editor/VREditor/VREditorUISystem.h b/Engine/Source/Editor/VREditor/VREditorUISystem.h
--- a/Engine/Source/Editor/VREditor/VREditorUISystem.h
+++ b/Engine/Source/Editor/VREditor/VREditorUISystem.h
@@ -88,6 +88,9 @@ public:
 	/** Returns the current dragged dock window, nullptr if none */
 	class AVREditorDockableWindow* GetDraggingDockUI() const;
 
+	/** Gets the hand index of the hand that is dragging a dock UI */
+	int32 GetDraggingDockUIHandIndex() const;
+
 	/** Resets all values to stop dragging the current dock window */
 	void StopDraggingDockUI();
 
@@ -97,6 +100,12 @@ public:
 	/** Hides and unhides all the editor UI panels */
 	void TogglePanelsVisibility();
 
+	/** Get the maximum dock window size */
+	float GetMaxDockWindowSize() const;
+
+	/** Get the minimum dock window size */
+	float GetMinDockWindowSize() const;
+
 protected:
 
 	/** Called when the user presses a button on their motion controller device */
diff --git a/Engine/Source/Editor/VREditor/VREditorWorldInteraction.cpp b/Engine/Source/Editor/VREditor/VREditorWorldInteraction.cpp
--- a/Engine/Source/Editor/VREditor/VREditorWorldInteraction.cpp
+++ b/Engine/Source/Editor/VREditor/VREditorWorldInteraction.cpp
@@ -814,7 +814,6 @@ void FVREditorWorldInteraction::Tick( FEditorViewportClient* ViewportClient, con
 				if( Hand.bIsTrackpadPositionValid[ 1 ] )
 				{
 					const bool bIsAbsolute = ( Owner.GetHMDDeviceType() == EHMDDeviceType::DT_SteamVR );
-
 					float SlideDelta = 0.0f;
 					if( Hand.bIsTouchingTrackpad || !bIsAbsolute )
 					{
@@ -880,7 +879,7 @@ void FVREditorWorldInteraction::Tick( FEditorViewportClient* ViewportClient, con
 			}
 			else
 			{
-				if (!Hand.bHasUIInFront && (Hand.bIsTrackpadPositionValid[0] && Hand.bIsTrackpadPositionValid[1] ))
+				if ( !Hand.bHasUIInFront && (Hand.bIsTrackpadPositionValid[0] && Hand.bIsTrackpadPositionValid[1] ) && Owner.GetUISystem().GetDraggingDockUIHandIndex() != HandIndex )
 				{
 					if( Owner.GetHMDDeviceType() == EHMDDeviceType::DT_SteamVR && Hand.bIsTouchingTrackpad )
 					{
@@ -1202,9 +1201,35 @@ void FVREditorWorldInteraction::Tick( FEditorViewportClient* ViewportClient, con
 				AVREditorDockableWindow* DraggingUI = Owner.GetUISystem().GetDraggingDockUI();
 				if( DraggingUI )
 				{
+					//@todo: VR Editor - This should be in the UI system
+					const bool bIsAbsolute = ( Owner.GetHMDDeviceType() == EHMDDeviceType::DT_SteamVR );
+					float SlideDeltaY = 0.0f;
+					if (Hand.bIsTouchingTrackpad || !bIsAbsolute)
+					{
+						if (bIsAbsolute)
+						{
+							SlideDeltaY = -((Hand.TrackpadPosition.Y - Hand.LastTrackpadPosition.Y) * VREd::TrackpadAbsoluteDragSpeed->GetFloat()) * WorldScaleFactor;
+						}
+						else
+						{
+							SlideDeltaY = -(Hand.TrackpadPosition.Y * VREd::TrackpadRelativeDragSpeed->GetFloat()) * WorldScaleFactor;
+						}
+					}
+
+					float NewUIScale = DraggingUI->GetScale() + SlideDeltaY;
+					if( NewUIScale <= Owner.GetUISystem().GetMinDockWindowSize() )
+					{
+						NewUIScale = Owner.GetUISystem().GetMinDockWindowSize();
+					}
+					else if( NewUIScale >= Owner.GetUISystem().GetMaxDockWindowSize() )
+					{
+						NewUIScale = Owner.GetUISystem().GetMaxDockWindowSize();
+					}
+					DraggingUI->SetScale( NewUIScale );
+
 					const FTransform DockedUIToWorld = Owner.GetUISystem().MakeDockableUITransform( DraggingUI, HandIndex, DockSelectDistance );
 					DraggingUI->SetActorTransform( DockedUIToWorld );
-
+					
 					DraggingUI->UpdateRelativeRoomTransform();
 				}
 			}
