Disabling the proxy tab manager.

diff --git a/Engine/Source/Editor/VREditor/VREditorUISystem.cpp b/Engine/Source/Editor/VREditor/VREditorUISystem.cpp
--- a/Engine/Source/Editor/VREditor/VREditorUISystem.cpp
+++ b/Engine/Source/Editor/VREditor/VREditorUISystem.cpp
@@ -736,29 +736,29 @@ void FVREditorUISystem::CreateUIs()
 			EditorUIPanels[ (int32)EEditorUIPanel::Tutorial ] = TutorialUI;		
 		}
 
-		{
-			const FIntPoint Resolution(VREd::EditorUIResolutionX->GetInt(), VREd::EditorUIResolutionY->GetInt());
+		//{
+		//	const FIntPoint Resolution(VREd::EditorUIResolutionX->GetInt(), VREd::EditorUIResolutionY->GetInt());
 
-			const bool bWithSceneComponent = false;
-			AVREditorFloatingUI* TabManagerUI = GetOwner().SpawnTransientSceneActor< AVREditorFloatingUI >(TEXT("TabManager"), bWithSceneComponent);
-			TabManagerUI->SetSlateWidget(*this, SNullWidget::NullWidget, Resolution, VREd::EditorUISize->GetFloat(), AVREditorFloatingUI::EDockedTo::Nothing);
-			TabManagerUI->ShowUI(true);
-			TabManagerUI->SetRelativeOffset(FVector(0, 0, 1000));
-			FloatingUIs.Add(TabManagerUI);
+		//	const bool bWithSceneComponent = false;
+		//	AVREditorFloatingUI* TabManagerUI = GetOwner().SpawnTransientSceneActor< AVREditorFloatingUI >(TEXT("TabManager"), bWithSceneComponent);
+		//	TabManagerUI->SetSlateWidget(*this, SNullWidget::NullWidget, Resolution, VREd::EditorUISize->GetFloat(), AVREditorFloatingUI::EDockedTo::Nothing);
+		//	TabManagerUI->ShowUI(true);
+		//	TabManagerUI->SetRelativeOffset(FVector(0, 0, 1000));
+		//	FloatingUIs.Add(TabManagerUI);
 
-			TSharedPtr<SWindow> TabManagerWindow = TabManagerUI->GetWidgetComponent()->GetSlateWindow();
-			TSharedRef<SWindow> TabManagerWindowRef = TabManagerWindow.ToSharedRef();
-			ProxyTabManager = MakeShareable(new FProxyTabmanager(TabManagerWindowRef));
+		//	TSharedPtr<SWindow> TabManagerWindow = TabManagerUI->GetWidgetComponent()->GetSlateWindow();
+		//	TSharedRef<SWindow> TabManagerWindowRef = TabManagerWindow.ToSharedRef();
+		//	ProxyTabManager = MakeShareable(new FProxyTabmanager(TabManagerWindowRef));
 
-			TSharedRef<FTabManager::FLayout> LoadedLayout = FTabManager::NewLayout(TEXT("VRTabManager"));
-			LoadedLayout->AddArea(FTabManager::NewPrimaryArea());
+		//	TSharedRef<FTabManager::FLayout> LoadedLayout = FTabManager::NewLayout(TEXT("VRTabManager"));
+		//	LoadedLayout->AddArea(FTabManager::NewPrimaryArea());
 
-			TSharedPtr<SWidget> DockLayout = FGlobalTabmanager::Get()->RestoreFrom(LoadedLayout, TabManagerWindowRef, false);
-			TabManagerWindowRef->SetContent(DockLayout.ToSharedRef());
+		//	TSharedPtr<SWidget> DockLayout = FGlobalTabmanager::Get()->RestoreFrom(LoadedLayout, TabManagerWindowRef, false);
+		//	TabManagerWindowRef->SetContent(DockLayout.ToSharedRef());
 
-			// We're going to start stealing tabs from the global tab manager inserting them into the world instead.
-			FGlobalTabmanager::Get()->SetProxyTabManager(ProxyTabManager);
-		}
+		//	// We're going to start stealing tabs from the global tab manager inserting them into the world instead.
+		//	FGlobalTabmanager::Get()->SetProxyTabManager(ProxyTabManager);
+		//}
 	}
 }
 
