VR Editor: Fixed inverted joystick on Oculus, and general trackpad bugs

diff --git a/Engine/Source/Editor/VREditor/VREditorMode.cpp b/Engine/Source/Editor/VREditor/VREditorMode.cpp
--- a/Engine/Source/Editor/VREditor/VREditorMode.cpp
+++ b/Engine/Source/Editor/VREditor/VREditorMode.cpp
@@ -1513,6 +1513,8 @@ bool FVREditorMode::InputAxis( FEditorViewportClient* ViewportClient, FViewport*
 
 			if( VRAction.ActionType == EVRActionType::TrackpadPositionY )
 			{
+				Delta = -Delta;	// Y axis is inverted from HMD
+
 				Hand.LastTrackpadPosition.Y = Hand.bIsTrackpadPositionValid[1] ? Hand.TrackpadPosition.Y : Delta;
 				Hand.LastTrackpadPositionUpdateTime = FTimespan::FromSeconds( FPlatformTime::Seconds() );
 				Hand.TrackpadPosition.Y = Delta;
diff --git a/Engine/Source/Editor/VREditor/VREditorRadialMenu.cpp b/Engine/Source/Editor/VREditor/VREditorRadialMenu.cpp
--- a/Engine/Source/Editor/VREditor/VREditorRadialMenu.cpp
+++ b/Engine/Source/Editor/VREditor/VREditorRadialMenu.cpp
@@ -14,7 +14,8 @@
 
 namespace VREd
 {
-	static FAutoConsoleVariable RadialMenuSelectionRadius( TEXT( "VREd.RadialMenuSelectionRadius" ), 0.5f, TEXT( "The radius that is used for selecting buttons on the radial menu" ) );
+	static FAutoConsoleVariable RadialMenuRiftSelectionRadius( TEXT( "VREd.RadialMenuRiftSelectionRadius" ), 0.25f, TEXT( "The radius that is used for selecting buttons on the radial menu on Oculus Rift" ) );
+	static FAutoConsoleVariable RadialMenuViveSelectionRadius( TEXT( "VREd.RadialMenuViveSelectionRadius" ), 0.5f, TEXT( "The radius that is used for selecting buttons on the radial menu on Vive" ) );
 }
 
 UVREditorRadialMenu::UVREditorRadialMenu( const FObjectInitializer& ObjectInitializer )
@@ -41,7 +42,7 @@ void UVREditorRadialMenu::Update( const FVirtualHand& Hand )
 		// Check if position is from the center
 		if( !IsInMenuRadius() )
 		{
-			float Angle = FRotator::NormalizeAxis( FMath::RadiansToDegrees(FMath::Atan2( TrackpadPosition.X, -TrackpadPosition.Y ) ) );
+			float Angle = FRotator::NormalizeAxis( FMath::RadiansToDegrees(FMath::Atan2( TrackpadPosition.X, TrackpadPosition.Y ) ) );
 			TrackpadAngle = Angle;
 			if( Angle < 0)
 			{
@@ -163,7 +164,8 @@ void UVREditorRadialMenu::SelectCurrentItem( const FVirtualHand& Hand )
 
 bool UVREditorRadialMenu::IsInMenuRadius() const
 {
-	return FVector2D::Distance( FVector2D::ZeroVector, TrackpadPosition ) < VREd::RadialMenuSelectionRadius->GetFloat();
+	const float RadialMenuSelectionRadius = Owner->GetOwner().GetOwner().GetHMDDeviceType() == EHMDDeviceType::DT_OculusRift ? VREd::RadialMenuRiftSelectionRadius->GetFloat() : VREd::RadialMenuViveSelectionRadius->GetFloat();
+	return FVector2D::Distance( FVector2D::ZeroVector, TrackpadPosition ) < RadialMenuSelectionRadius;
 }
 
 void UVREditorRadialMenu::UpdateGizmoTypeLabel()
diff --git a/Engine/Source/Editor/VREditor/VREditorUISystem.cpp b/Engine/Source/Editor/VREditor/VREditorUISystem.cpp
--- a/Engine/Source/Editor/VREditor/VREditorUISystem.cpp
+++ b/Engine/Source/Editor/VREditor/VREditorUISystem.cpp
@@ -346,20 +346,20 @@ void FVREditorUISystem::OnVRHoverUpdate( FEditorViewportClient& ViewportClient,
 						// Route the mouse scrolling
 						if ( Hand.bIsTrackpadPositionValid[1] )
 						{
-							const bool bIsAbsolute = ( Owner.GetHMDDeviceType() == EHMDDeviceType::DT_OculusRift );
+							const bool bIsAbsolute = ( Owner.GetHMDDeviceType() == EHMDDeviceType::DT_SteamVR );
 
 							float ScrollDelta = 0.0f;
-							if( Hand.bIsTouchingTrackpad || Owner.GetHMDDeviceType() != EHMDDeviceType::DT_SteamVR )
+							if( Hand.bIsTouchingTrackpad || !bIsAbsolute )
 							{
 								if( bIsAbsolute )
 								{
-									const float ScrollSpeed = VREd::UIRelativeScrollSpeed->GetFloat();
-									ScrollDelta = Hand.TrackpadPosition.Y * ScrollSpeed;
+									const float ScrollSpeed = VREd::UIAbsoluteScrollSpeed->GetFloat();
+									ScrollDelta = ( Hand.TrackpadPosition.Y - Hand.LastTrackpadPosition.Y ) * ScrollSpeed;
 								}
-								else if( Hand.bIsTouchingTrackpad )
+								else
 								{
-									const float ScrollSpeed = VREd::UIAbsoluteScrollSpeed->GetFloat();
-									ScrollDelta = ( Hand.LastTrackpadPosition.Y - Hand.TrackpadPosition.Y ) * ScrollSpeed;
+									const float ScrollSpeed = VREd::UIRelativeScrollSpeed->GetFloat();
+									ScrollDelta = Hand.TrackpadPosition.Y * ScrollSpeed;
 								}
 							}
 
@@ -521,7 +521,7 @@ void FVREditorUISystem::Tick( FEditorViewportClient* ViewportClient, const float
 
 			if( Hand.bIsTrackpadPositionValid[ 0 ] && Hand.bIsTrackpadPositionValid[ 1 ] )
 			{
-				if( Hand.TrackpadPosition.Size() > 0.25f )	// @todo vreditor tweak
+				if( Hand.TrackpadPosition.Size() > 0.1f )	// @todo vreditor tweak
 				{
 					TryToSpawnRadialMenu( HandIndex );
 					
diff --git a/Engine/Source/Editor/VREditor/VREditorWorldInteraction.cpp b/Engine/Source/Editor/VREditor/VREditorWorldInteraction.cpp
--- a/Engine/Source/Editor/VREditor/VREditorWorldInteraction.cpp
+++ b/Engine/Source/Editor/VREditor/VREditorWorldInteraction.cpp
@@ -819,11 +819,11 @@ void FVREditorWorldInteraction::Tick( FEditorViewportClient* ViewportClient, con
 					{
 						if( bIsAbsolute )
 						{
-							SlideDelta = -( ( Hand.TrackpadPosition.Y - Hand.LastTrackpadPosition.Y ) * VREd::TrackpadAbsoluteDragSpeed->GetFloat() ) * WorldScaleFactor;
+							SlideDelta = ( ( Hand.TrackpadPosition.Y - Hand.LastTrackpadPosition.Y ) * VREd::TrackpadAbsoluteDragSpeed->GetFloat() ) * WorldScaleFactor;
 						}
 						else
 						{
-							SlideDelta = -( Hand.TrackpadPosition.Y * VREd::TrackpadRelativeDragSpeed->GetFloat() ) * WorldScaleFactor;
+							SlideDelta = ( Hand.TrackpadPosition.Y * VREd::TrackpadRelativeDragSpeed->GetFloat() ) * WorldScaleFactor;
 						}
 					}
 
