VR Editor: Fixed being able to drag materials onto widget handles (UE-27833)
- Also, tweaked quick menu size a bit

diff --git a/Engine/Source/Editor/VREditor/VREditorUISystem.cpp b/Engine/Source/Editor/VREditor/VREditorUISystem.cpp
--- a/Engine/Source/Editor/VREditor/VREditorUISystem.cpp
+++ b/Engine/Source/Editor/VREditor/VREditorUISystem.cpp
@@ -39,8 +39,8 @@ namespace VREd
 	static FAutoConsoleVariable QuickMenuUIResolutionY( TEXT( "VREd.QuickMenuUIResolutionY" ), 1130, TEXT( "Vertical resolution to use for Quick Menu VR UI render targets" ) );
 	static FAutoConsoleVariable ContentBrowserUISize( TEXT( "VREd.ContentBrowserUISize" ), 60.0f, TEXT( "How big content browser UIs should be" ) );
 	static FAutoConsoleVariable EditorUISize( TEXT( "VREd.EditorUISize" ), 50.0f, TEXT( "How big editor UIs should be" ) );
-	static FAutoConsoleVariable ContentBrowserUIScale( TEXT( "VREd.ContentBrowserUIScale" ), 2.5f, TEXT( "How much to scale up (or down) the content browser for VR" ) );
-	static FAutoConsoleVariable EditorUIScale( TEXT( "VREd.EditorUIScale" ), 2.5f, TEXT( "How much to scale up (or down) editor UIs for VR" ) );
+	static FAutoConsoleVariable ContentBrowserUIScale( TEXT( "VREd.ContentBrowserUIScale" ), 2.0f, TEXT( "How much to scale up (or down) the content browser for VR" ) );
+	static FAutoConsoleVariable EditorUIScale( TEXT( "VREd.EditorUIScale" ), 2.0f, TEXT( "How much to scale up (or down) editor UIs for VR" ) );
 	static FAutoConsoleVariable EditorUIOffsetFromHand( TEXT( "VREd.EditorUIOffsetFromHand" ), 12.0f, TEXT( "How far to offset editor UIs from the origin of the hand mesh" ) );
 	static FAutoConsoleVariable RadialMenuFadeDelay( TEXT( "VREd.RadialMenuFadeDelay" ), 0.2f, TEXT( "The delay for the radial menu after selecting a button" ) );
 	static FAutoConsoleVariable UIAbsoluteScrollSpeed( TEXT( "VREd.UIAbsoluteScrollSpeed" ), 8.0f, TEXT( "How fast the UI scrolls when dragging the touchpad" ) );
@@ -473,8 +473,8 @@ void FVREditorUISystem::Tick( FEditorViewportClient* ViewportClient, const float
 
 			// @todo vreditor tweak: Weird to hard code this here.  Probably should be an accessor on the hand itself, and based on the actual device type
 			const FTransform UICapsuleTransform = OtherHand.Transform;
-			const FVector UICapsuleStart = FVector( -6.0f, 0.0f, 0.0f ) * WorldScaleFactor;
-			const FVector UICapsuleEnd = FVector( -14.0f, 0.0f, 0.0f ) * WorldScaleFactor;
+			const FVector UICapsuleStart = FVector( -8.0f, 0.0f, 0.0f ) * WorldScaleFactor;
+			const FVector UICapsuleEnd = FVector( -18.0f, 0.0f, 0.0f ) * WorldScaleFactor;
 			const float UICapsuleLocalRadius = 6.0f * WorldScaleFactor;
 			const float MinDistanceToUICapsule = 8.0f * WorldScaleFactor;	// @todo vreditor tweak
 			const FVector UIForwardVector = FVector::UpVector;
@@ -504,7 +504,6 @@ void FVREditorUISystem::Tick( FEditorViewportClient* ViewportClient, const float
 
 		if( HandIndexThatNeedsQuickMenu != INDEX_NONE && !QuickMenuUI->IsUIVisible() )
 		{
-			const float Size = 20.0f;
 			const AVREditorFloatingUI::EDockedTo DockTo = ( HandIndexThatNeedsQuickMenu == VREditorConstants::LeftHandIndex ) ? AVREditorFloatingUI::EDockedTo::LeftArm : AVREditorFloatingUI::EDockedTo::RightArm;
 			QuickMenuUI->SetDockedTo( DockTo );
 			QuickMenuUI->ShowUI( true );
@@ -570,9 +569,9 @@ void FVREditorUISystem::CreateUIs()
 
 		const bool bWithSceneComponent = false;
 		QuickMenuUI = GetOwner().SpawnTransientSceneActor< AVREditorFloatingUI >(TEXT("QuickMenu"), bWithSceneComponent);
-		QuickMenuUI->SetUMGWidget( *this, QuickMenuWidgetClass, Resolution, 22.0f, AVREditorFloatingUI::EDockedTo::Nothing );
+		QuickMenuUI->SetUMGWidget( *this, QuickMenuWidgetClass, Resolution, 30.0f, AVREditorFloatingUI::EDockedTo::Nothing );
 		QuickMenuUI->ShowUI( false );
-		QuickMenuUI->SetRelativeOffset( FVector( -10.0f, 0.0f, 3.0f ) );
+		QuickMenuUI->SetRelativeOffset( FVector( -12.0f, 0.0f, 3.0f ) );
 		FloatingUIs.Add( QuickMenuUI );
 	}
 
diff --git a/Engine/Source/Editor/VREditor/VREditorWorldInteraction.cpp b/Engine/Source/Editor/VREditor/VREditorWorldInteraction.cpp
--- a/Engine/Source/Editor/VREditor/VREditorWorldInteraction.cpp
+++ b/Engine/Source/Editor/VREditor/VREditorWorldInteraction.cpp
@@ -2926,7 +2926,11 @@ void FVREditorWorldInteraction::PlaceDraggedMaterialOrTexture( const int32 HandI
 			{
 				if( IsInteractableComponent( HitResult.GetComponent() ) )	// @todo vreditor placement: We don't necessarily need to restrict to only VR-interactive components here
 				{
-					HitComponent = HitResult.GetComponent();
+					// Don't place materials on UI widget handles though!
+					if( Cast<AVREditorFloatingUI>( HitResult.GetComponent()->GetOwner() ) == nullptr )
+					{
+						HitComponent = HitResult.GetComponent();
+					}
 				}
 			}
 		}
