Continued itteration onf getting tabs opening in the world.

diff --git a/Engine/Source/Editor/VREditor/VREditorDockableWindow.cpp b/Engine/Source/Editor/VREditor/VREditorDockableWindow.cpp
--- a/Engine/Source/Editor/VREditor/VREditorDockableWindow.cpp
+++ b/Engine/Source/Editor/VREditor/VREditorDockableWindow.cpp
@@ -6,6 +6,7 @@
 #include "VREditorMode.h"
 #include "VREditorBaseUserWidget.h"
 #include "WidgetComponent.h"
+#include "VREditorWidgetComponent.h"
 
 namespace VREd
 {
diff --git a/Engine/Source/Editor/VREditor/VREditorFloatingUI.cpp b/Engine/Source/Editor/VREditor/VREditorFloatingUI.cpp
--- a/Engine/Source/Editor/VREditor/VREditorFloatingUI.cpp
+++ b/Engine/Source/Editor/VREditor/VREditorFloatingUI.cpp
@@ -6,6 +6,7 @@
 #include "VREditorBaseUserWidget.h"
 #include "VREditorMode.h"
 #include "WidgetComponent.h"
+#include "VREditorWidgetComponent.h"
 #include "SLevelViewport.h"	// For tab manager tricks
 
 
@@ -37,7 +38,7 @@ AVREditorFloatingUI::AVREditorFloatingUI()
 	check( SceneComponent != nullptr );
 	this->RootComponent = SceneComponent;
 
-	WidgetComponent = CreateDefaultSubobject<UWidgetComponent>( TEXT( "WidgetComponent" ), bTransient );
+	WidgetComponent = CreateDefaultSubobject<UVREditorWidgetComponent>( TEXT( "WidgetComponent" ), bTransient );
 	WidgetComponent->SetEditTimeUsable(true);
 	WidgetComponent->AttachTo( SceneComponent );
 
diff --git a/Engine/Source/Editor/VREditor/VREditorFloatingUI.h b/Engine/Source/Editor/VREditor/VREditorFloatingUI.h
--- a/Engine/Source/Editor/VREditor/VREditorFloatingUI.h
+++ b/Engine/Source/Editor/VREditor/VREditorFloatingUI.h
@@ -63,7 +63,7 @@ public:
 	void SetDockedTo( const EDockedTo NewDockedTo );
 
 	/** Returns the widget component for this UI, or nullptr if not spawned right now */
-	class UWidgetComponent* GetWidgetComponent()
+	class UVREditorWidgetComponent* GetWidgetComponent()
 	{
 		return WidgetComponent;
 	}
@@ -171,7 +171,7 @@ protected:
 	
 	/** When in a spawned state, this is the widget component to represent the widget */
 	UPROPERTY()
-	class UWidgetComponent* WidgetComponent;
+	class UVREditorWidgetComponent* WidgetComponent;
 	
 	/** How big the UI should be */
 	float Scale;
diff --git a/Engine/Source/Editor/VREditor/VREditorMode.cpp b/Engine/Source/Editor/VREditor/VREditorMode.cpp
--- a/Engine/Source/Editor/VREditor/VREditorMode.cpp
+++ b/Engine/Source/Editor/VREditor/VREditorMode.cpp
@@ -391,17 +391,17 @@ void FVREditorMode::Enter()
 
 	WorldInteraction.Reset( new FVREditorWorldInteraction( *this ) );
 
-	// If we're not actually in VR mode, and instead we've forced it, we should spawn some UI at the Origin to help test
-	if ( !bActuallyUsingVR )
-	{
-		const int32 HandIndex = 1;
-		const bool bShouldShow = true;
+	//// If we're not actually in VR mode, and instead we've forced it, we should spawn some UI at the Origin to help test
+	//if ( !bActuallyUsingVR )
+	//{
+	//	const int32 HandIndex = 1;
+	//	const bool bShouldShow = true;
 
-		FVirtualHand& Hand = VirtualHands[HandIndex];
-		Hand.Transform = FTransform(FVector(0, 0, 0));
+	//	FVirtualHand& Hand = VirtualHands[HandIndex];
+	//	Hand.Transform = FTransform(FVector(0, 0, 0));
 
-		UISystem->ShowEditorUIPanel(FVREditorUISystem::EEditorUIPanel::ContentBrowser, HandIndex, bShouldShow);
-	}
+	//	UISystem->ShowEditorUIPanel(FVREditorUISystem::EEditorUIPanel::ContentBrowser, HandIndex, bShouldShow);
+	//}
 
 	bIsFullyInitialized = true;
 }
diff --git a/Engine/Source/Editor/VREditor/VREditorUISystem.cpp b/Engine/Source/Editor/VREditor/VREditorUISystem.cpp
--- a/Engine/Source/Editor/VREditor/VREditorUISystem.cpp
+++ b/Engine/Source/Editor/VREditor/VREditorUISystem.cpp
@@ -13,6 +13,7 @@
 
 // UI
 #include "WidgetComponent.h"
+#include "VREditorWidgetComponent.h"
 
 // Content Browser
 #include "Editor/ContentBrowser/Public/ContentBrowserModule.h"
@@ -189,6 +190,11 @@ void FVREditorUISystem::OnVRAction( FEditorViewportClient& ViewportClient, const
 								if( WidgetComponent != Hand.HoveringOverWidgetComponent )
 								{
 									LastLocalHitLocation = LocalHitLocation;
+
+									if ( UVREditorWidgetComponent* VRWidgetComponent = Cast<UVREditorWidgetComponent>(Hand.HoveringOverWidgetComponent) )
+									{
+										VRWidgetComponent->SetIsHovering(false);
+									}
 								}
 
 								FWidgetPath WidgetPathUnderFinger = FWidgetPath( WidgetComponent->GetHitWidgetPath( HitResult.ImpactPoint, /*bIgnoreEnabledStatus*/ false ) );
@@ -210,6 +216,11 @@ void FVREditorUISystem::OnVRAction( FEditorViewportClient& ViewportClient, const
 
 									Hand.HoveringOverWidgetComponent = WidgetComponent;
 
+									if ( UVREditorWidgetComponent* VRWidgetComponent = Cast<UVREditorWidgetComponent>(WidgetComponent) )
+									{
+										VRWidgetComponent->SetIsHovering(true);
+									}
+
 									FReply Reply = FReply::Unhandled();
 									if( Event == IE_Pressed )
 									{
@@ -316,6 +327,11 @@ void FVREditorUISystem::OnVRHoverUpdate( FEditorViewportClient& ViewportClient,
 					if( WidgetComponent != Hand.HoveringOverWidgetComponent )
 					{
 						LastLocalHitLocation = LocalHitLocation;
+
+						if ( UVREditorWidgetComponent* VRWidgetComponent = Cast<UVREditorWidgetComponent>(Hand.HoveringOverWidgetComponent) )
+						{
+							VRWidgetComponent->SetIsHovering(false);
+						}
 					}
 
 					// @todo vreditor UI: There is a CursorRadius optional arg that we migth want to make use of wherever we call GetHitWidgetPath()!
@@ -342,6 +358,10 @@ void FVREditorUISystem::OnVRHoverUpdate( FEditorViewportClient& ViewportClient,
 						Hand.HoveringOverWidgetComponent = WidgetComponent;
 						bIsHoveringOverUI = true;
 
+						if ( UVREditorWidgetComponent* VRWidgetComponent = Cast<UVREditorWidgetComponent>(WidgetComponent) )
+						{
+							VRWidgetComponent->SetIsHovering(true);
+						}
 
 						// Route the mouse scrolling
 						if ( Hand.bIsTrackpadPositionValid[1] )
@@ -432,6 +452,11 @@ void FVREditorUISystem::OnVRHoverUpdate( FEditorViewportClient& ViewportClient,
 	// If nothing was hovered, make sure we tell Slate about that
 	if( !bWasHandled && Hand.HoveringOverWidgetComponent != nullptr )
 	{
+		if ( UVREditorWidgetComponent* VRWidgetComponent = Cast<UVREditorWidgetComponent>(Hand.HoveringOverWidgetComponent) )
+		{
+			VRWidgetComponent->SetIsHovering(false);
+		}
+
 		const FVector2D LastLocalHitLocation = Hand.HoveringOverWidgetComponent->GetLastLocalHitLocation();
 		Hand.HoveringOverWidgetComponent = nullptr;
 
@@ -735,29 +760,32 @@ void FVREditorUISystem::CreateUIs()
 			EditorUIPanels[ (int32)EEditorUIPanel::Tutorial ] = TutorialUI;		
 		}
 
-		//{
-		//	const FIntPoint Resolution(VREd::EditorUIResolutionX->GetInt(), VREd::EditorUIResolutionY->GetInt());
-
-		//	const bool bWithSceneComponent = false;
-		//	AVREditorFloatingUI* TabManagerUI = GetOwner().SpawnTransientSceneActor< AVREditorFloatingUI >(TEXT("TabManager"), bWithSceneComponent);
-		//	TabManagerUI->SetSlateWidget(*this, SNullWidget::NullWidget, Resolution, VREd::EditorUISize->GetFloat(), AVREditorFloatingUI::EDockedTo::Nothing);
-		//	TabManagerUI->ShowUI(true);
-		//	TabManagerUI->SetRelativeOffset(FVector(0, 0, 1000));
-		//	FloatingUIs.Add(TabManagerUI);
-
-		//	TSharedPtr<SWindow> TabManagerWindow = TabManagerUI->GetWidgetComponent()->GetSlateWindow();
-		//	TSharedRef<SWindow> TabManagerWindowRef = TabManagerWindow.ToSharedRef();
-		//	ProxyTabManager = MakeShareable(new FProxyTabmanager(TabManagerWindowRef));
-
-		//	TSharedRef<FTabManager::FLayout> LoadedLayout = FTabManager::NewLayout(TEXT("VRTabManager"));
-		//	LoadedLayout->AddArea(FTabManager::NewPrimaryArea());
-
-		//	TSharedPtr<SWidget> DockLayout = FGlobalTabmanager::Get()->RestoreFrom(LoadedLayout, TabManagerWindowRef, false);
-		//	TabManagerWindowRef->SetContent(DockLayout.ToSharedRef());
+		if ( 0 )
+		{
+			const FIntPoint Resolution(VREd::EditorUIResolutionX->GetInt(), VREd::EditorUIResolutionY->GetInt());
 
-		//	// We're going to start stealing tabs from the global tab manager inserting them into the world instead.
-		//	FGlobalTabmanager::Get()->SetProxyTabManager(ProxyTabManager);
-		//}
+			const bool bWithSceneComponent = false;
+			AVREditorDockableWindow* TabManagerUI = GetOwner().SpawnTransientSceneActor< AVREditorDockableWindow >(TEXT("TabManager"), bWithSceneComponent);
+			TabManagerUI->SetSlateWidget(*this, SNullWidget::NullWidget, Resolution, VREd::EditorUISize->GetFloat(), AVREditorFloatingUI::EDockedTo::Nothing);
+			TabManagerUI->ShowUI(true);
+			TabManagerUI->SetRelativeOffset(FVector(0, 3000, 1000));
+			TabManagerUI->GetWidgetComponent()->SetDrawingPolicy(EVREditorWidgetDrawingPolicy::Hovering);
+			TabManagerUI->GetWidgetComponent()->SetBlendMode(EWidgetBlendMode::Masked);
+			FloatingUIs.Add(TabManagerUI);
+
+			TSharedPtr<SWindow> TabManagerWindow = TabManagerUI->GetWidgetComponent()->GetSlateWindow();
+			TSharedRef<SWindow> TabManagerWindowRef = TabManagerWindow.ToSharedRef();
+			ProxyTabManager = MakeShareable(new FProxyTabmanager(TabManagerWindowRef));
+
+			TSharedRef<FTabManager::FLayout> LoadedLayout = FTabManager::NewLayout(TEXT("VRTabManager"));
+			LoadedLayout->AddArea(FTabManager::NewPrimaryArea());
+
+			TSharedPtr<SWidget> DockLayout = FGlobalTabmanager::Get()->RestoreFrom(LoadedLayout, TabManagerWindowRef, false);
+			TabManagerWindowRef->SetContent(DockLayout.ToSharedRef());
+
+			// We're going to start stealing tabs from the global tab manager inserting them into the world instead.
+			FGlobalTabmanager::Get()->SetProxyTabManager(ProxyTabManager);
+		}
 	}
 }
 
diff --git a/Engine/Source/Editor/VREditor/VREditorWidgetComponent.cpp b/Engine/Source/Editor/VREditor/VREditorWidgetComponent.cpp
--- a/Engine/Source/Editor/VREditor/VREditorWidgetComponent.cpp
+++ b/Engine/Source/Editor/VREditor/VREditorWidgetComponent.cpp
@@ -6,13 +6,16 @@
 UVREditorWidgetComponent::UVREditorWidgetComponent(const FObjectInitializer& ObjectInitializer)
 	: Super( ObjectInitializer )
 {
-	bAllowDrawing = false;
+	bIsHovering = false;
+	DrawingPolicy = EVREditorWidgetDrawingPolicy::Always;
 	bHasEverDrawn = false;
 }
 
 bool UVREditorWidgetComponent::ShouldDrawWidget() const
 {
-	if ( bAllowDrawing || !bHasEverDrawn )
+	if ( DrawingPolicy == EVREditorWidgetDrawingPolicy::Always ||
+		(DrawingPolicy == EVREditorWidgetDrawingPolicy::Hovering && bIsHovering) ||
+		!bHasEverDrawn )
 	{
 		return Super::ShouldDrawWidget();
 	}
diff --git a/Engine/Source/Editor/VREditor/VREditorWidgetComponent.h b/Engine/Source/Editor/VREditor/VREditorWidgetComponent.h
--- a/Engine/Source/Editor/VREditor/VREditorWidgetComponent.h
+++ b/Engine/Source/Editor/VREditor/VREditorWidgetComponent.h
@@ -5,6 +5,13 @@
 #include "Components/WidgetComponent.h"
 #include "VREditorWidgetComponent.generated.h"
 
+UENUM()
+enum class EVREditorWidgetDrawingPolicy : uint8
+{
+	Always,
+	Hovering,
+};
+
 /**
  * A specialized WidgetComponent for the VREditor
  */
@@ -17,8 +24,11 @@ public:
 	/** Default constructor that sets up CDO properties */
 	UVREditorWidgetComponent(const FObjectInitializer& ObjectInitializer);
 
-	void SetAllowDrawing(bool Value) { bAllowDrawing = Value; }
-	bool GetAllowDrawing() const { return bAllowDrawing; }
+	void SetDrawingPolicy(EVREditorWidgetDrawingPolicy Value) { DrawingPolicy = Value; }
+	EVREditorWidgetDrawingPolicy GetDrawingPolicy() const { return DrawingPolicy; }
+
+	void SetIsHovering(bool Value) { bIsHovering = Value; }
+	bool GetIsHovering() const { return bIsHovering; }
 
 protected:
 	virtual bool ShouldDrawWidget() const override;
@@ -26,12 +36,18 @@ protected:
 	virtual void DrawWidgetToRenderTarget(float DeltaTime) override;
 
 private:
+	/**
+	 * High level redrawing policy for the widget component.
+	 */
+	UPROPERTY()
+	EVREditorWidgetDrawingPolicy DrawingPolicy;
+
 	/**
 	 * Controls if we draw, the VREditorWidget allows for manual enabling or 
 	 * disabling of updating the slate widget.
 	 */
 	UPROPERTY()
-	bool bAllowDrawing;
+	bool bIsHovering;
 
 	/**
 	 * Recorders if we've drawn at least once, that way we can always draw the first
diff --git a/Engine/Source/Runtime/Slate/Private/Framework/Docking/TabManager.cpp b/Engine/Source/Runtime/Slate/Private/Framework/Docking/TabManager.cpp
--- a/Engine/Source/Runtime/Slate/Private/Framework/Docking/TabManager.cpp
+++ b/Engine/Source/Runtime/Slate/Private/Framework/Docking/TabManager.cpp
@@ -2015,8 +2015,8 @@ void FProxyTabmanager::OpenUnmanagedTab(FName PlaceholderId, const FSearchPrefer
 	TSharedRef<FArea> NewAreaForTab = FTabManager::NewArea(FTabManager::FallbackWindowSize)
 		->Split
 		(
-		FTabManager::NewStack()
-		->AddTab(UnmanagedTab->GetLayoutIdentifier(), ETabState::OpenedTab)
+			FTabManager::NewStack()
+			->AddTab(UnmanagedTab->GetLayoutIdentifier(), ETabState::OpenedTab)
 		);
 
 	TSharedRef<SDockingArea> DockingArea = RestoreArea(NewAreaForTab, ParentWindow.Pin());
@@ -2025,6 +2025,7 @@ void FProxyTabmanager::OpenUnmanagedTab(FName PlaceholderId, const FSearchPrefer
 	check(NewlyOpenedTab.IsValid());
 		
 	NewlyOpenedTab->GetParent()->GetParentDockTabStack()->OpenTab(UnmanagedTab);
+	NewlyOpenedTab->RequestCloseTab();
 }
 
 
diff --git a/Engine/Source/Runtime/SlateCore/Private/Widgets/SWindow.cpp b/Engine/Source/Runtime/SlateCore/Private/Widgets/SWindow.cpp
--- a/Engine/Source/Runtime/SlateCore/Private/Widgets/SWindow.cpp
+++ b/Engine/Source/Runtime/SlateCore/Private/Widgets/SWindow.cpp
@@ -516,7 +516,7 @@ FSlateRect SWindow::GetNonMaximizedRectInScreen() const
 	int Width = 0;
 	int Height = 0;
 	
-	if ( NativeWindow->GetRestoredDimensions(X, Y, Width, Height) )
+	if ( NativeWindow.IsValid() && NativeWindow->GetRestoredDimensions(X, Y, Width, Height) )
 	{
 		return FSlateRect( X, Y, X+Width, Y+Height );
 	}
@@ -947,10 +947,11 @@ void SWindow::RequestDestroyWindow()
 /** Warning: use Request Destroy Window whenever possible!  This method destroys the window immediately! */
 void SWindow::DestroyWindowImmediately()
 {
-	check( NativeWindow.IsValid() );
-
-	// Destroy the native window
-	NativeWindow->Destroy();
+	if ( NativeWindow.IsValid() )
+	{
+		// Destroy the native window
+		NativeWindow->Destroy();
+	}
 }
 
 /** Calls the OnWindowClosed delegate when this window is about to be closed */
