VR Editor: NickD's new editor UI now behaves like a normal dockable UI
- It spawns to your right in your tracking space.  You can drag it around.
- You can show/hide it by clicking on the "Asset" button on the quick menu, or by closing it with the X

diff --git a/Engine/Source/Editor/VREditor/VREditorQuickMenu.cpp b/Engine/Source/Editor/VREditor/VREditorQuickMenu.cpp
--- a/Engine/Source/Editor/VREditor/VREditorQuickMenu.cpp
+++ b/Engine/Source/Editor/VREditor/VREditorQuickMenu.cpp
@@ -329,4 +329,17 @@ bool UVREditorQuickMenu::IsTutorialVisible() const
 	return ( GetOwner()->GetOwner().IsShowingEditorUIPanel( FVREditorUISystem::EEditorUIPanel::Tutorial ) );
 }
 
+
+void UVREditorQuickMenu::OnAssetEditorButtonClicked( const bool bIsChecked )
+{
+	bool bShouldShow = !IsAssetEditorVisible();
+	const int32 MyHandIndex = GetOwner()->GetDockedTo() == AVREditorFloatingUI::EDockedTo::LeftArm ? VREditorConstants::LeftHandIndex : VREditorConstants::RightHandIndex;
+	GetOwner()->GetOwner().ShowEditorUIPanel( FVREditorUISystem::EEditorUIPanel::AssetEditor, MyHandIndex, bShouldShow );
+}
+
+bool UVREditorQuickMenu::IsAssetEditorVisible() const
+{
+	return ( GetOwner()->GetOwner().IsShowingEditorUIPanel( FVREditorUISystem::EEditorUIPanel::AssetEditor ) );
+}
+
 #undef LOCTEXT_NAMESPACE
\ No newline at end of file
diff --git a/Engine/Source/Editor/VREditor/VREditorQuickMenu.h b/Engine/Source/Editor/VREditor/VREditorQuickMenu.h
--- a/Engine/Source/Editor/VREditor/VREditorQuickMenu.h
+++ b/Engine/Source/Editor/VREditor/VREditorQuickMenu.h
@@ -132,6 +132,12 @@ public:
 	UFUNCTION( BlueprintPure, Category = "Default" )
 	bool IsTutorialVisible() const;
 
+	UFUNCTION( BlueprintCallable, Category = "Default" )
+	void OnAssetEditorButtonClicked( const bool bIsChecked );
+
+	UFUNCTION( BlueprintPure, Category = "Default" )
+	bool IsAssetEditorVisible() const;
+
 	UFUNCTION(BlueprintImplementableEvent, meta=(DisplayName = "Refresh"))
 	void RefreshUI();
 
diff --git a/Engine/Source/Editor/VREditor/VREditorUISystem.cpp b/Engine/Source/Editor/VREditor/VREditorUISystem.cpp
--- a/Engine/Source/Editor/VREditor/VREditorUISystem.cpp
+++ b/Engine/Source/Editor/VREditor/VREditorUISystem.cpp
@@ -760,19 +760,20 @@ void FVREditorUISystem::CreateUIs()
 			EditorUIPanels[ (int32)EEditorUIPanel::Tutorial ] = TutorialUI;		
 		}
 
-		if ( 0 )
 		{
 			const FIntPoint Resolution(VREd::EditorUIResolutionX->GetInt(), VREd::EditorUIResolutionY->GetInt());
 
 			const bool bWithSceneComponent = false;
-			AVREditorDockableWindow* TabManagerUI = GetOwner().SpawnTransientSceneActor< AVREditorDockableWindow >(TEXT("TabManager"), bWithSceneComponent);
-			TabManagerUI->SetSlateWidget(*this, SNullWidget::NullWidget, Resolution, VREd::EditorUISize->GetFloat(), AVREditorFloatingUI::EDockedTo::Nothing);
+			AVREditorDockableWindow* TabManagerUI = GetOwner().SpawnTransientSceneActor< AVREditorDockableWindow >(TEXT("AssetEditor"), bWithSceneComponent);
+			TabManagerUI->SetSlateWidget(*this, SNullWidget::NullWidget, Resolution, VREd::EditorUISize->GetFloat() * 2.0f, AVREditorFloatingUI::EDockedTo::Room);
 			TabManagerUI->ShowUI(true);
-			TabManagerUI->SetRelativeOffset(FVector(0, 3000, 1000));
 			TabManagerUI->GetWidgetComponent()->SetDrawingPolicy(EVREditorWidgetDrawingPolicy::Hovering);
-			TabManagerUI->GetWidgetComponent()->SetBlendMode(EWidgetBlendMode::Masked);
+			TabManagerUI->SetRelativeOffset( FVector( 200.0f, 0.0f, 150.0f ) );
+			TabManagerUI->SetLocalRotation( FRotator( 0.0f, 180.0f, 0.0f ) );
 			FloatingUIs.Add(TabManagerUI);
 
+			EditorUIPanels[ (int32)EEditorUIPanel::AssetEditor ] = TabManagerUI;
+
 			TSharedPtr<SWindow> TabManagerWindow = TabManagerUI->GetWidgetComponent()->GetSlateWindow();
 			TSharedRef<SWindow> TabManagerWindowRef = TabManagerWindow.ToSharedRef();
 			ProxyTabManager = MakeShareable(new FProxyTabmanager(TabManagerWindowRef));
diff --git a/Engine/Source/Editor/VREditor/VREditorUISystem.h b/Engine/Source/Editor/VREditor/VREditorUISystem.h
--- a/Engine/Source/Editor/VREditor/VREditorUISystem.h
+++ b/Engine/Source/Editor/VREditor/VREditorUISystem.h
@@ -16,6 +16,7 @@ public:
 		ActorDetails,
 		Modes,
 		Tutorial,
+		AssetEditor,
 
 		// ...
 
