VR Editor: Added console command to invert trackpad vertical axis (randomly backwards on some devices?)

diff --git a/Engine/Source/Editor/VREditor/VREditorMode.cpp b/Engine/Source/Editor/VREditor/VREditorMode.cpp
--- a/Engine/Source/Editor/VREditor/VREditorMode.cpp
+++ b/Engine/Source/Editor/VREditor/VREditorMode.cpp
@@ -56,6 +56,7 @@ namespace VREd
 	static FAutoConsoleVariable MinHapticTimeForRift( TEXT( "VREd.MinHapticTimeForRift" ), 0.005f, TEXT( "How long to play haptic effects on the Rift" ) );
 	static FAutoConsoleVariable SleepForRiftHaptics( TEXT( "VREd.SleepForRiftHaptics" ), 1, TEXT( "When enabled, we'll sleep the game thread mid-frame to wait for haptic effects to finish.  This can be devasting to performance!" ) );
 
+	static FAutoConsoleVariable InvertTrackpadVertical( TEXT( "VREd.InvertTrackpadVertical" ), 1, TEXT( "Toggles inverting the touch pad vertical axis" ) );
 }
 
 FEditorModeID FVREditorMode::VREditorModeID( "VREditor" );
@@ -1513,7 +1514,10 @@ bool FVREditorMode::InputAxis( FEditorViewportClient* ViewportClient, FViewport*
 
 			if( VRAction.ActionType == EVRActionType::TrackpadPositionY )
 			{
-				Delta = -Delta;	// Y axis is inverted from HMD
+				if( VREd::InvertTrackpadVertical->GetInt() != 0 )
+				{
+					Delta = -Delta;	// Y axis is inverted from HMD
+				}
 
 				Hand.LastTrackpadPosition.Y = Hand.bIsTrackpadPositionValid[1] ? Hand.TrackpadPosition.Y : Delta;
 				Hand.LastTrackpadPositionUpdateTime = FTimespan::FromSeconds( FPlatformTime::Seconds() );
