VR Editor: Inverted mouse wheel scroll direction (so that it feels more like scroll on a mobile device)

diff --git a/Engine/Source/Editor/VREditor/VREditorUISystem.cpp b/Engine/Source/Editor/VREditor/VREditorUISystem.cpp
--- a/Engine/Source/Editor/VREditor/VREditorUISystem.cpp
+++ b/Engine/Source/Editor/VREditor/VREditorUISystem.cpp
@@ -383,6 +383,9 @@ void FVREditorUISystem::OnVRHoverUpdate( FEditorViewportClient& ViewportClient,
 								}
 							}
 
+							// Invert scroll direction so that it feels more like scrolling on a mobile device
+							ScrollDelta *= -1.0f;
+
 							if( !FMath::IsNearlyZero( ScrollDelta ) )
 							{
 								FPointerEvent MouseWheelEvent(
