VR Editor: Tweaked color and highlight of UI selection bar

diff --git a/Engine/Source/Editor/VREditor/VREditorDockableWindow.cpp b/Engine/Source/Editor/VREditor/VREditorDockableWindow.cpp
--- a/Engine/Source/Editor/VREditor/VREditorDockableWindow.cpp
+++ b/Engine/Source/Editor/VREditor/VREditorDockableWindow.cpp
@@ -241,12 +241,12 @@ void AVREditorDockableWindow::OnEnterHover( const FHitResult& HitResult )
 {
 	if( SelectionBarMeshComponent == HitResult.GetComponent() )
 	{
-		SetSelectionBarColor( GetOwner().GetOwner().GetColor( FVREditorMode::EColors::HoverGizmoColor ) );
+		SetSelectionBarColor( GetOwner().GetOwner().GetColor( FVREditorMode::EColors::UISelectionBarHoverColor ) );
 	}
 
 	if( CloseButtonMeshComponent == HitResult.GetComponent() )
 	{
-		SetCloseButtonColor( GetOwner().GetOwner().GetColor( FVREditorMode::EColors::HoverGizmoColor ) );
+		SetCloseButtonColor( GetOwner().GetOwner().GetColor( FVREditorMode::EColors::UICloseButtonHoverColor ) );
 	}
 }
 
diff --git a/Engine/Source/Editor/VREditor/VREditorMode.cpp b/Engine/Source/Editor/VREditor/VREditorMode.cpp
--- a/Engine/Source/Editor/VREditor/VREditorMode.cpp
+++ b/Engine/Source/Editor/VREditor/VREditorMode.cpp
@@ -212,8 +212,10 @@ void FVREditorMode::Enter()
 		Colors[ (int32)EColors::WhiteGizmoColor ] = FLinearColor( 0.7f, 0.7f, 0.7f, 1.0f );
 		Colors[ (int32)EColors::HoverGizmoColor ] = FLinearColor( FLinearColor::Yellow );
 		Colors[ (int32)EColors::DraggingGizmoColor ] = FLinearColor( FLinearColor::White );
-		Colors[ (int32)EColors::UISelectionBarColor ] = FLinearColor( 0.0f, 0.7f, 0.7f, 1.0f );
+		Colors[ (int32)EColors::UISelectionBarColor ] = FLinearColor( 0.025f, 0.025f, 0.025f, 1.0f );
+		Colors[ (int32)EColors::UISelectionBarHoverColor ] = FLinearColor( 0.1f, 0.1f, 0.1f, 1.0f );
 		Colors[ (int32)EColors::UICloseButtonColor ] = FLinearColor( 0.1f, 0.1f, 0.1f, 1.0f );
+		Colors[ (int32)EColors::UICloseButtonHoverColor ] = FLinearColor( 1.0f, 1.0f, 1.0f, 1.0f );
 	}
 
 	// @todo vreditor: We need to make sure the user can never switch to orthographic mode, or activate settings that
diff --git a/Engine/Source/Editor/VREditor/VREditorMode.h b/Engine/Source/Editor/VREditor/VREditorMode.h
--- a/Engine/Source/Editor/VREditor/VREditorMode.h
+++ b/Engine/Source/Editor/VREditor/VREditorMode.h
@@ -525,7 +525,9 @@ public:
 		HoverGizmoColor,
 		DraggingGizmoColor,
 		UISelectionBarColor,
+		UISelectionBarHoverColor,
 		UICloseButtonColor,
+		UICloseButtonHoverColor,
 		TotalCount
 	};
 
