VR Editor: Added "right click" support (hold Modifier and pull trigger)
- Isn't very useful yet.  Just adding it for testing.

diff --git a/Engine/Source/Editor/VREditor/VREditorUISystem.cpp b/Engine/Source/Editor/VREditor/VREditorUISystem.cpp
--- a/Engine/Source/Editor/VREditor/VREditorUISystem.cpp
+++ b/Engine/Source/Editor/VREditor/VREditorUISystem.cpp
@@ -178,6 +178,11 @@ void FVREditorUISystem::OnVRAction( FEditorViewportClient& ViewportClient, const
 
 						if( Event != IE_Repeat )
 						{
+							// If the Modifier button is held down, treat this like a right click instead of a left click
+							const bool bIsRightClicking =
+								( Event == IE_Pressed && Hand.bIsModifierPressed ) ||
+								( Event == IE_Released && Hand.bIsRightClickingOnUI );
+
 							UWidgetComponent* WidgetComponent = Cast<UWidgetComponent>( HitResult.GetComponent() );
 							if( WidgetComponent )
 							{
@@ -203,14 +208,14 @@ void FVREditorUISystem::OnVRAction( FEditorViewportClient& ViewportClient, const
 									TSet<FKey> PressedButtons;
 									if( Event == IE_Pressed )
 									{
-										PressedButtons.Add( EKeys::LeftMouseButton );
+										PressedButtons.Add( bIsRightClicking ? EKeys::RightMouseButton : EKeys::LeftMouseButton );
 									}
 									FPointerEvent PointerEvent(
 										1 + VRAction.HandIndex,
 										LocalHitLocation,
 										LastLocalHitLocation,
 										PressedButtons,
-										EKeys::LeftMouseButton,
+										bIsRightClicking ? EKeys::RightMouseButton : EKeys::LeftMouseButton,
 										0.0f,	// Wheel delta
 										FModifierKeysState() );
 
@@ -226,6 +231,7 @@ void FVREditorUISystem::OnVRAction( FEditorViewportClient& ViewportClient, const
 									{
 										Reply = FSlateApplication::Get().RoutePointerDownEvent( WidgetPathUnderFinger, PointerEvent );
 										Hand.bIsClickingOnUI = true;
+										Hand.bIsRightClickingOnUI = bIsRightClicking;
 										bIsInputCaptured = true;
 									}
 									else if( Event == IE_Released )
@@ -250,9 +256,15 @@ void FVREditorUISystem::OnVRAction( FEditorViewportClient& ViewportClient, const
 
 			if( Event == IE_Released )
 			{
+				bool bWasRightClicking = false;
 				if( Hand.bIsClickingOnUI )
 				{
+					if( Hand.bIsRightClickingOnUI )
+					{
+						bWasRightClicking = true;
+					}
 					Hand.bIsClickingOnUI = false;
+					Hand.bIsRightClickingOnUI = false;
 					bIsInputCaptured = false;
 				}
 
@@ -264,13 +276,14 @@ void FVREditorUISystem::OnVRAction( FEditorViewportClient& ViewportClient, const
 						FVector2D::ZeroVector,
 						FVector2D::ZeroVector,
 						PressedButtons,
-						EKeys::LeftMouseButton,
+						bWasRightClicking ? EKeys::RightMouseButton : EKeys::LeftMouseButton,
 						0.0f,	// Wheel delta
 						FModifierKeysState() );
 
 					FWidgetPath EmptyWidgetPath;
 
 					Hand.bIsClickingOnUI = false;
+					Hand.bIsRightClickingOnUI = false;
 					FReply Reply = FSlateApplication::Get().RoutePointerUpEvent( EmptyWidgetPath, PointerEvent );
 				}
 			}
diff --git a/Engine/Source/Editor/VREditor/VREditorWorldInteraction.cpp b/Engine/Source/Editor/VREditor/VREditorWorldInteraction.cpp
--- a/Engine/Source/Editor/VREditor/VREditorWorldInteraction.cpp
+++ b/Engine/Source/Editor/VREditor/VREditorWorldInteraction.cpp
@@ -2700,7 +2700,7 @@ void FVREditorWorldInteraction::OnAssetDragStartedFromContentBrowser( const TArr
 	for( int32 HandIndex = 0; HandIndex < VREditorConstants::NumVirtualHands; ++HandIndex )
 	{
 		FVirtualHand& Hand = Owner.GetVirtualHand( HandIndex );
-		if( Hand.IsInputCaptured[ (int32)EVRActionType::SelectAndMove ] && Hand.bIsClickingOnUI )
+		if( Hand.IsInputCaptured[ (int32)EVRActionType::SelectAndMove ] && ( Hand.bIsClickingOnUI && !Hand.bIsRightClickingOnUI ) )
 		{
 			PlacingWithHandIndex = HandIndex;
 			OtherHandIndex = Owner.GetOtherHandIndex( PlacingWithHandIndex );
@@ -2716,6 +2716,7 @@ void FVREditorWorldInteraction::OnAssetDragStartedFromContentBrowser( const TArr
 		// Cancel UI input
 		{
 			Hand.bIsClickingOnUI = false;
+			Hand.bIsRightClickingOnUI = false;
 			Hand.IsInputCaptured[ (int32)EVRActionType::SelectAndMove ] = false;
 			Hand.IsInputCaptured[ (int32)EVRActionType::SelectAndMove_LightlyPressed ] = false;
 		}
diff --git a/Engine/Source/Editor/VREditor/VirtualHand.h b/Engine/Source/Editor/VREditor/VirtualHand.h
--- a/Engine/Source/Editor/VREditor/VirtualHand.h
+++ b/Engine/Source/Editor/VREditor/VirtualHand.h
@@ -168,6 +168,9 @@ struct FVirtualHand
 	/** True if we're currently holding the 'SelectAndMove' button down after clicking on UI */
 	bool bIsClickingOnUI;
 
+	/** When bIsClickingOnUI is true, this will be true if we're "right clicking".  That is, the Modifier key was held down at the time that the user clicked */
+	bool bIsRightClickingOnUI;
+
 	/** Last real time that we released the 'SelectAndMove' button on UI.  This is used to detect double-clicks. */
 	double LastClickReleaseTime;
 
@@ -320,6 +323,7 @@ struct FVirtualHand
 		}
 		ClickingOnComponent = nullptr;
 		bIsClickingOnUI = false;
+		bIsRightClickingOnUI = false;
 		LastClickReleaseTime = 0.0;
 		LastHapticTime = 0.0;
 		
