VR Editor: Don't show 'Editor Tab' UI by default

diff --git a/Engine/Source/Editor/VREditor/VREditorUISystem.cpp b/Engine/Source/Editor/VREditor/VREditorUISystem.cpp
--- a/Engine/Source/Editor/VREditor/VREditorUISystem.cpp
+++ b/Engine/Source/Editor/VREditor/VREditorUISystem.cpp
@@ -781,8 +781,8 @@ void FVREditorUISystem::CreateUIs()
 
 			const bool bWithSceneComponent = false;
 			AVREditorDockableWindow* TabManagerUI = GetOwner().SpawnTransientSceneActor< AVREditorDockableWindow >(TEXT("AssetEditor"), bWithSceneComponent);
-			TabManagerUI->SetSlateWidget(*this, SNullWidget::NullWidget, Resolution, VREd::EditorUISize->GetFloat() * 2.0f, AVREditorFloatingUI::EDockedTo::Room);
-			TabManagerUI->ShowUI(true);
+			TabManagerUI->SetSlateWidget(*this, SNullWidget::NullWidget, Resolution, VREd::EditorUISize->GetFloat() * 2.0f, AVREditorFloatingUI::EDockedTo::Nothing);
+			TabManagerUI->ShowUI( false );
 			TabManagerUI->GetWidgetComponent()->SetDrawingPolicy(EVREditorWidgetDrawingPolicy::Hovering);
 			TabManagerUI->SetRelativeOffset( FVector( 200.0f, 0.0f, 150.0f ) );
 			TabManagerUI->SetLocalRotation( FRotator( 0.0f, 180.0f, 0.0f ) );
