VR Editor: Added hover animation for UI controls that appear when you aim at the UI

diff --git a/Engine/Source/Editor/VREditor/VREditorDockableWindow.cpp b/Engine/Source/Editor/VREditor/VREditorDockableWindow.cpp
--- a/Engine/Source/Editor/VREditor/VREditorDockableWindow.cpp
+++ b/Engine/Source/Editor/VREditor/VREditorDockableWindow.cpp
@@ -13,6 +13,9 @@ namespace VREd
 	static FAutoConsoleVariable DockWindowThickness( TEXT( "VREd.DockWindowTickness" ), 1.0f, TEXT( "How thick the window is" ) );
 	static FAutoConsoleVariable DockUISelectionBarVerticalOffset( TEXT( "VREd.DockUISelectionBarVerticalOffset" ), 2.0f, TEXT( "Z Distance between the selectionbar and the UI" ) );
 	static FAutoConsoleVariable DockUICloseButtonOffsetFromCorner( TEXT( "VREd.DockUICloseButtonOffsetFromCorner" ), 1.5f, TEXT( "How far away from the corner (along each 2D axis) of the UI the close button draws" ) );
+	static FAutoConsoleVariable DockUIFadeAnimationDuration( TEXT( "VREd.DockUIFadeAnimationDuration" ), 0.5f, TEXT( "How quick the fade animation should complete in" ) );
+	static FAutoConsoleVariable DockUIHoverScale( TEXT( "VREd.DockUIHoverScale" ), 1.15f, TEXT( "How big the selection bar gets when you hover over it" ) );
+	static FAutoConsoleVariable DockUIHoverAnimationDuration( TEXT( "VREd.DockUIHoverAnimationDuration" ), 0.15f, TEXT( "How quick the hover animation should complete in" ) );
 }
 
 AVREditorDockableWindow::AVREditorDockableWindow() 
@@ -156,8 +159,9 @@ void AVREditorDockableWindow::TickManually( float DeltaTime )
 		const float WorldScaleFactor = GetOwner().GetOwner().GetWorldScaleFactor();
 		const FVector AnimatedScale = CalculateAnimatedScale();
 
-		
-		bool bIsLaserAimingTowardUI = false;
+
+		// Update whether the user is aiming toward us or not
+		bIsLaserAimingTowardUI = false;
 		{
 			const FTransform UICapsuleTransform = this->GetActorTransform();
 
@@ -175,11 +179,22 @@ void AVREditorDockableWindow::TickManually( float DeltaTime )
 					bIsLaserAimingTowardUI = true;
 				}
 			}
+
+			if( bIsLaserAimingTowardUI )
+			{
+				AimingAtMeFadeAlpha += DeltaTime / VREd::DockUIFadeAnimationDuration->GetFloat();
+			}
+			else
+			{
+				AimingAtMeFadeAlpha -= DeltaTime / VREd::DockUIFadeAnimationDuration->GetFloat();;
+			}
+			AimingAtMeFadeAlpha = FMath::Clamp( AimingAtMeFadeAlpha, 0.0f, 1.0f );
 		}
 
 		// Only show our extra buttons and controls if the user is roughly aiming toward us.  This just reduces clutter.
-		SelectionBarMeshComponent->SetVisibility( bIsLaserAimingTowardUI );
-		CloseButtonMeshComponent->SetVisibility( bIsLaserAimingTowardUI );
+		SelectionBarMeshComponent->SetVisibility( AimingAtMeFadeAlpha > KINDA_SMALL_NUMBER ? true : false );
+		CloseButtonMeshComponent->SetVisibility( AimingAtMeFadeAlpha > KINDA_SMALL_NUMBER ? true : false );
+
 
 		// Update the window border mesh
 		{
@@ -194,30 +209,60 @@ void AVREditorDockableWindow::TickManually( float DeltaTime )
 
 		// Update the selection bar
 		{
+			if( bIsHoveringOverSelectionBar )
+			{
+				SelectionBarHoverAlpha += DeltaTime / VREd::DockUIHoverAnimationDuration->GetFloat();
+			}
+			else
+			{
+				SelectionBarHoverAlpha -= DeltaTime / VREd::DockUIHoverAnimationDuration->GetFloat();;
+			}
+			SelectionBarHoverAlpha = FMath::Clamp( SelectionBarHoverAlpha, 0.0f, 1.0f );
+
 			// How big the selection bar should be
 			const FVector SelectionBarSize( 1.0f, GetSize().X * 0.8f, 6.0f );	// @todo vreditor tweak
+			FVector SelectionBarScale = SelectionBarSize * AnimatedScale * WorldScaleFactor;
+			SelectionBarScale *= FMath::Lerp( 1.0f, VREd::DockUIHoverScale->GetFloat(), SelectionBarHoverAlpha );
+
+			// Scale vertically based on our fade alpha
+			SelectionBarScale.Z *= AimingAtMeFadeAlpha;
 
-			const FVector SelectionBarScale = SelectionBarSize * AnimatedScale * WorldScaleFactor;
 			SelectionBarMeshComponent->SetRelativeScale3D( SelectionBarScale );
 			const FVector SelectionBarRelativeLocation = FVector(
 				0.0f,
 				0.0f,
 				-( GetSize().Y * 0.5f + SelectionBarSize.Z * 0.5f + VREd::DockUISelectionBarVerticalOffset->GetFloat() ) ) * AnimatedScale * WorldScaleFactor;
 			SelectionBarMeshComponent->SetRelativeLocation( SelectionBarRelativeLocation );
+
+			SetSelectionBarColor( GetOwner().GetOwner().GetColor( bIsHoveringOverSelectionBar ? FVREditorMode::EColors::UISelectionBarHoverColor : FVREditorMode::EColors::UISelectionBarColor ) );
 		}
 
 		// Update the close button
 		{
+			if( bIsHoveringOverCloseButton )
+			{
+				CloseButtonHoverAlpha += DeltaTime / VREd::DockUIHoverAnimationDuration->GetFloat();
+			}
+			else
+			{
+				CloseButtonHoverAlpha -= DeltaTime / VREd::DockUIHoverAnimationDuration->GetFloat();;
+			}
+			CloseButtonHoverAlpha = FMath::Clamp( CloseButtonHoverAlpha, 0.0f, 1.0f );
+
 			// How big the close button should be
 			const FVector CloseButtonSize( 1.0f, 2.0f, 2.0f );	// @todo vreditor tweak
+			FVector CloseButtonScale = CloseButtonSize * AnimatedScale * WorldScaleFactor * AimingAtMeFadeAlpha;
+			CloseButtonScale *= FMath::Lerp( 1.0f, VREd::DockUIHoverScale->GetFloat(), CloseButtonHoverAlpha );
+
 
-			const FVector CloseButtonScale = CloseButtonSize * AnimatedScale * WorldScaleFactor;
 			CloseButtonMeshComponent->SetRelativeScale3D( CloseButtonScale );
 			const FVector CloseButtonRelativeLocation = FVector(
 				0.0f,
 				-( GetSize().X * 0.5f + CloseButtonSize.Y * 0.5f + VREd::DockUICloseButtonOffsetFromCorner->GetFloat() ),
 				( GetSize().Y * 0.5f + CloseButtonSize.Z * 0.5f + VREd::DockUICloseButtonOffsetFromCorner->GetFloat() ) ) * AnimatedScale * WorldScaleFactor;
 			CloseButtonMeshComponent->SetRelativeLocation( CloseButtonRelativeLocation );
+
+			SetCloseButtonColor( GetOwner().GetOwner().GetColor( bIsHoveringOverCloseButton ? FVREditorMode::EColors::UICloseButtonHoverColor : FVREditorMode::EColors::UICloseButtonColor ) );
 		}
 	}
 } 
@@ -241,19 +286,25 @@ void AVREditorDockableWindow::OnEnterHover( const FHitResult& HitResult )
 {
 	if( SelectionBarMeshComponent == HitResult.GetComponent() )
 	{
-		SetSelectionBarColor( GetOwner().GetOwner().GetColor( FVREditorMode::EColors::UISelectionBarHoverColor ) );
+		bIsHoveringOverSelectionBar = true;
+		bIsHoveringOverCloseButton = false;
 	}
-
-	if( CloseButtonMeshComponent == HitResult.GetComponent() )
+	else if( CloseButtonMeshComponent == HitResult.GetComponent() )
+	{
+		bIsHoveringOverCloseButton = true;
+		bIsHoveringOverSelectionBar = false;
+	}
+	else
 	{
-		SetCloseButtonColor( GetOwner().GetOwner().GetColor( FVREditorMode::EColors::UICloseButtonHoverColor ) );
+		bIsHoveringOverCloseButton = false;
+		bIsHoveringOverSelectionBar = false;
 	}
 }
 
 void AVREditorDockableWindow::OnLeaveHover()
 {
-	SetSelectionBarColor( GetOwner().GetOwner().GetColor( FVREditorMode::EColors::UISelectionBarColor ) );
-	SetCloseButtonColor( GetOwner().GetOwner().GetColor( FVREditorMode::EColors::UICloseButtonColor ) );
+	bIsHoveringOverSelectionBar = false;
+	bIsHoveringOverCloseButton = false;
 }
 
 
diff --git a/Engine/Source/Editor/VREditor/VREditorDockableWindow.h b/Engine/Source/Editor/VREditor/VREditorDockableWindow.h
--- a/Engine/Source/Editor/VREditor/VREditorDockableWindow.h
+++ b/Engine/Source/Editor/VREditor/VREditorDockableWindow.h
@@ -74,4 +74,25 @@ private:
 	/** Close button dynamic material (translucent) */
 	UPROPERTY()
 	class UMaterialInstanceDynamic* CloseButtonTranslucentMID;
+
+
+	/** True if at least one hand's laser is aiming toward the UI */
+	bool bIsLaserAimingTowardUI;
+
+	/** Scalar that ramps up toward 1.0 after the user aims toward the UI */
+	float AimingAtMeFadeAlpha;
+
+
+	/** True if we're hovering over the selection bar */
+	bool bIsHoveringOverSelectionBar;
+
+	/** Scalar that will advance toward 1.0 over time as we hover over the selection bar */
+	float SelectionBarHoverAlpha;
+
+
+	/** True if we're hovering over the close button */
+	bool bIsHoveringOverCloseButton;
+
+	/** Scalar that will advance toward 1.0 over time as we hover over the close button */
+	float CloseButtonHoverAlpha;
 };
\ No newline at end of file
