VR Editor: Fixed slow UI controls animation speed

diff --git a/Engine/Source/Editor/VREditor/VREditorDockableWindow.cpp b/Engine/Source/Editor/VREditor/VREditorDockableWindow.cpp
--- a/Engine/Source/Editor/VREditor/VREditorDockableWindow.cpp
+++ b/Engine/Source/Editor/VREditor/VREditorDockableWindow.cpp
@@ -13,7 +13,7 @@ namespace VREd
 	static FAutoConsoleVariable DockWindowThickness( TEXT( "VREd.DockWindowTickness" ), 1.0f, TEXT( "How thick the window is" ) );
 	static FAutoConsoleVariable DockUISelectionBarVerticalOffset( TEXT( "VREd.DockUISelectionBarVerticalOffset" ), 2.0f, TEXT( "Z Distance between the selectionbar and the UI" ) );
 	static FAutoConsoleVariable DockUICloseButtonOffsetFromCorner( TEXT( "VREd.DockUICloseButtonOffsetFromCorner" ), 1.5f, TEXT( "How far away from the corner (along each 2D axis) of the UI the close button draws" ) );
-	static FAutoConsoleVariable DockUIFadeAnimationDuration( TEXT( "VREd.DockUIFadeAnimationDuration" ), 0.5f, TEXT( "How quick the fade animation should complete in" ) );
+	static FAutoConsoleVariable DockUIFadeAnimationDuration( TEXT( "VREd.DockUIFadeAnimationDuration" ), 0.15f, TEXT( "How quick the fade animation should complete in" ) );
 	static FAutoConsoleVariable DockUIHoverScale( TEXT( "VREd.DockUIHoverScale" ), 1.15f, TEXT( "How big the selection bar gets when you hover over it" ) );
 	static FAutoConsoleVariable DockUIHoverAnimationDuration( TEXT( "VREd.DockUIHoverAnimationDuration" ), 0.15f, TEXT( "How quick the hover animation should complete in" ) );
 }
