VR Editor: Don't show radial menu when using analog stick (or trackpad) to resize UIs while dragging

diff --git a/Engine/Source/Editor/VREditor/VREditorUISystem.cpp b/Engine/Source/Editor/VREditor/VREditorUISystem.cpp
--- a/Engine/Source/Editor/VREditor/VREditorUISystem.cpp
+++ b/Engine/Source/Editor/VREditor/VREditorUISystem.cpp
@@ -559,19 +559,24 @@ void FVREditorUISystem::Tick( FEditorViewportClient* ViewportClient, const float
 		{
 			const FVirtualHand& Hand = Owner.GetVirtualHand( HandIndex );
 
-			if( Hand.bIsTrackpadPositionValid[ 0 ] && Hand.bIsTrackpadPositionValid[ 1 ] )
+			bool bShouldShowRadialMenu = false;
+			if( !Hand.bHasUIInFront && Hand.bIsTrackpadPositionValid[ 0 ] && Hand.bIsTrackpadPositionValid[ 1 ] && GetDraggingDockUIHandIndex() != HandIndex )
 			{
-				if( Hand.TrackpadPosition.Size() > 0.1f )	// @todo vreditor tweak
+				if( Hand.TrackpadPosition.Size() > 0.15f )	// @todo vreditor tweak
 				{
-					TryToSpawnRadialMenu( HandIndex );
-					
-				}
-				else
-				{
-					// Close it
-					HideRadialMenu( HandIndex );
+					bShouldShowRadialMenu = true;
 				}
 			}
+
+			if( bShouldShowRadialMenu )
+			{
+				TryToSpawnRadialMenu( HandIndex );
+			}
+			else
+			{
+				// Close it
+				HideRadialMenu( HandIndex );
+			}
 		}
 	}
 	else
