VR Editor: Fixed resizing UI using trackpad being inverted

diff --git a/Engine/Source/Editor/VREditor/VREditorWorldInteraction.cpp b/Engine/Source/Editor/VREditor/VREditorWorldInteraction.cpp
--- a/Engine/Source/Editor/VREditor/VREditorWorldInteraction.cpp
+++ b/Engine/Source/Editor/VREditor/VREditorWorldInteraction.cpp
@@ -1208,11 +1208,11 @@ void FVREditorWorldInteraction::Tick( FEditorViewportClient* ViewportClient, con
 					{
 						if (bIsAbsolute)
 						{
-							SlideDeltaY = -((Hand.TrackpadPosition.Y - Hand.LastTrackpadPosition.Y) * VREd::TrackpadAbsoluteDragSpeed->GetFloat()) * WorldScaleFactor;
+							SlideDeltaY = ((Hand.TrackpadPosition.Y - Hand.LastTrackpadPosition.Y) * VREd::TrackpadAbsoluteDragSpeed->GetFloat()) * WorldScaleFactor;
 						}
 						else
 						{
-							SlideDeltaY = -(Hand.TrackpadPosition.Y * VREd::TrackpadRelativeDragSpeed->GetFloat()) * WorldScaleFactor;
+							SlideDeltaY = (Hand.TrackpadPosition.Y * VREd::TrackpadRelativeDragSpeed->GetFloat()) * WorldScaleFactor;
 						}
 					}
 
