VR Editor: Added 'overshoot' animation effect when floating UI controls appear

diff --git a/Engine/Source/Editor/VREditor/VREditorDockableWindow.cpp b/Engine/Source/Editor/VREditor/VREditorDockableWindow.cpp
--- a/Engine/Source/Editor/VREditor/VREditorDockableWindow.cpp
+++ b/Engine/Source/Editor/VREditor/VREditorDockableWindow.cpp
@@ -191,9 +191,12 @@ void AVREditorDockableWindow::TickManually( float DeltaTime )
 			AimingAtMeFadeAlpha = FMath::Clamp( AimingAtMeFadeAlpha, 0.0f, 1.0f );
 		}
 
+		const float AnimationOvershootAmount = 1.0f;	// @todo vreditor tweak
+		const float EasedAimingAtMeFadeAlpha = FVREditorMode::OvershootEaseOut( AimingAtMeFadeAlpha, AnimationOvershootAmount );
+
 		// Only show our extra buttons and controls if the user is roughly aiming toward us.  This just reduces clutter.
-		SelectionBarMeshComponent->SetVisibility( AimingAtMeFadeAlpha > KINDA_SMALL_NUMBER ? true : false );
-		CloseButtonMeshComponent->SetVisibility( AimingAtMeFadeAlpha > KINDA_SMALL_NUMBER ? true : false );
+		SelectionBarMeshComponent->SetVisibility( EasedAimingAtMeFadeAlpha > KINDA_SMALL_NUMBER ? true : false );
+		CloseButtonMeshComponent->SetVisibility( EasedAimingAtMeFadeAlpha > KINDA_SMALL_NUMBER ? true : false );
 
 
 		// Update the window border mesh
@@ -225,7 +228,7 @@ void AVREditorDockableWindow::TickManually( float DeltaTime )
 			SelectionBarScale *= FMath::Lerp( 1.0f, VREd::DockUIHoverScale->GetFloat(), SelectionBarHoverAlpha );
 
 			// Scale vertically based on our fade alpha
-			SelectionBarScale.Z *= AimingAtMeFadeAlpha;
+			SelectionBarScale.Z *= EasedAimingAtMeFadeAlpha;
 
 			SelectionBarMeshComponent->SetRelativeScale3D( SelectionBarScale );
 			const FVector SelectionBarRelativeLocation = FVector(
@@ -251,7 +254,7 @@ void AVREditorDockableWindow::TickManually( float DeltaTime )
 
 			// How big the close button should be
 			const FVector CloseButtonSize( 1.0f, 2.0f, 2.0f );	// @todo vreditor tweak
-			FVector CloseButtonScale = CloseButtonSize * AnimatedScale * WorldScaleFactor * AimingAtMeFadeAlpha;
+			FVector CloseButtonScale = CloseButtonSize * AnimatedScale * WorldScaleFactor * EasedAimingAtMeFadeAlpha;
 			CloseButtonScale *= FMath::Lerp( 1.0f, VREd::DockUIHoverScale->GetFloat(), CloseButtonHoverAlpha );
 
 
diff --git a/Engine/Source/Editor/VREditor/VREditorFloatingUI.cpp b/Engine/Source/Editor/VREditor/VREditorFloatingUI.cpp
--- a/Engine/Source/Editor/VREditor/VREditorFloatingUI.cpp
+++ b/Engine/Source/Editor/VREditor/VREditorFloatingUI.cpp
@@ -276,20 +276,11 @@ void AVREditorFloatingUI::UpdateTransformIfDocked()
 	}
 }
 
-// @todo vreditor: Move this to FMath so we can use it everywhere
-// NOTE: OvershootAmount is usually between 0.5 and 2.0, but can go lower and higher for extreme overshots
-template<class T>
-static T OvershootEaseOut( T Alpha, const float OvershootAmount = 1.0f )
-{
-	Alpha--;
-	return 1.0f - ( ( Alpha * ( ( OvershootAmount + 1 ) * Alpha + OvershootAmount ) + 1 ) - 1.0f );
-}
-
 
 FVector AVREditorFloatingUI::CalculateAnimatedScale() const
 {
-	const float UIAnimationOvershootAmount = 0.7f;	// @todo vreditor tweak
-	const float EasedAlpha = OvershootEaseOut( FadeAlpha, UIAnimationOvershootAmount );
+	const float AnimationOvershootAmount = 0.7f;	// @todo vreditor tweak
+	const float EasedAlpha = FVREditorMode::OvershootEaseOut( FadeAlpha, AnimationOvershootAmount );
 
 	// Animate vertically more than horizontally; just looks a little better
 	const float ZScale = FMath::Max( 0.001f, EasedAlpha );
diff --git a/Engine/Source/Editor/VREditor/VREditorMode.h b/Engine/Source/Editor/VREditor/VREditorMode.h
--- a/Engine/Source/Editor/VREditor/VREditorMode.h
+++ b/Engine/Source/Editor/VREditor/VREditorMode.h
@@ -99,6 +99,15 @@ public:
 		return FTimespan::FromSeconds( FApp::GetCurrentTime() ) - AppTimeModeEntered;
 	}
 
+	// @todo vreditor: Move this to FMath so we can use it everywhere
+	// NOTE: OvershootAmount is usually between 0.5 and 2.0, but can go lower and higher for extreme overshots
+	template<class T>
+	static T OvershootEaseOut( T Alpha, const float OvershootAmount = 1.0f )
+	{
+		Alpha--;
+		return 1.0f - ( ( Alpha * ( ( OvershootAmount + 1 ) * Alpha + OvershootAmount ) + 1 ) - 1.0f );
+	}
+
 	/** Gets the maximum length of a laser pointer */
 	float GetLaserPointerMaxLength() const;
 
