VR Editor: Never show on-screen messages in VR Editor (e.g. "lighting is out of date")

diff --git a/Engine/Source/Editor/VREditor/VREditorMode.cpp b/Engine/Source/Editor/VREditor/VREditorMode.cpp
--- a/Engine/Source/Editor/VREditor/VREditorMode.cpp
+++ b/Engine/Source/Editor/VREditor/VREditorMode.cpp
@@ -316,7 +316,7 @@ void FVREditorMode::Enter()
 
 			// Make sure we are in perspective mode
 			// @todo vreditor: We should never allow ortho switching while in VR
-			SavedViewportState.ViewportType = VREditorViewportClient.GetViewportType();
+			SavedEditorState.ViewportType = VREditorViewportClient.GetViewportType();
 			VREditorViewportClient.SetViewportType( LVT_Perspective );
 
 			// Set the initial camera location
@@ -324,8 +324,8 @@ void FVREditorMode::Enter()
 			// location and orientation, compensating for the current HMD offset from the tracking space origin.
 			// Perhaps, we also want to teleport the original viewport's camera back when we exit this mode, too!
 			// @todo vreditor: Should save and restore camera position and any other settings we change (viewport type, pitch locking, etc.)
-			SavedViewportState.ViewLocation = VRViewportClient.GetViewLocation();
-			SavedViewportState.ViewRotation = VRViewportClient.GetViewRotation();
+			SavedEditorState.ViewLocation = VRViewportClient.GetViewLocation();
+			SavedEditorState.ViewRotation = VRViewportClient.GetViewRotation();
 			if( bActuallyUsingVR )
 			{
 				VRViewportClient.SetViewLocation( FVector::ZeroVector );
@@ -334,20 +334,20 @@ void FVREditorMode::Enter()
 
 			// Don't allow the tracking space to pitch up or down.  People hate that in VR.
 			// @todo vreditor: This doesn't seem to prevent people from pitching the camera with RMB drag
-			SavedViewportState.bLockedPitch = VRViewportClient.GetCameraController()->GetConfig().bLockedPitch;
+			SavedEditorState.bLockedPitch = VRViewportClient.GetCameraController()->GetConfig().bLockedPitch;
 			if( bActuallyUsingVR )
 			{
 				VRViewportClient.GetCameraController()->AccessConfig().bLockedPitch = true;
 			}
 
 			// Set "game mode" to be enabled, to get better performance.  Also hit proxies won't work in VR, anyway
-			SavedViewportState.bGameView = VREditorViewportClient.IsInGameView();
+			SavedEditorState.bGameView = VREditorViewportClient.IsInGameView();
 			VREditorViewportClient.SetGameView( true );
 
-			SavedViewportState.bRealTime = VREditorViewportClient.IsRealtime();
+			SavedEditorState.bRealTime = VREditorViewportClient.IsRealtime();
 			VREditorViewportClient.SetRealtime( true );
 
-			SavedViewportState.ShowFlags = VREditorViewportClient.EngineShowFlags;
+			SavedEditorState.ShowFlags = VREditorViewportClient.EngineShowFlags;
 
 			// Never show the traditional Unreal transform widget.  It doesn't work in VR because we don't have hit proxies.
 			VREditorViewportClient.EngineShowFlags.SetModeWidgets( false );
@@ -357,9 +357,12 @@ void FVREditorMode::Enter()
 
 			// Force tiny near clip plane distance, because user can scale themselves to be very small.
 			// @todo vreditor: Make this automatically change based on WorldToMetersScale?
-			SavedViewportState.NearClipPlane = GNearClippingPlane;
+			SavedEditorState.NearClipPlane = GNearClippingPlane;
 			GNearClippingPlane = 1.0f;	// Normally defaults to 10cm (NOTE: If we go too small, skyboxes become affected)
 
+			SavedEditorState.bOnScreenMessages = GAreScreenMessagesEnabled;
+			GAreScreenMessagesEnabled = false;
+
 			// Make the new viewport the active level editing viewport right away
 			GCurrentLevelEditingViewportClient = &VRViewportClient;
 		}
@@ -376,7 +379,7 @@ void FVREditorMode::Enter()
 		{
 			// Tell Slate to require a larger pixel distance threshold before the drag starts.  This is important for things 
 			// like Content Browser drag and drop.
-			SavedViewportState.DragTriggerDistance = FSlateApplication::Get().GetDragTriggerDistance();
+			SavedEditorState.DragTriggerDistance = FSlateApplication::Get().GetDragTriggerDistance();
 			FSlateApplication::Get().SetDragTriggerDistance( 30.0f );
 		}
 	}
@@ -472,7 +475,7 @@ void FVREditorMode::Exit()
 		if( bActuallyUsingVR )
 		{ 
 			// Restore Slate drag trigger distance
-			FSlateApplication::Get().SetDragTriggerDistance( SavedViewportState.DragTriggerDistance );
+			FSlateApplication::Get().SetDragTriggerDistance( SavedEditorState.DragTriggerDistance );
 		}
 
 		TSharedPtr<SLevelViewport> VREditorLevelViewport( VREditorLevelViewportWeakPtr.Pin() );
@@ -486,16 +489,17 @@ void FVREditorMode::Exit()
 				FEditorViewportClient& VREditorViewportClient = VRViewportClient;
 
 				// Restore settings that we changed on the viewport
-				VREditorViewportClient.SetViewportType( SavedViewportState.ViewportType );
-				VRViewportClient.GetCameraController()->AccessConfig().bLockedPitch = SavedViewportState.bLockedPitch;
-				VRViewportClient.bAlwaysShowModeWidgetAfterSelectionChanges = SavedViewportState.bAlwaysShowModeWidgetAfterSelectionChanges;
-				VRViewportClient.EngineShowFlags = SavedViewportState.ShowFlags;
-				VRViewportClient.SetGameView( SavedViewportState.bGameView );
-				VRViewportClient.SetViewLocation( SavedViewportState.ViewLocation );
-				VRViewportClient.SetViewRotation( SavedViewportState.ViewRotation );
-				VRViewportClient.SetRealtime( SavedViewportState.bRealTime );
-
-				GNearClippingPlane = SavedViewportState.NearClipPlane;
+				VREditorViewportClient.SetViewportType( SavedEditorState.ViewportType );
+				VRViewportClient.GetCameraController()->AccessConfig().bLockedPitch = SavedEditorState.bLockedPitch;
+				VRViewportClient.bAlwaysShowModeWidgetAfterSelectionChanges = SavedEditorState.bAlwaysShowModeWidgetAfterSelectionChanges;
+				VRViewportClient.EngineShowFlags = SavedEditorState.ShowFlags;
+				VRViewportClient.SetGameView( SavedEditorState.bGameView );
+				VRViewportClient.SetViewLocation( SavedEditorState.ViewLocation );
+				VRViewportClient.SetViewRotation( SavedEditorState.ViewRotation );
+				VRViewportClient.SetRealtime( SavedEditorState.bRealTime );
+
+				GNearClippingPlane = SavedEditorState.NearClipPlane;
+				GAreScreenMessagesEnabled = SavedEditorState.bOnScreenMessages;
 			}
 
 			if( bActuallyUsingVR )
diff --git a/Engine/Source/Editor/VREditor/VREditorMode.h b/Engine/Source/Editor/VREditor/VREditorMode.h
--- a/Engine/Source/Editor/VREditor/VREditorMode.h
+++ b/Engine/Source/Editor/VREditor/VREditorMode.h
@@ -355,8 +355,8 @@ protected:
 	/** The VR level viewport, if we're in VR mode */
 	TWeakPtr< class SLevelViewport > VREditorLevelViewportWeakPtr;
 
-	/** Saved information about the viewport we possessed, so we can restore it after exiting VR mode */
-	struct FSavedViewportState
+	/** Saved information about the editor and viewport we possessed, so we can restore it after exiting VR mode */
+	struct FSavedEditorState
 	{
 		ELevelViewportType ViewportType;
 		FVector ViewLocation;
@@ -368,8 +368,9 @@ protected:
 		float NearClipPlane;
 		bool bRealTime;
 		float DragTriggerDistance;
+		bool bOnScreenMessages;
 
-		FSavedViewportState()
+		FSavedEditorState()
 			: ViewportType( LVT_Perspective ),
 			  ViewLocation( FVector::ZeroVector ),
 			  ViewRotation( FRotator::ZeroRotator ),
@@ -379,11 +380,12 @@ protected:
 			  bAlwaysShowModeWidgetAfterSelectionChanges( false ),
 			  NearClipPlane( 0.0f ),
 			  bRealTime( false ),
-			  DragTriggerDistance( 0.0f )
+			  DragTriggerDistance( 0.0f ),
+			  bOnScreenMessages( false )
 		{
 		}
 		
-	} SavedViewportState;
+	} SavedEditorState;
 
 	/** Static: True if we're in using an actual HMD in this mode, or false if we're "faking" VR mode for testing */
 	static bool bActuallyUsingVR;
