VR Editor: Fixed UI controls scaling down to zero, causing tons of low spew every frame

diff --git a/Engine/Source/Editor/VREditor/VREditorDockableWindow.cpp b/Engine/Source/Editor/VREditor/VREditorDockableWindow.cpp
--- a/Engine/Source/Editor/VREditor/VREditorDockableWindow.cpp
+++ b/Engine/Source/Editor/VREditor/VREditorDockableWindow.cpp
@@ -192,7 +192,7 @@ void AVREditorDockableWindow::TickManually( float DeltaTime )
 		}
 
 		const float AnimationOvershootAmount = 1.0f;	// @todo vreditor tweak
-		const float EasedAimingAtMeFadeAlpha = FVREditorMode::OvershootEaseOut( AimingAtMeFadeAlpha, AnimationOvershootAmount );
+		const float EasedAimingAtMeFadeAlpha = FMath::Max( 0.001f, FVREditorMode::OvershootEaseOut( AimingAtMeFadeAlpha, AnimationOvershootAmount ) );
 
 		// Only show our extra buttons and controls if the user is roughly aiming toward us.  This just reduces clutter.
 		SelectionBarMeshComponent->SetVisibility( EasedAimingAtMeFadeAlpha > KINDA_SMALL_NUMBER ? true : false );
