VR Editor: Fixed Rift analog stick inverted when scrolling UI up and down

diff --git a/Engine/Source/Editor/VREditor/VREditorUISystem.cpp b/Engine/Source/Editor/VREditor/VREditorUISystem.cpp
--- a/Engine/Source/Editor/VREditor/VREditorUISystem.cpp
+++ b/Engine/Source/Editor/VREditor/VREditorUISystem.cpp
@@ -396,8 +396,11 @@ void FVREditorUISystem::OnVRHoverUpdate( FEditorViewportClient& ViewportClient,
 								}
 							}
 
-							// Invert scroll direction so that it feels more like scrolling on a mobile device
-							ScrollDelta *= -1.0f;
+							// If using a trackpad (Vive), invert scroll direction so that it feels more like scrolling on a mobile device
+							if( Owner.GetHMDDeviceType() == EHMDDeviceType::DT_SteamVR )
+							{
+								ScrollDelta *= -1.0f;
+							}
 
 							if( !FMath::IsNearlyZero( ScrollDelta ) )
 							{
