VR Editor: Fixed not being able to click on UI if laser was hitting a UI handle behind it

diff --git a/Engine/Source/Editor/VREditor/VREditorMode.cpp b/Engine/Source/Editor/VREditor/VREditorMode.cpp
--- a/Engine/Source/Editor/VREditor/VREditorMode.cpp
+++ b/Engine/Source/Editor/VREditor/VREditorMode.cpp
@@ -6,6 +6,7 @@
 #include "VREditorWorldInteraction.h"
 #include "VREditorTransformGizmo.h"
 #include "VREditorFloatingText.h"
+#include "VREditorFloatingUI.h"
 
 #include "CameraController.h"
 #include "DynamicMeshBuilder.h"
@@ -1716,7 +1717,7 @@ bool FVREditorMode::IsInputCaptured( const FVRAction VRAction ) const
 FHitResult FVREditorMode::GetHitResultFromLaserPointer( int32 HandIndex, TArray<AActor*>* OptionalListOfIgnoredActors, const bool bIgnoreGizmos, const bool bEvenIfUIIsInFront, const float LaserLengthOverride )
 {
 	FHitResult BestHitResult;
-	bool bBestHitResultIsUI = false;
+	bool bBestHitResultSoFarIsUI = false;
 
 	FVector LaserPointerStart, LaserPointerEnd;
 	if( GetLaserPointer( HandIndex, LaserPointerStart, LaserPointerEnd, bEvenIfUIIsInFront, LaserLengthOverride ) )
@@ -1751,13 +1752,16 @@ FHitResult FVREditorMode::GetHitResultFromLaserPointer( int32 HandIndex, TArray<
 			{
 				// Is this better than what we have already?  Always prefer transform gizmos even if they were not using
 				// ECC_EditorGizmo (some gizmo handles are opaque and use ECC_Visibility as their collision channel)
-				const bool bHitResultIsUI = UISystem->IsWidgetAnEditorUIWidget( HitResult.GetComponent() );;
+				// NOTE: We're treating components of floating UI actors as "gizmos" for the purpose of hit testing as long as the component is not the actual UI widget component
+				const bool bHitResultIsUI = UISystem->IsWidgetAnEditorUIWidget( HitResult.GetComponent() );
+				const bool bHitResultIsGizmo = ( HitResult.GetActor() == WorldInteraction->GetTransformGizmoActor() ) || ( !bHitResultIsUI && HitResult.GetActor()->IsA( AVREditorFloatingUI::StaticClass() ) );
+				const bool bBestHitResultSoFarIsGizmo = BestHitResult.GetActor() != nullptr && ( ( BestHitResult.GetActor() == WorldInteraction->GetTransformGizmoActor() ) || ( !bBestHitResultSoFarIsUI && BestHitResult.GetActor()->IsA( AVREditorFloatingUI::StaticClass() ) ) );
 				if( BestHitResult.GetActor() == nullptr ||	// Don't have anything yet?
-					( bHitResultIsUI && BestHitResult.GetActor() == WorldInteraction->GetTransformGizmoActor() ) ||
-					( HitResult.GetActor() == WorldInteraction->GetTransformGizmoActor() && !bBestHitResultIsUI ) )	// Always prefer gizmos, unless clicking on UI
+					( bHitResultIsUI && bBestHitResultSoFarIsGizmo ) ||
+					( bHitResultIsGizmo && !bBestHitResultSoFarIsUI ) )	// Always prefer gizmos, unless clicking on UI
 				{
 					BestHitResult = HitResult;
-					bBestHitResultIsUI = bHitResultIsUI;
+					bBestHitResultSoFarIsUI = bHitResultIsUI;
 				}
 			}
 		}
