VREditor - Cleaning up some destruction code to unify the teardown logic for several things.

diff --git a/Engine/Source/Editor/VREditor/VREditorUISystem.cpp b/Engine/Source/Editor/VREditor/VREditorUISystem.cpp
--- a/Engine/Source/Editor/VREditor/VREditorUISystem.cpp
+++ b/Engine/Source/Editor/VREditor/VREditorUISystem.cpp
@@ -104,23 +104,7 @@ FVREditorUISystem::~FVREditorUISystem()
 	Owner.OnVRAction().RemoveAll( this );
 	Owner.OnVRHoverUpdate().RemoveAll( this );
 
-	// Kill floating UIs
-	{
-		for( AVREditorFloatingUI* FloatingUIPtr : FloatingUIs )
-		{
-			if( FloatingUIPtr != nullptr )
-			{
-				FloatingUIPtr->Destroy( false, false );
-				FloatingUIPtr = nullptr;
-			}
-		}
-
-		FloatingUIs.Reset();
-		EditorUIPanels.Reset();
-		EditorUIPanels.SetNumZeroed( (int32)EEditorUIPanel::TotalCount );
-		QuickMenuUI = nullptr;
-		QuickRadialMenu = nullptr;
-	}
+	CleanUpActorsBeforeMapChangeOrSimulate();
 
 	QuickMenuWidgetClass = nullptr;
 	QuickRadialWidgetClass = nullptr;
