VR Editor: Force Oculus HMD mirroring mode 3 (letterboxed) when "-VREditorDemo" is used

diff --git a/Engine/Source/Editor/VREditor/VREditorMode.cpp b/Engine/Source/Editor/VREditor/VREditorMode.cpp
--- a/Engine/Source/Editor/VREditor/VREditorMode.cpp
+++ b/Engine/Source/Editor/VREditor/VREditorMode.cpp
@@ -58,6 +58,7 @@ namespace VREd
 	static FAutoConsoleVariable SleepForRiftHaptics( TEXT( "VREd.SleepForRiftHaptics" ), 1, TEXT( "When enabled, we'll sleep the game thread mid-frame to wait for haptic effects to finish.  This can be devasting to performance!" ) );
 
 	static FAutoConsoleVariable InvertTrackpadVertical( TEXT( "VREd.InvertTrackpadVertical" ), 1, TEXT( "Toggles inverting the touch pad vertical axis" ) );
+	static FAutoConsoleVariable ForceOculusMirrorMode( TEXT( "VREd.ForceOculusMirrorMode" ), 3, TEXT( "Which Oculus display mirroring mode to use (see MirrorWindowModeType in OculusRiftHMD.h)" ) );
 }
 
 FEditorModeID FVREditorMode::VREditorModeID( "VREditor" );
@@ -374,6 +375,14 @@ void FVREditorMode::Enter()
 		if( bActuallyUsingVR )
 		{
 			GEngine->HMDDevice->EnableStereo( true );
+
+			// @todo vreditor: Force single eye, undistorted mirror for demos
+			const bool bIsVREditorDemo = FParse::Param( FCommandLine::Get(), TEXT( "VREditorDemo" ) );	// @todo vreditor: Remove this when no longer needed (console variable, too!)
+			if( bIsVREditorDemo && GetHMDDeviceType() == EHMDDeviceType::DT_OculusRift )
+			{
+				// If we're using an Oculus Rift, go ahead and set the mirror mode to a single undistorted eye
+				GEngine->DeferredCommands.Add( FString::Printf( TEXT( "HMD MIRROR MODE %i" ), VREd::ForceOculusMirrorMode->GetInt() ) );
+			}
 		}
 
 		if( bActuallyUsingVR )
