VREditor - Attempting to refocus the viewport after we launch a new asset editor, can now detect it, to a degree.

diff --git a/Engine/Source/Editor/UnrealEd/Public/EditorViewportClient.h b/Engine/Source/Editor/UnrealEd/Public/EditorViewportClient.h
--- a/Engine/Source/Editor/UnrealEd/Public/EditorViewportClient.h
+++ b/Engine/Source/Editor/UnrealEd/Public/EditorViewportClient.h
@@ -939,6 +939,9 @@ public:
 		return ModeTools;
 	}
 
+	/** Get the editor viewport widget */
+	TSharedPtr<SEditorViewport> GetEditorViewportWidget() const { return EditorViewportWidget.Pin(); }
+
 protected:
 
 	/** Camera speed setting */
diff --git a/Engine/Source/Editor/VREditor/VREditorUISystem.cpp b/Engine/Source/Editor/VREditor/VREditorUISystem.cpp
--- a/Engine/Source/Editor/VREditor/VREditorUISystem.cpp
+++ b/Engine/Source/Editor/VREditor/VREditorUISystem.cpp
@@ -76,7 +76,8 @@ FVREditorUISystem::FVREditorUISystem( FVREditorMode& InitOwner )
 	  DraggingUIHandIndex( INDEX_NONE ),
 	  DraggingUIOffsetTransform( FTransform::Identity ),
 	  bPanelVisibilityToggle( false ),
-	  RadialMenuHideDelayTime( 0.0f )
+	  RadialMenuHideDelayTime( 0.0f ),
+	  bRefocusViewport(false)
 {
 	// Register to find out about VR events
 	Owner.OnVRAction().AddRaw( this, &FVREditorUISystem::OnVRAction );
@@ -478,6 +479,12 @@ void FVREditorUISystem::OnVRHoverUpdate( FEditorViewportClient& ViewportClient,
 
 void FVREditorUISystem::Tick( FEditorViewportClient* ViewportClient, const float DeltaTime )
 {
+	if ( bRefocusViewport )
+	{
+		bRefocusViewport = false;
+		FSlateApplication::Get().SetUserFocus(0, ViewportClient->GetEditorViewportWidget());
+	}
+
 	// Figure out if one hand is "aiming toward" the other hand.  We'll fade in a UI on the hand being
 	// aimed at when the user does this.
 	if( QuickMenuUI != nullptr )
@@ -785,6 +792,7 @@ void FVREditorUISystem::CreateUIs()
 			TSharedPtr<SWindow> TabManagerWindow = TabManagerUI->GetWidgetComponent()->GetSlateWindow();
 			TSharedRef<SWindow> TabManagerWindowRef = TabManagerWindow.ToSharedRef();
 			ProxyTabManager = MakeShareable(new FProxyTabmanager(TabManagerWindowRef));
+			ProxyTabManager->OnProxyTabLaunched.Add(FOnProxyTabLaunched::FDelegate::CreateRaw(this, &FVREditorUISystem::OnProxyTabLaunched));
 
 			// We're going to start stealing tabs from the global tab manager inserting them into the world instead.
 			FGlobalTabmanager::Get()->SetProxyTabManager(ProxyTabManager);
@@ -1111,3 +1119,8 @@ float FVREditorUISystem::GetMinDockWindowSize() const
 {
 	return VREd::MinDockWindowSize->GetFloat();
 }
+
+void FVREditorUISystem::OnProxyTabLaunched(TSharedPtr<SDockTab> NewTab)
+{
+	bRefocusViewport = true;
+}
diff --git a/Engine/Source/Editor/VREditor/VREditorUISystem.h b/Engine/Source/Editor/VREditor/VREditorUISystem.h
--- a/Engine/Source/Editor/VREditor/VREditorUISystem.h
+++ b/Engine/Source/Editor/VREditor/VREditorUISystem.h
@@ -118,6 +118,11 @@ protected:
 	/** Testing Slate UI */
 	void CreateUIs();
 
+	/**
+	 * Called when the injected proxy tab manager reports a new tab has been launched, 
+	 * signaling an asset editor has been launched, but really it could be any major global tab.
+	 */
+	void OnProxyTabLaunched(TSharedPtr<SDockTab> NewTab);
 
 protected:
 
@@ -171,5 +176,8 @@ protected:
 
 	/** Allows us to steal the global tabs and show them in the world. */
 	TSharedPtr<FProxyTabmanager> ProxyTabManager;
+
+	/** Set to true when we need to refocus the viewport. */
+	bool bRefocusViewport;
 };
 
diff --git a/Engine/Source/Runtime/Slate/Private/Framework/Docking/TabManager.cpp b/Engine/Source/Runtime/Slate/Private/Framework/Docking/TabManager.cpp
--- a/Engine/Source/Runtime/Slate/Private/Framework/Docking/TabManager.cpp
+++ b/Engine/Source/Runtime/Slate/Private/Framework/Docking/TabManager.cpp
@@ -2028,6 +2028,8 @@ void FProxyTabmanager::OpenUnmanagedTab(FName PlaceholderId, const FSearchPrefer
 		
 	NewlyOpenedTab->GetParent()->GetParentDockTabStack()->OpenTab(UnmanagedTab);
 	NewlyOpenedTab->RequestCloseTab();
+
+	OnProxyTabLaunched.Broadcast(UnmanagedTab);
 }
 
 
diff --git a/Engine/Source/Runtime/Slate/Public/Framework/Docking/TabManager.h b/Engine/Source/Runtime/Slate/Public/Framework/Docking/TabManager.h
--- a/Engine/Source/Runtime/Slate/Public/Framework/Docking/TabManager.h
+++ b/Engine/Source/Runtime/Slate/Public/Framework/Docking/TabManager.h
@@ -990,6 +990,10 @@ private:
 	TSharedPtr<FProxyTabmanager> ProxyTabManager;
 };
 
+//#HACK VREDITOR - Had to introduce the proxy tab manager to steal asset tabs.
+
+DECLARE_MULTICAST_DELEGATE_OneParam(FOnProxyTabLaunched, TSharedPtr<SDockTab>);
+
 
 class SLATE_API FProxyTabmanager : public FTabManager
 {
@@ -1003,6 +1007,9 @@ public:
 
 	virtual void OpenUnmanagedTab(FName PlaceholderId, const FSearchPreference& SearchPreference, const TSharedRef<SDockTab>& UnmanagedTab) override;
 
+public:
+	FOnProxyTabLaunched OnProxyTabLaunched;
+
 private:
 	TWeakPtr<SWindow> ParentWindow;
 };
