VR Editor: Clicking on UI now only requires a "light trigger press".  This makes it easy to double-click.

diff --git a/Engine/Source/Editor/VREditor/VREditorMode.cpp b/Engine/Source/Editor/VREditor/VREditorMode.cpp
--- a/Engine/Source/Editor/VREditor/VREditorMode.cpp
+++ b/Engine/Source/Editor/VREditor/VREditorMode.cpp
@@ -1324,7 +1324,7 @@ bool FVREditorMode::InputKey( FEditorViewportClient* ViewportClient, FViewport*
 					{
 						// How long since the trigger was lightly pressed?
 						const float TimeSinceLightlyPressed = (float)( FPlatformTime::Seconds() - Hand.RealTimeTriggerWasLightlyPressed );
-						if( !Hand.bAllowTriggerLightPressLocking || TimeSinceLightlyPressed < VREd::TriggerLightlyPressedLockTime->GetFloat() )
+						if( !Hand.bIsClickingOnUI && ( !Hand.bAllowTriggerLightPressLocking || TimeSinceLightlyPressed < VREd::TriggerLightlyPressedLockTime->GetFloat() ) )
 						{
 							Hand.bIsTriggerAtLeastLightlyPressed = false;
 							Hand.bHasTriggerBeenReleasedSinceLastPress = false;
diff --git a/Engine/Source/Editor/VREditor/VREditorUISystem.cpp b/Engine/Source/Editor/VREditor/VREditorUISystem.cpp
--- a/Engine/Source/Editor/VREditor/VREditorUISystem.cpp
+++ b/Engine/Source/Editor/VREditor/VREditorUISystem.cpp
@@ -145,7 +145,7 @@ void FVREditorUISystem::OnVRAction( FEditorViewportClient& ViewportClient, const
 
 			bWasHandled = true;
 		}
-		else if( VRAction.ActionType == EVRActionType::SelectAndMove )
+		else if( VRAction.ActionType == EVRActionType::SelectAndMove_LightlyPressed )
 		{
 			FVirtualHand& Hand = Owner.GetVirtualHand( VRAction.HandIndex );
 			FVector LaserPointerStart, LaserPointerEnd;
