VR Editor: Fixed regression with not being able to drag objects out of content browser

diff --git a/Engine/Source/Editor/VREditor/VREditorWorldInteraction.cpp b/Engine/Source/Editor/VREditor/VREditorWorldInteraction.cpp
--- a/Engine/Source/Editor/VREditor/VREditorWorldInteraction.cpp
+++ b/Engine/Source/Editor/VREditor/VREditorWorldInteraction.cpp
@@ -2700,7 +2700,9 @@ void FVREditorWorldInteraction::OnAssetDragStartedFromContentBrowser( const TArr
 	for( int32 HandIndex = 0; HandIndex < VREditorConstants::NumVirtualHands; ++HandIndex )
 	{
 		FVirtualHand& Hand = Owner.GetVirtualHand( HandIndex );
-		if( Hand.IsInputCaptured[ (int32)EVRActionType::SelectAndMove ] && ( Hand.bIsClickingOnUI && !Hand.bIsRightClickingOnUI ) )
+		if( ( Hand.IsInputCaptured[ (int32)EVRActionType::SelectAndMove ] ||
+			  Hand.IsInputCaptured[ (int32)EVRActionType::SelectAndMove_LightlyPressed ] ) && 
+			( Hand.bIsClickingOnUI && !Hand.bIsRightClickingOnUI ) )
 		{
 			PlacingWithHandIndex = HandIndex;
 			OtherHandIndex = Owner.GetOtherHandIndex( PlacingWithHandIndex );
