VR Editor: Fixed regression with UI selection bar still visible even not aiming at UI

diff --git a/Engine/Source/Editor/VREditor/VREditorDockableWindow.cpp b/Engine/Source/Editor/VREditor/VREditorDockableWindow.cpp
--- a/Engine/Source/Editor/VREditor/VREditorDockableWindow.cpp
+++ b/Engine/Source/Editor/VREditor/VREditorDockableWindow.cpp
@@ -192,12 +192,13 @@ void AVREditorDockableWindow::TickManually( float DeltaTime )
 		}
 
 		const float AnimationOvershootAmount = 1.0f;	// @todo vreditor tweak
-		const float EasedAimingAtMeFadeAlpha = FMath::Max( 0.001f, FVREditorMode::OvershootEaseOut( AimingAtMeFadeAlpha, AnimationOvershootAmount ) );
+		float EasedAimingAtMeFadeAlpha = FVREditorMode::OvershootEaseOut( AimingAtMeFadeAlpha, AnimationOvershootAmount );
 
 		// Only show our extra buttons and controls if the user is roughly aiming toward us.  This just reduces clutter.
 		SelectionBarMeshComponent->SetVisibility( EasedAimingAtMeFadeAlpha > KINDA_SMALL_NUMBER ? true : false );
 		CloseButtonMeshComponent->SetVisibility( EasedAimingAtMeFadeAlpha > KINDA_SMALL_NUMBER ? true : false );
 
+		EasedAimingAtMeFadeAlpha = FMath::Max( 0.001f, EasedAimingAtMeFadeAlpha );
 
 		// Update the window border mesh
 		{
