VR Editor: Force asset UI to always repaint when visible for now
- Later, we could disable this after we manually force repaint after certain actions.

diff --git a/Engine/Source/Editor/VREditor/VREditorUISystem.cpp b/Engine/Source/Editor/VREditor/VREditorUISystem.cpp
--- a/Engine/Source/Editor/VREditor/VREditorUISystem.cpp
+++ b/Engine/Source/Editor/VREditor/VREditorUISystem.cpp
@@ -782,7 +782,10 @@ void FVREditorUISystem::CreateUIs()
 			AVREditorDockableWindow* TabManagerUI = GetOwner().SpawnTransientSceneActor< AVREditorDockableWindow >(TEXT("AssetEditor"), bWithSceneComponent);
 			TabManagerUI->SetSlateWidget(*this, SNullWidget::NullWidget, Resolution, VREd::EditorUISize->GetFloat() * 2.0f, AVREditorFloatingUI::EDockedTo::Nothing);
 			TabManagerUI->ShowUI( false );
-			TabManagerUI->GetWidgetComponent()->SetDrawingPolicy(EVREditorWidgetDrawingPolicy::Hovering);
+
+			// @todo vreditor: Could use "Hovering" instead for better performance with many UIs, but needs to be manually refreshed in some cases
+			TabManagerUI->GetWidgetComponent()->SetDrawingPolicy(EVREditorWidgetDrawingPolicy::Always);
+
 			TabManagerUI->SetRelativeOffset( FVector( 200.0f, 0.0f, 150.0f ) );
 			TabManagerUI->SetLocalRotation( FRotator( 0.0f, 180.0f, 0.0f ) );
 			FloatingUIs.Add(TabManagerUI);
