VR Editor: Fixed Pressing the modifier button in vive pre highlights the wrong button press (UE-27902)

diff --git a/Engine/Source/Editor/VREditor/VREditorMode.cpp b/Engine/Source/Editor/VREditor/VREditorMode.cpp
--- a/Engine/Source/Editor/VREditor/VREditorMode.cpp
+++ b/Engine/Source/Editor/VREditor/VREditorMode.cpp
@@ -2602,7 +2602,7 @@ void FVREditorMode::ApplyButtonPressColors( FVRAction VRAction, EInputEvent Even
 
 	if( VRAction.ActionType == EVRActionType::Modifier )
 	{
-		static FName StaticModifierParameter( "B5" );
+		static FName StaticModifierParameter( "B4" );
 
 		if( Event == IE_Pressed )
 		{
