VR Editor: Made it easier to "throw" the world to move around faster

diff --git a/Engine/Source/Editor/VREditor/VREditorMode.cpp b/Engine/Source/Editor/VREditor/VREditorMode.cpp
--- a/Engine/Source/Editor/VREditor/VREditorMode.cpp
+++ b/Engine/Source/Editor/VREditor/VREditorMode.cpp
@@ -583,12 +583,6 @@ void FVREditorMode::SpawnAvatarMeshActor()
 
 				Hand.MotionControllerComponent->Hand = ( HandIndex == VREditorConstants::LeftHandIndex ) ? EControllerHand::Left : EControllerHand::Right;
 
-				// @todo vreditor urgent: Late updates on motion controllers has a race condition bug that causes controller
-				// meshes to pop around, especially while world scaling.  This manifests in different ways depending on the HMD.
-				// On Oculus, the render thread will query cached state from a game thread data structure that is updated at 
-				// the top of every frame, potentially clobbering what the RT uses before rendering (popping).  On the Vive,
-				// controller state is cached by the HMD's late update code, then re-used by the game thread later, however
-				// the state can be clobbered either before or after the game thread reads it (popping.)
 				Hand.MotionControllerComponent->bDisableLowLatencyUpdate = false;
 			}
 
@@ -1738,7 +1732,7 @@ FHitResult FVREditorMode::GetHitResultFromLaserPointer( int32 HandIndex, TArray<
 		{
 			const bool bOnlyEditorGizmos = ( PassIndex == 0 );
 
-			const bool bTraceComplex = true;	// @todo vreditor: Not sure if we should ever use complex collision here.  UI doesn't seem to get hit if we don't.
+			const bool bTraceComplex = true;
 			FCollisionQueryParams TraceParams( NAME_None, bTraceComplex, nullptr );
 			
 			const FCollisionResponseParams& ResponseParam = FCollisionResponseParams::DefaultResponseParam;
@@ -2528,9 +2522,9 @@ void FVREditorMode::ApplyVelocityDamping( FVector& Velocity, const bool bVelocit
 			// @todo vreditor: Experimenting with speed-based dampening amounts here, so that you can drift further
 			// after you've been "thrown faster", but come to a stop rather quickly if you are not already moving fast.
 			//		--> If this is successful, consider generalizing this and doing it everywhere else we need inertial movement
-			const float DampenMultiplierAtLowSpeeds = 0.95f;	// @todo vreditor tweak
-			const float DampenMultiplierAtHighSpeeds = 0.9925f;	// @todo vreditor tweak
-			const float SpeedForMinimalDamping = 3.0f * GetWorldScaleFactor();	// cm/frame	// @todo vreditor tweak
+			const float DampenMultiplierAtLowSpeeds = 0.94f;	// @todo vreditor tweak
+			const float DampenMultiplierAtHighSpeeds = 0.99f;	// @todo vreditor tweak
+			const float SpeedForMinimalDamping = 2.5f * GetWorldScaleFactor();	// cm/frame	// @todo vreditor tweak
 			const float SpeedBasedDampeningScalar = FMath::Clamp( Velocity.Size(), 0.0f, SpeedForMinimalDamping ) / SpeedForMinimalDamping;	// @todo vreditor: Probably needs a curve applied to this to compensate for our framerate insensitivity
 			Velocity = Velocity * FMath::Lerp( DampenMultiplierAtLowSpeeds, DampenMultiplierAtHighSpeeds, SpeedBasedDampeningScalar );	// @todo vreditor: Frame rate sensitive damping.  Make use of delta time!
 		}
