VR Editor: Fixed the dockable window hover and leave state for its components

diff --git a/Engine/Source/Editor/VREditor/VREditorDockableWindow.cpp b/Engine/Source/Editor/VREditor/VREditorDockableWindow.cpp
--- a/Engine/Source/Editor/VREditor/VREditorDockableWindow.cpp
+++ b/Engine/Source/Editor/VREditor/VREditorDockableWindow.cpp
@@ -286,31 +286,34 @@ void AVREditorDockableWindow::UpdateRelativeRoomTransform()
 	LocalRotation = RoomSpaceWindowRotation.Rotator();
 }
 
-void AVREditorDockableWindow::OnEnterHover( const FHitResult& HitResult )
+void AVREditorDockableWindow::OnEnterHover( const FHitResult& HitResult, const int32 HandIndex )
 {
 	if( SelectionBarMeshComponent == HitResult.GetComponent() )
 	{
 		bIsHoveringOverSelectionBar = true;
-		bIsHoveringOverCloseButton = false;
 	}
-	else if( CloseButtonMeshComponent == HitResult.GetComponent() )
+	
+	if( CloseButtonMeshComponent == HitResult.GetComponent() )
 	{
 		bIsHoveringOverCloseButton = true;
+	}
+}
+
+void AVREditorDockableWindow::OnLeaveHover( const int32 HandIndex, const UActorComponent* NewComponent  )
+{
+	FVirtualHand& OtherHand = GetOwner().GetOwner().GetOtherHand( HandIndex );
+
+	if( OtherHand.HoveredActorComponent != SelectionBarMeshComponent && NewComponent != SelectionBarMeshComponent )
+	{
 		bIsHoveringOverSelectionBar = false;
 	}
-	else
+	
+	if( OtherHand.HoveredActorComponent != CloseButtonMeshComponent && NewComponent != CloseButtonMeshComponent  )
 	{
 		bIsHoveringOverCloseButton = false;
-		bIsHoveringOverSelectionBar = false;
 	}
 }
 
-void AVREditorDockableWindow::OnLeaveHover()
-{
-	bIsHoveringOverSelectionBar = false;
-	bIsHoveringOverCloseButton = false;
-}
-
 
 void AVREditorDockableWindow::SetSelectionBarColor( const FLinearColor& LinearColor )
 {
diff --git a/Engine/Source/Editor/VREditor/VREditorDockableWindow.h b/Engine/Source/Editor/VREditor/VREditorDockableWindow.h
--- a/Engine/Source/Editor/VREditor/VREditorDockableWindow.h
+++ b/Engine/Source/Editor/VREditor/VREditorDockableWindow.h
@@ -24,11 +24,11 @@ public:
 	/** Updates the last dragged relative position */
 	void UpdateRelativeRoomTransform();
 
-	/** Enter hover with laser changes the color of SelectionMesh */
-	void OnEnterHover( const FHitResult& HitResult );
+	/** Enter hover with laser changes the color of SelectionMesh and CloseButtonMesh */
+	void OnEnterHover( const FHitResult& HitResult, const int32 HandIndex );
 
-	/** Leaving hover with laser changes the color of SelectionMesh */
-	void OnLeaveHover();
+	/** Leaving hover with laser changes the color of SelectionMesh and CloseButtonMesh */
+	void OnLeaveHover( const int32 HandIndex, const class UActorComponent* NewHoveredComponent );
 	
 	/** Gets the close button component */
 	UStaticMeshComponent* GetCloseButtonMeshComponent()
@@ -36,6 +36,12 @@ public:
 		return CloseButtonMeshComponent;
 	}
 
+	/** Gets the selection bar component */
+	UStaticMeshComponent* GetSelectionBarMeshComponent()
+	{
+		return SelectionBarMeshComponent;
+	}
+
 private:
 
 	// AVREditorFloatingUI overrides
diff --git a/Engine/Source/Editor/VREditor/VREditorWorldInteraction.cpp b/Engine/Source/Editor/VREditor/VREditorWorldInteraction.cpp
--- a/Engine/Source/Editor/VREditor/VREditorWorldInteraction.cpp
+++ b/Engine/Source/Editor/VREditor/VREditorWorldInteraction.cpp
@@ -45,7 +45,6 @@ namespace VREd
 	static FAutoConsoleVariable GridHapticFeedbackStrength( TEXT( "VREd.GridHapticFeedbackStrength" ), 0.4f, TEXT( "Default strength for haptic feedback when moving across grid points" ) );
 	static FAutoConsoleVariable SelectionHapticFeedbackStrength( TEXT( "VREd.SelectionHapticFeedbackStrength" ), 0.5f, TEXT( "Default strength for haptic feedback when selecting objects" ) );
 	static FAutoConsoleVariable DragHapticFeedbackStrength( TEXT( "VREd.DragHapticFeedbackStrength" ), 1.0f, TEXT( "Default strength for haptic feedback when starting to drag objects" ) );
-	static FAutoConsoleVariable GizmoHoverHapticFeedbackStrength( TEXT( "VREd.GizmoHoverHapticFeedbackStrength" ), 0.1f, TEXT( "Default strength for haptic feedback when hovering over gizmos" ) );
 	static FAutoConsoleVariable ForceGizmoPivotToCenterOfSelectedActorsBounds( TEXT( "VREd.ForceGizmoPivotToCenterOfSelectedActorsBounds" ), 0, TEXT( "When enabled, the gizmo's pivot will always be centered on the selected actors.  Otherwise, we use the pivot of the last selected actor." ) );
 	static FAutoConsoleVariable TrackpadAbsoluteDragSpeed( TEXT( "VREd.TrackpadAbsoluteDragSpeed" ), 40.0f, TEXT( "How fast objects move toward or away when you drag on the touchpad while carrying them" ) );
 	static FAutoConsoleVariable TrackpadRelativeDragSpeed( TEXT( "VREd.TrackpadRelativeDragSpeed" ), 8.0f, TEXT( "How fast objects move toward or away when you hold a direction on an analog stick while carrying them" ) );
@@ -53,6 +52,10 @@ namespace VREd
 	static FAutoConsoleVariable GizmoHandleHoverAnimationDuration( TEXT( "VREd.GizmoHandleHoverAnimationDuration" ), 0.1f, TEXT( "How quickly to animate gizmo handle hover state" ) );
 	static FAutoConsoleVariable SnapGridSize( TEXT( "VREd.SnapGridSize" ), 3.0f, TEXT( "How big the snap grid should be, in multiples of the gizmo bounding box size" ) );
 	static FAutoConsoleVariable SnapGridLineWidth( TEXT( "VREd.SnapGridLineWidth" ), 3.0f, TEXT( "Width of the grid lines on the snap grid" ) );
+	static FAutoConsoleVariable GizmoHoverHapticFeedbackStrength( TEXT( "VREd.GizmoHoverHapticFeedbackStrength" ), 0.1f, TEXT( "Default strength for haptic feedback when hovering over gizmos" ) );
+	static FAutoConsoleVariable GizmoHoverHapticFeedbackTime( TEXT( "VREd.GizmoHoverHapticFeedbackTime" ), 1.0f, TEXT( "The minimum time between haptic feedback between hovering on a gizmo handle" ) );
+	static FAutoConsoleVariable WindowHoverHapticFeedbackStrength( TEXT( "VREd.WindowHoverHapticFeedbackStrength" ), 0.4f, TEXT( "Width of the grid lines on the snap grid" ) );
+	static FAutoConsoleVariable WindowHoverHapticFeedbackTime( TEXT( "VREd.WindowHoverHapticFeedbackTime" ), 0.2f, TEXT( "The minimum time between haptic feedback between hovering on a dockable window" ) );
 }
 
 
@@ -396,7 +399,7 @@ void FVREditorWorldInteraction::OnVRAction( FEditorViewportClient& ViewportClien
 										const bool bRefreshQuickMenu = true;
 										Owner.GetUISystem().ShowEditorUIPanel( DockableWindow, Owner.GetOtherHandIndex( VRAction.HandIndex ), bShouldShow, bShowOnHand, bRefreshQuickMenu );
 									}
-									else
+									else if ( HitResult.Component == DockableWindow->GetSelectionBarMeshComponent() )
 									{
 										Hand.DraggingMode = EVREditorDraggingMode::DockableWindow;
 										DockSelectDistance = ( LaserPointerStart - HitResult.Location ).Size();
@@ -606,8 +609,8 @@ void FVREditorWorldInteraction::OnVRAction( FEditorViewportClient& ViewportClien
 						FloatingUIAssetDraggedFrom = nullptr;
 					}
 
-					bWasHandled = true;
 					StopDragging( VRAction.HandIndex );
+					bWasHandled = true;
 				}
 			}
 		}
@@ -675,6 +678,9 @@ void FVREditorWorldInteraction::OnVRHoverUpdate( FEditorViewportClient& Viewport
 		FVirtualHand& Hand = Owner.GetVirtualHand( HandIndex );
 		FVirtualHand& OtherHand = Owner.GetOtherHand( HandIndex );
 
+		float HapticFeedbackStrength = 0.0f;
+		float HapticFeedbackTimeBuffer = 0.0f;
+
 		UActorComponent* PreviousHoverGizmoComponent = Hand.HoveringOverTransformGizmoComponent.Get();
 		Hand.HoveringOverTransformGizmoComponent = nullptr;
 
@@ -690,65 +696,74 @@ void FVREditorWorldInteraction::OnVRHoverUpdate( FEditorViewportClient& Viewport
 				this->HoveredObjects.Add( FViewportHoverTarget( Actor ) );
 
 				HoverImpactPoint = HitResult.ImpactPoint;
-				NewHoveredActorComponent = HitResult.GetComponent();
 
 				if( Actor == TransformGizmoActor )
 				{
+					NewHoveredActorComponent = HitResult.GetComponent();
 					Hand.HoveringOverTransformGizmoComponent = HitResult.GetComponent();
 
 					if( Hand.HoveringOverTransformGizmoComponent != PreviousHoverGizmoComponent )
 					{
-						const FTimespan CurrentTime = FTimespan::FromSeconds( FPlatformTime::Seconds() );
-						const FTimespan TimeSinceLastHover = CurrentTime - Hand.LastGizmoHandleHapticTime;
-
-						// Only play a haptic if we hovered over a different gizmo, or if we haven't already played a haptic in a little while
-						if( Hand.HoveringOverTransformGizmoComponent != Hand.HoverHapticCheckLastHoveredGizmoComponent ||
-							TimeSinceLastHover.GetTotalSeconds() > 1.0f )	// @todo vreditor tweak
-						{
-							const float Strength = VREd::GizmoHoverHapticFeedbackStrength->GetFloat();
-							Owner.PlayHapticEffect(
-								HandIndex == VREditorConstants::LeftHandIndex ? Strength : 0.0f,
-								HandIndex == VREditorConstants::RightHandIndex ? Strength : 0.0f );
-
-							Hand.LastGizmoHandleHapticTime = CurrentTime;
-						}
+						HapticFeedbackStrength = VREd::GizmoHoverHapticFeedbackStrength->GetFloat();
+						HapticFeedbackTimeBuffer = VREd::GizmoHoverHapticFeedbackTime->GetFloat();
 					}
 
 					Hand.HoverHapticCheckLastHoveredGizmoComponent = HitResult.GetComponent();
 				}
-				else if ( Actor->IsA( AVREditorDockableWindow::StaticClass() ) )
+				else if( Actor->IsA( AVREditorDockableWindow::StaticClass() ) )
 				{
+					NewHoveredActorComponent = HitResult.GetComponent();
 					AVREditorDockableWindow* DockableWindow = Cast< AVREditorDockableWindow >( Actor );
-					if (DockableWindow && NewHoveredActorComponent != OtherHand.HoveredActorComponent )
+					if(DockableWindow && NewHoveredActorComponent != Hand.HoveredActorComponent )
 					{
-						DockableWindow->OnEnterHover( HitResult );
+						if( NewHoveredActorComponent != OtherHand.HoveredActorComponent )
+						{
+							DockableWindow->OnEnterHover( HitResult, HandIndex );
+						}
+						
+						HapticFeedbackStrength = VREd::WindowHoverHapticFeedbackStrength->GetFloat();
+						HapticFeedbackTimeBuffer = VREd::WindowHoverHapticFeedbackTime->GetFloat();
 					}
 				}
 
 				bWasHandled = true;
 			}
 		}
-		
-		// Leave dockable window hover
-		if( Hand.HoveredActorComponent != nullptr && Hand.HoveredActorComponent != OtherHand.HoveredActorComponent &&
-			( NewHoveredActorComponent == nullptr || Hand.HoveredActorComponent != NewHoveredActorComponent || !HitResult.Actor.IsValid() ) )
+
+		// Leave dockable window
+		if( Hand.HoveredActorComponent != nullptr && ( Hand.HoveredActorComponent != NewHoveredActorComponent || NewHoveredActorComponent == nullptr ) && Hand.HoveredActorComponent != OtherHand.HoveredActorComponent )
 		{
-			AActor* Actor = Hand.HoveredActorComponent->GetOwner();
-			if (Actor->IsA( AVREditorDockableWindow::StaticClass() ))
+			AActor* HoveredActor = Hand.HoveredActorComponent->GetOwner();
+			if (HoveredActor->IsA( AVREditorDockableWindow::StaticClass() ))
 			{
-				AVREditorDockableWindow* DockableWindow = Cast< AVREditorDockableWindow >( Actor );
+				AVREditorDockableWindow* DockableWindow = Cast< AVREditorDockableWindow >( HoveredActor );
 				if (DockableWindow)
 				{
-					DockableWindow->OnLeaveHover();
+					DockableWindow->OnLeaveHover( HandIndex, NewHoveredActorComponent );
 				}
 			}
 		}
 
-		// Set the new hovered actor component if there is any
-		if( HitResult.Actor.IsValid() && NewHoveredActorComponent )
+		// Play the haptic feedback
+		if( HapticFeedbackStrength > 0.0f )
 		{
-			Hand.HoveredActorComponent = NewHoveredActorComponent;
+			const FTimespan CurrentTime = FTimespan::FromSeconds( FPlatformTime::Seconds() );
+			const FTimespan TimeSinceLastHover = CurrentTime - Hand.LastGizmoHandleHapticTime;
+
+			// Only play a haptic if we hovered over a different gizmo, or if we haven't already played a haptic in a little while
+			if (Hand.HoveringOverTransformGizmoComponent != Hand.HoverHapticCheckLastHoveredGizmoComponent ||
+				TimeSinceLastHover.GetTotalSeconds() > HapticFeedbackTimeBuffer )
+			{
+				Owner.PlayHapticEffect(
+					HandIndex == VREditorConstants::LeftHandIndex ? HapticFeedbackStrength : 0.0f,
+					HandIndex == VREditorConstants::RightHandIndex ? HapticFeedbackStrength : 0.0f );
+
+				Hand.LastGizmoHandleHapticTime = CurrentTime;
+			}
 		}
+
+		// Update the hovered actor component with the new component
+		Hand.HoveredActorComponent = NewHoveredActorComponent;
 	}
 }
 
