VR Editor: Added subtle haptic vibration when clicking on UI

diff --git a/Engine/Source/Editor/VREditor/VREditorUISystem.cpp b/Engine/Source/Editor/VREditor/VREditorUISystem.cpp
--- a/Engine/Source/Editor/VREditor/VREditorUISystem.cpp
+++ b/Engine/Source/Editor/VREditor/VREditorUISystem.cpp
@@ -49,7 +49,7 @@ namespace VREd
 	static FAutoConsoleVariable MinUIScrollDeltaForInertia( TEXT( "VREd.MinUIScrollDeltaForInertia" ), 0.25f, TEXT( "Minimum amount of touch pad input before inertial UI scrolling kicks in" ) );
 	static FAutoConsoleVariable MinDockDragDistance( TEXT( "VREd.MinDockDragDistance" ), 10.0f, TEXT( "Minimum amount of distance needed to drag dock from hand" ) );
 	static FAutoConsoleVariable DoubleClickTime( TEXT( "VREd.DoubleClickTime" ), 0.25f, TEXT( "Minimum duration between clicks for a double click event to fire" ) );
-
+	static FAutoConsoleVariable UIPressHapticFeedbackStrength( TEXT( "VREd.UIPressHapticFeedbackStrength" ), 0.4f, TEXT( "Strenth of haptic when clicking on the UI" ) );
 	static FAutoConsoleVariable MaxDockWindowSize( TEXT( "VREd.MaxDockWindowSize" ), 250, TEXT( "Maximum size for dockable windows" ) );
 	static FAutoConsoleVariable MinDockWindowSize( TEXT( "VREd.MinDockWindowSize" ), 40, TEXT( "Minimum size for dockable windows" ) );
 
@@ -218,6 +218,12 @@ void FVREditorUISystem::OnVRAction( FEditorViewportClient& ViewportClient, const
 										Hand.bIsClickingOnUI = true;
 										Hand.bIsRightClickingOnUI = bIsRightClicking;
 										bIsInputCaptured = true;
+
+										// Play a haptic effect on press
+										const float Strength = VREd::UIPressHapticFeedbackStrength->GetFloat();
+										Owner.PlayHapticEffect(
+											VRAction.HandIndex == VREditorConstants::LeftHandIndex ? Strength : 0.0f,
+											VRAction.HandIndex == VREditorConstants::RightHandIndex ? Strength : 0.0f );
 									}
 									else if( Event == IE_Released )
 									{
