Hiding and showing dockable windows
VR Editor:
- Fixed multiple windows on hand when using hide/show all button (UE-27968)
- Added sounds for dragging, dropping and visibility of dockable windows
- Added dragging state for floating UI, so when dragged the UI doesn't have to be state nothing, which caused a few bugs
- Using quick menu to set the visibility of a dockable window now sets the window back on the hand when in room space

diff --git a/Engine/Source/Editor/VREditor/VREditorFloatingUI.cpp b/Engine/Source/Editor/VREditor/VREditorFloatingUI.cpp
--- a/Engine/Source/Editor/VREditor/VREditorFloatingUI.cpp
+++ b/Engine/Source/Editor/VREditor/VREditorFloatingUI.cpp
@@ -222,11 +222,11 @@ void AVREditorFloatingUI::UpdateFadingState( const float DeltaTime )
 		WidgetComponent->SetTintColorAndOpacity( FLinearColor( UIBrightness, UIBrightness, UIBrightness ).CopyWithNewOpacity( FadeAlpha ) );
 	}
 }
-
+	
 
 void AVREditorFloatingUI::UpdateTransformIfDocked()
 {
-	if( DockedTo != EDockedTo::Nothing )
+	if( DockedTo != EDockedTo::Nothing && DockedTo != EDockedTo::Dragging )
 	{
 		FTransform NewTransform;
 		
diff --git a/Engine/Source/Editor/VREditor/VREditorFloatingUI.h b/Engine/Source/Editor/VREditor/VREditorFloatingUI.h
--- a/Engine/Source/Editor/VREditor/VREditorFloatingUI.h
+++ b/Engine/Source/Editor/VREditor/VREditorFloatingUI.h
@@ -30,6 +30,7 @@ public:
 		Room,
 		Head,
 		Custom,
+		Dragging,
 	};
 
 	/** Creates a FVREditorFloatingUI using a Slate widget, and sets up safe defaults */
diff --git a/Engine/Source/Editor/VREditor/VREditorQuickMenu.cpp b/Engine/Source/Editor/VREditor/VREditorQuickMenu.cpp
--- a/Engine/Source/Editor/VREditor/VREditorQuickMenu.cpp
+++ b/Engine/Source/Editor/VREditor/VREditorQuickMenu.cpp
@@ -245,9 +245,7 @@ bool UVREditorQuickMenu::IsSimulatingEnabled() const
 
 void UVREditorQuickMenu::OnContentBrowserButtonClicked( const bool bIsChecked )
 {
-	bool bShouldShow = !IsContentBrowserVisible();
-	const int32 MyHandIndex = GetOwner()->GetDockedTo() == AVREditorFloatingUI::EDockedTo::LeftArm ? VREditorConstants::LeftHandIndex : VREditorConstants::RightHandIndex;
-	GetOwner()->GetOwner().ShowEditorUIPanel( FVREditorUISystem::EEditorUIPanel::ContentBrowser, MyHandIndex, bShouldShow );
+	GetOwner()->GetOwner().TogglePanelVisibility( FVREditorUISystem::EEditorUIPanel::ContentBrowser );
 }
 
 
@@ -259,9 +257,7 @@ bool UVREditorQuickMenu::IsContentBrowserVisible() const
 
 void UVREditorQuickMenu::OnWorldOutlinerButtonClicked( const bool bIsChecked )
 {
-	bool bShouldShow = !IsWorldOutlinerVisible();
-	const int32 MyHandIndex = GetOwner()->GetDockedTo() == AVREditorFloatingUI::EDockedTo::LeftArm ? VREditorConstants::LeftHandIndex : VREditorConstants::RightHandIndex;
-	GetOwner()->GetOwner().ShowEditorUIPanel( FVREditorUISystem::EEditorUIPanel::WorldOutliner, MyHandIndex, bShouldShow );
+	GetOwner()->GetOwner().TogglePanelVisibility( FVREditorUISystem::EEditorUIPanel::WorldOutliner );
 }
 
 
@@ -273,9 +269,7 @@ bool UVREditorQuickMenu::IsWorldOutlinerVisible() const
 
 void UVREditorQuickMenu::OnActorDetailsButtonClicked( const bool bIsChecked )
 {
-	bool bShouldShow = !IsActorDetailsVisible();
-	const int32 MyHandIndex = GetOwner()->GetDockedTo() == AVREditorFloatingUI::EDockedTo::LeftArm ? VREditorConstants::LeftHandIndex : VREditorConstants::RightHandIndex;
-	GetOwner()->GetOwner().ShowEditorUIPanel( FVREditorUISystem::EEditorUIPanel::ActorDetails, MyHandIndex, bShouldShow );
+	GetOwner()->GetOwner().TogglePanelVisibility( FVREditorUISystem::EEditorUIPanel::ActorDetails );
 }
 
 
@@ -287,9 +281,7 @@ bool UVREditorQuickMenu::IsActorDetailsVisible() const
 
 void UVREditorQuickMenu::OnModesButtonClicked( const bool bIsChecked )
 {
-	bool bShouldShow = !IsModesVisible();
-	const int32 MyHandIndex = GetOwner()->GetDockedTo() == AVREditorFloatingUI::EDockedTo::LeftArm ? VREditorConstants::LeftHandIndex : VREditorConstants::RightHandIndex;
-	GetOwner()->GetOwner().ShowEditorUIPanel( FVREditorUISystem::EEditorUIPanel::Modes, MyHandIndex, bShouldShow );
+	GetOwner()->GetOwner().TogglePanelVisibility( FVREditorUISystem::EEditorUIPanel::Modes );
 }
 
 
@@ -319,9 +311,7 @@ bool UVREditorQuickMenu::IsGameModeEnabled() const
 
 void UVREditorQuickMenu::OnTutorialButtonClicked( const bool bIsChecked )
 {
-	bool bShouldShow = !IsTutorialVisible();
-	const int32 MyHandIndex = GetOwner()->GetDockedTo() == AVREditorFloatingUI::EDockedTo::LeftArm ? VREditorConstants::LeftHandIndex : VREditorConstants::RightHandIndex;
-	GetOwner()->GetOwner().ShowEditorUIPanel( FVREditorUISystem::EEditorUIPanel::Tutorial, MyHandIndex, bShouldShow );
+	GetOwner()->GetOwner().TogglePanelVisibility( FVREditorUISystem::EEditorUIPanel::Tutorial );
 }
 
 bool UVREditorQuickMenu::IsTutorialVisible() const
@@ -332,9 +322,8 @@ bool UVREditorQuickMenu::IsTutorialVisible() const
 
 void UVREditorQuickMenu::OnAssetEditorButtonClicked( const bool bIsChecked )
 {
-	bool bShouldShow = !IsAssetEditorVisible();
-	const int32 MyHandIndex = GetOwner()->GetDockedTo() == AVREditorFloatingUI::EDockedTo::LeftArm ? VREditorConstants::LeftHandIndex : VREditorConstants::RightHandIndex;
-	GetOwner()->GetOwner().ShowEditorUIPanel( FVREditorUISystem::EEditorUIPanel::AssetEditor, MyHandIndex, bShouldShow );
+	GetOwner()->GetOwner().TogglePanelVisibility( FVREditorUISystem::EEditorUIPanel::AssetEditor );
+	
 }
 
 bool UVREditorQuickMenu::IsAssetEditorVisible() const
diff --git a/Engine/Source/Editor/VREditor/VREditorRadialMenu.cpp b/Engine/Source/Editor/VREditor/VREditorRadialMenu.cpp
--- a/Engine/Source/Editor/VREditor/VREditorRadialMenu.cpp
+++ b/Engine/Source/Editor/VREditor/VREditorRadialMenu.cpp
@@ -15,7 +15,7 @@
 namespace VREd
 {
 	static FAutoConsoleVariable RadialMenuRiftSelectionRadius( TEXT( "VREd.RadialMenuRiftSelectionRadius" ), 0.25f, TEXT( "The radius that is used for selecting buttons on the radial menu on Oculus Rift" ) );
-	static FAutoConsoleVariable RadialMenuViveSelectionRadius( TEXT( "VREd.RadialMenuViveSelectionRadius" ), 0.5f, TEXT( "The radius that is used for selecting buttons on the radial menu on Vive" ) );
+	static FAutoConsoleVariable RadialMenuViveSelectionRadius( TEXT( "VREd.RadialMenuViveSelectionRadius" ), 0.3f, TEXT( "The radius that is used for selecting buttons on the radial menu on Vive" ) );
 }
 
 UVREditorRadialMenu::UVREditorRadialMenu( const FObjectInitializer& ObjectInitializer )
diff --git a/Engine/Source/Editor/VREditor/VREditorUISystem.cpp b/Engine/Source/Editor/VREditor/VREditorUISystem.cpp
--- a/Engine/Source/Editor/VREditor/VREditorUISystem.cpp
+++ b/Engine/Source/Editor/VREditor/VREditorUISystem.cpp
@@ -94,6 +94,19 @@ FVREditorUISystem::FVREditorUISystem( FVREditorMode& InitOwner )
 	TutorialWidgetClass = LoadClass<UVREditorBaseUserWidget>( nullptr, TEXT( "/Engine/VREditor/Tutorial/UI_VR_Tutorial_00.UI_VR_Tutorial_00_C" ) );
 	check( TutorialWidgetClass != nullptr );
 
+	// Load sounds
+	StartDragUISound = LoadObject<USoundCue>( nullptr, TEXT( "/Engine/VREditor/Sounds/VR_negative_Cue" ) );
+	check( StartDragUISound != nullptr );
+
+	StopDragUISound = LoadObject<USoundCue>( nullptr, TEXT( "/Engine/VREditor/Sounds/VR_negative_Cue" ) );
+	check( StopDragUISound != nullptr );
+
+	HideUISound = LoadObject<USoundCue>( nullptr, TEXT( "/Engine/VREditor/Sounds/VR_close_Cue" ) );
+	check( HideUISound != nullptr );
+
+	ShowUISound = LoadObject<USoundCue>( nullptr, TEXT( "/Engine/VREditor/Sounds/VR_open_Cue" ) );
+	check( ShowUISound != nullptr );
+
 	// Create all of our UI panels
 	CreateUIs();
 }
@@ -290,16 +303,16 @@ void FVREditorUISystem::OnVRAction( FEditorViewportClient& ViewportClient, const
 			FVirtualHand& Hand = Owner.GetVirtualHand( VRAction.HandIndex );
 			FVirtualHand& OtherHand = Owner.GetOtherHand( VRAction.HandIndex );
 			const float Distance = FVector::Dist( Hand.HoverLocation, OtherHand.Transform.GetLocation() ) / Owner.GetWorldScaleFactor();
+			bool bOnHand = true;
 			if (Distance > VREd::MinDockDragDistance->GetFloat())
 			{
 				DraggingUI->SetDockedTo( AVREditorFloatingUI::EDockedTo::Room );
-			}
-			else
-			{
-				const int32 OtherHandIndex = Owner.GetOtherHandIndex( DraggingUIHandIndex );
-				ShowEditorUIPanel( DraggingUI, OtherHandIndex, true, true );
+				bOnHand = false;
 			}
 
+			const int32 OtherHandIndex = Owner.GetOtherHandIndex( DraggingUIHandIndex );
+			ShowEditorUIPanel( DraggingUI, OtherHandIndex, true, bOnHand, true );
+			
 			StopDraggingDockUI();
 		}
 	}
@@ -773,9 +786,12 @@ void FVREditorUISystem::CreateUIs()
 		{
 			AVREditorFloatingUI* TutorialUI = GetOwner().SpawnTransientSceneActor< AVREditorDockableWindow >( TEXT( "TutorialUI" ), false );
 			TutorialUI->SetUMGWidget( *this, TutorialWidgetClass, FIntPoint( VREd::TutorialUIResolutionX->GetFloat(), VREd::TutorialUIResolutionY->GetFloat() ), VREd::TutorialUISize->GetFloat(), AVREditorFloatingUI::EDockedTo::Room );
-			TutorialUI->ShowUI( true );
 			TutorialUI->SetRelativeOffset( FVector( VREd::TutorialUILocationX->GetFloat(), VREd::TutorialUILocationY->GetFloat(), VREd::TutorialUILocationZ->GetFloat() ) );
 			TutorialUI->SetLocalRotation( FRotator( VREd::TutorialUIPitch->GetFloat(), VREd::TutorialUIYaw->GetFloat(), 0 ) );
+			
+			const bool bIsVREditorDemo = FParse::Param( FCommandLine::Get(), TEXT( "VREditorDemo" ) );
+			TutorialUI->ShowUI( !bIsVREditorDemo );
+
 			FloatingUIs.Add( TutorialUI );
 
 			EditorUIPanels[ (int32)EEditorUIPanel::Tutorial ] = TutorialUI;		
@@ -864,7 +880,7 @@ bool FVREditorUISystem::IsShowingEditorUIPanel( const EEditorUIPanel EditorUIPan
 }
 
 
-void FVREditorUISystem::ShowEditorUIPanel( const UWidgetComponent* WidgetComponent, const int32 HandIndex, const bool bShouldShow, const bool OnHand, const bool bRefreshQuickMenu  )
+void FVREditorUISystem::ShowEditorUIPanel( const UWidgetComponent* WidgetComponent, const int32 HandIndex, const bool bShouldShow, const bool OnHand, const bool bRefreshQuickMenu )
 {
 	AVREditorFloatingUI* Panel = nullptr;
 	for( AVREditorFloatingUI* CurrentPanel : EditorUIPanels )
@@ -880,7 +896,7 @@ void FVREditorUISystem::ShowEditorUIPanel( const UWidgetComponent* WidgetCompone
 }
 
 
-void FVREditorUISystem::ShowEditorUIPanel( const EEditorUIPanel EditorUIPanel, const int32 HandIndex, const bool bShouldShow, const bool OnHand, const bool bRefreshQuickMenu  )
+void FVREditorUISystem::ShowEditorUIPanel( const EEditorUIPanel EditorUIPanel, const int32 HandIndex, const bool bShouldShow, const bool OnHand, const bool bRefreshQuickMenu )
 {
 	AVREditorFloatingUI* Panel = EditorUIPanels[ (int32)EditorUIPanel ];
 	ShowEditorUIPanel( Panel, HandIndex, bShouldShow, OnHand, bRefreshQuickMenu );
@@ -890,8 +906,8 @@ void FVREditorUISystem::ShowEditorUIPanel( AVREditorFloatingUI* Panel, const int
 {
 	if( Panel != nullptr )
 	{
-		AVREditorFloatingUI::EDockedTo DockTo = Panel->GetDockedTo();
-		if( OnHand || DockTo == AVREditorFloatingUI::EDockedTo::Nothing || DockTo == AVREditorFloatingUI::EDockedTo::LeftHand || DockTo == AVREditorFloatingUI::EDockedTo::RightHand )
+		AVREditorFloatingUI::EDockedTo DockedTo = Panel->GetDockedTo();
+		if( OnHand || DockedTo == AVREditorFloatingUI::EDockedTo::Nothing )
 		{
 			// Hide any panels that are already shown on this hand
 			for( int32 PanelIndex = 0; PanelIndex < (int32)EEditorUIPanel::TotalCount; ++PanelIndex )
@@ -908,11 +924,14 @@ void FVREditorUISystem::ShowEditorUIPanel( AVREditorFloatingUI* Panel, const int
 					}
 				}
 			}
-
-			Owner.GetVirtualHand( HandIndex ).bHasUIInFront = bShouldShow;
-
-			const AVREditorFloatingUI::EDockedTo DockedTo = HandIndex == VREditorConstants::LeftHandIndex ? AVREditorFloatingUI::EDockedTo::LeftHand : AVREditorFloatingUI::EDockedTo::RightHand;
-			Panel->SetDockedTo( DockedTo );
+			
+			const AVREditorFloatingUI::EDockedTo NewDockedTo = HandIndex == VREditorConstants::LeftHandIndex ? AVREditorFloatingUI::EDockedTo::LeftHand : AVREditorFloatingUI::EDockedTo::RightHand;
+			Panel->SetDockedTo( NewDockedTo );
+			
+			if ( NewDockedTo == AVREditorFloatingUI::EDockedTo::LeftHand || NewDockedTo == AVREditorFloatingUI::EDockedTo::RightHand )
+			{
+				Owner.GetVirtualHand( HandIndex ).bHasUIInFront = bShouldShow;
+			}
 
 			if (bShouldShow)
 			{
@@ -924,11 +943,15 @@ void FVREditorUISystem::ShowEditorUIPanel( AVREditorFloatingUI* Panel, const int
 			Panel->SetLocalRotation( FRotator( 90.0f, 180.0f, 0.0f ) );
 		}
 
-		Panel->ShowUI( bShouldShow );
-
-		if( bRefreshQuickMenu && QuickMenuUI )
+		if( Panel->IsUIVisible() != bShouldShow )
 		{
-			QuickMenuUI->GetUserWidget<UVREditorQuickMenu>()->RefreshUI();
+			Panel->ShowUI( bShouldShow );
+			UGameplayStatics::PlaySound2D( Owner.GetWorld(), bShouldShow ? ShowUISound : HideUISound ); 
+			
+			if (bRefreshQuickMenu && QuickMenuUI)
+			{
+				QuickMenuUI->GetUserWidget<UVREditorQuickMenu>()->RefreshUI();
+			}
 		}
 	}
 }
@@ -1048,7 +1071,10 @@ AVREditorFloatingUI* FVREditorUISystem::StartDraggingDockUI( AVREditorDockableWi
 	DraggingUIOffsetTransform = UIToUIOnLaser;
 
 	DraggingUI = InitDraggingDockUI;
-	DraggingUI->SetDockedTo( AVREditorFloatingUI::EDockedTo::Nothing );
+	DraggingUI->SetDockedTo( AVREditorFloatingUI::EDockedTo::Dragging );
+
+	UGameplayStatics::PlaySound2D( Owner.GetWorld(), StartDragUISound );
+
 	return DraggingUI;
 }
 
@@ -1068,6 +1094,8 @@ void FVREditorUISystem::StopDraggingDockUI()
 	FVirtualHand& Hand = Owner.GetVirtualHand( DraggingUIHandIndex );
 	Hand.DraggingMode = EVREditorDraggingMode::Nothing;
 	DraggingUIHandIndex = INDEX_NONE;
+
+	UGameplayStatics::PlaySound2D( Owner.GetWorld(), StopDragUISound );
 }
 
 bool FVREditorUISystem::IsDraggingDockUI()
@@ -1079,26 +1107,40 @@ void FVREditorUISystem::TogglePanelsVisibility()
 {
 	bPanelVisibilityToggle = !bPanelVisibilityToggle;
 	
-	bool bOneHandUsed = false;
-	for ( AVREditorFloatingUI* Panel : EditorUIPanels )
+	UGameplayStatics::PlaySound2D( Owner.GetWorld(), bPanelVisibilityToggle ? ShowUISound : HideUISound );
+
+	bool bFirstCheckOnHand = true;
+	AVREditorFloatingUI* PanelOnHand = nullptr;
+	for( AVREditorFloatingUI* Panel : EditorUIPanels )
 	{
-		if ( Panel != nullptr && Panel->IsUIVisible() != bPanelVisibilityToggle )
+		if( Panel != nullptr && Panel->IsUIVisible() != bPanelVisibilityToggle )
 		{
 			bool bShouldSetNewVisibility = true;
 			const AVREditorFloatingUI::EDockedTo DockedTo = Panel->GetDockedTo();
-			
-			// Prevent the panel from spawning on the hand if there is already a UI on the hand
 			if( DockedTo == AVREditorFloatingUI::EDockedTo::LeftHand ||  DockedTo == AVREditorFloatingUI::EDockedTo::RightHand )
 			{
-				if( bOneHandUsed && bPanelVisibilityToggle )
+				// Prevent the panel from spawning on the hand if there is already a UI on the hand
+				if( bFirstCheckOnHand )
 				{
-					bShouldSetNewVisibility = false;
+					bFirstCheckOnHand = false;
+					for(AVREditorFloatingUI* OtherPanel : EditorUIPanels)
+					{
+						if( OtherPanel != Panel && OtherPanel->IsUIVisible() && ( OtherPanel->GetDockedTo() == AVREditorFloatingUI::EDockedTo::LeftHand || OtherPanel->GetDockedTo() == AVREditorFloatingUI::EDockedTo::RightHand ) )
+						{
+							PanelOnHand = OtherPanel;
+							break;
+						}
+					}
 				}
-				else
+
+				if( !PanelOnHand )
 				{
 					const uint32 HandIndex = Panel->GetDockedTo() == AVREditorFloatingUI::EDockedTo::LeftHand ? VREditorConstants::LeftHandIndex : VREditorConstants::RightHandIndex;
 					Owner.GetVirtualHand( HandIndex ).bHasUIInFront = bPanelVisibilityToggle;
-					bOneHandUsed = true;
+				}
+				else
+				{
+					bShouldSetNewVisibility = false;
 				}
 			}
 			else if ( DockedTo == AVREditorFloatingUI::EDockedTo::Nothing )
@@ -1133,3 +1175,22 @@ void FVREditorUISystem::OnProxyTabLaunched(TSharedPtr<SDockTab> NewTab)
 {
 	bRefocusViewport = true;
 }
+
+void FVREditorUISystem::TogglePanelVisibility( const EEditorUIPanel EditorUIPanel )
+{
+	AVREditorFloatingUI* Panel = EditorUIPanels[(int32)EditorUIPanel];
+	if (Panel != nullptr)
+	{
+		const bool bIsShowing = Panel->IsUIVisible();
+		const int32 HandIndexWithQuickMenu = QuickMenuUI->GetDockedTo() == AVREditorFloatingUI::EDockedTo::LeftArm ? VREditorConstants::LeftHandIndex : VREditorConstants::RightHandIndex;
+		if( Panel->GetDockedTo() == AVREditorFloatingUI::EDockedTo::Room )
+		{
+			Panel->ShowUI( false );
+			Panel->SetDockedTo( AVREditorFloatingUI::EDockedTo::Nothing );
+		}
+		else
+		{
+			ShowEditorUIPanel( Panel, HandIndexWithQuickMenu, !bIsShowing, true );
+		}
+	}
+}
diff --git a/Engine/Source/Editor/VREditor/VREditorUISystem.h b/Engine/Source/Editor/VREditor/VREditorUISystem.h
--- a/Engine/Source/Editor/VREditor/VREditorUISystem.h
+++ b/Engine/Source/Editor/VREditor/VREditorUISystem.h
@@ -107,6 +107,9 @@ public:
 	/** Get the minimum dock window size */
 	float GetMinDockWindowSize() const;
 
+	/** Toggles the visibility of the panel, if the panel is in room space it will be hidden and docked to nothing */
+	void TogglePanelVisibility( const EEditorUIPanel EditorUIPanel );
+
 protected:
 
 	/** Called when the user presses a button on their motion controller device */
@@ -179,5 +182,22 @@ protected:
 
 	/** Set to true when we need to refocus the viewport. */
 	bool bRefocusViewport;
+
+
+	//
+	// Sounds
+	//
+
+	/** Start dragging UI sound */
+	class USoundCue* StartDragUISound;
+
+	/** Stop dragging UI sound */
+	class USoundCue* StopDragUISound;
+
+	/** Hide UI sound */
+	class USoundCue* HideUISound;
+
+	/** Show UI sound */
+	class USoundCue* ShowUISound;
 };
 
