VR Editor: Fixed radial menu button not hovering when hovering a button, going back to the middle and then go back to that same button

diff --git a/Engine/Source/Editor/VREditor/VREditorRadialMenu.cpp b/Engine/Source/Editor/VREditor/VREditorRadialMenu.cpp
--- a/Engine/Source/Editor/VREditor/VREditorRadialMenu.cpp
+++ b/Engine/Source/Editor/VREditor/VREditorRadialMenu.cpp
@@ -101,6 +101,7 @@ void UVREditorRadialMenu::Update( const FVirtualHand& Hand )
 			{
 				CurrentItem->OnLeaveHover();
 				CurrentItem = nullptr;
+				PreviousItem = nullptr;
 			}
 		}
 	}
