VR Editor: Fixed: close dockable window using the close button with the window docked to hand was causing the laser not to appear

diff --git a/Engine/Source/Editor/VREditor/VREditorUISystem.cpp b/Engine/Source/Editor/VREditor/VREditorUISystem.cpp
--- a/Engine/Source/Editor/VREditor/VREditorUISystem.cpp
+++ b/Engine/Source/Editor/VREditor/VREditorUISystem.cpp
@@ -134,6 +134,10 @@ void FVREditorUISystem::AddReferencedObjects( FReferenceCollector& Collector )
 	Collector.AddReferencedObject( QuickMenuWidgetClass );
 	Collector.AddReferencedObject( QuickRadialWidgetClass );
 	Collector.AddReferencedObject( TutorialWidgetClass );
+	Collector.AddReferencedObject( StartDragUISound );
+	Collector.AddReferencedObject( StopDragUISound );
+	Collector.AddReferencedObject( HideUISound );
+	Collector.AddReferencedObject( ShowUISound );
 }
 
 
@@ -928,11 +932,6 @@ void FVREditorUISystem::ShowEditorUIPanel( AVREditorFloatingUI* Panel, const int
 			const AVREditorFloatingUI::EDockedTo NewDockedTo = HandIndex == VREditorConstants::LeftHandIndex ? AVREditorFloatingUI::EDockedTo::LeftHand : AVREditorFloatingUI::EDockedTo::RightHand;
 			Panel->SetDockedTo( NewDockedTo );
 			
-			if ( NewDockedTo == AVREditorFloatingUI::EDockedTo::LeftHand || NewDockedTo == AVREditorFloatingUI::EDockedTo::RightHand )
-			{
-				Owner.GetVirtualHand( HandIndex ).bHasUIInFront = bShouldShow;
-			}
-
 			if (bShouldShow)
 			{
 				Panel->SetScale( Panel->GetInitialScale() );
@@ -943,15 +942,18 @@ void FVREditorUISystem::ShowEditorUIPanel( AVREditorFloatingUI* Panel, const int
 			Panel->SetLocalRotation( FRotator( 90.0f, 180.0f, 0.0f ) );
 		}
 
-		if( Panel->IsUIVisible() != bShouldShow )
+		Panel->ShowUI( bShouldShow );
+
+		if (Panel->GetDockedTo() == AVREditorFloatingUI::EDockedTo::LeftHand || Panel->GetDockedTo() == AVREditorFloatingUI::EDockedTo::RightHand)
 		{
-			Panel->ShowUI( bShouldShow );
-			UGameplayStatics::PlaySound2D( Owner.GetWorld(), bShouldShow ? ShowUISound : HideUISound ); 
-			
-			if (bRefreshQuickMenu && QuickMenuUI)
-			{
-				QuickMenuUI->GetUserWidget<UVREditorQuickMenu>()->RefreshUI();
-			}
+			Owner.GetVirtualHand( HandIndex ).bHasUIInFront = bShouldShow;
+		}
+
+		UGameplayStatics::PlaySound2D( Owner.GetWorld(), bShouldShow ? ShowUISound : HideUISound );
+
+		if (bRefreshQuickMenu && QuickMenuUI)
+		{
+			QuickMenuUI->GetUserWidget<UVREditorQuickMenu>()->RefreshUI();
 		}
 	}
 }
@@ -1147,7 +1149,7 @@ void FVREditorUISystem::TogglePanelsVisibility()
 			{
 				bShouldSetNewVisibility = false;
 			}
-			
+
 			if( bShouldSetNewVisibility )
 			{
 				Panel->ShowUI( bPanelVisibilityToggle );
