VR Editor: Fixed materials/meshes dragged from Content Browser sometimes being dropped immediately at camera position

diff --git a/Engine/Source/Editor/VREditor/VREditorMode.cpp b/Engine/Source/Editor/VREditor/VREditorMode.cpp
--- a/Engine/Source/Editor/VREditor/VREditorMode.cpp
+++ b/Engine/Source/Editor/VREditor/VREditorMode.cpp
@@ -1318,7 +1318,10 @@ bool FVREditorMode::InputKey( FEditorViewportClient* ViewportClient, FViewport*
 					{
 						// How long since the trigger was lightly pressed?
 						const float TimeSinceLightlyPressed = (float)( FPlatformTime::Seconds() - Hand.RealTimeTriggerWasLightlyPressed );
-						if( !Hand.bIsClickingOnUI && ( !Hand.bAllowTriggerLightPressLocking || TimeSinceLightlyPressed < VREd::TriggerLightlyPressedLockTime->GetFloat() ) )
+						if( !Hand.bIsClickingOnUI &&	// @todo vreditor: UI clicks happen with light presses; we never want to convert to a hard press!
+							Hand.DraggingMode != EVREditorDraggingMode::Material &&	// @todo vreditor: Material dragging happens with light presses, don't convert to a hard press!
+							Hand.DraggingMode != EVREditorDraggingMode::ActorsAtLaserImpact &&	// @todo vreditor: Actor dragging happens with light presses, don't convert to a hard press!
+							( !Hand.bAllowTriggerLightPressLocking || TimeSinceLightlyPressed < VREd::TriggerLightlyPressedLockTime->GetFloat() ) )
 						{
 							Hand.bIsTriggerAtLeastLightlyPressed = false;
 							Hand.bHasTriggerBeenReleasedSinceLastPress = false;
