VR Editor: Greatly increased the 2D distance you must drag an asset from Content Browser before it starts to place into the world
- This makes it easier to double-click on assets without accidentially triggering a drag.

diff --git a/Engine/Source/Editor/VREditor/VREditorMode.cpp b/Engine/Source/Editor/VREditor/VREditorMode.cpp
--- a/Engine/Source/Editor/VREditor/VREditorMode.cpp
+++ b/Engine/Source/Editor/VREditor/VREditorMode.cpp
@@ -390,7 +390,7 @@ void FVREditorMode::Enter()
 			// Tell Slate to require a larger pixel distance threshold before the drag starts.  This is important for things 
 			// like Content Browser drag and drop.
 			SavedEditorState.DragTriggerDistance = FSlateApplication::Get().GetDragTriggerDistance();
-			FSlateApplication::Get().SetDragTriggerDistance( 30.0f );
+			FSlateApplication::Get().SetDragTriggerDistance( 100.0f );	// @todo vreditor tweak
 		}
 	}
 
