VR Editor: Fixed: With hiding/showing windows it is possible to have multiple windows docked to the controller (UE-27968)

diff --git a/Engine/Source/Editor/VREditor/VREditorUISystem.cpp b/Engine/Source/Editor/VREditor/VREditorUISystem.cpp
--- a/Engine/Source/Editor/VREditor/VREditorUISystem.cpp
+++ b/Engine/Source/Editor/VREditor/VREditorUISystem.cpp
@@ -1119,10 +1119,11 @@ void FVREditorUISystem::TogglePanelsVisibility()
 		{
 			bool bShouldSetNewVisibility = true;
 			const AVREditorFloatingUI::EDockedTo DockedTo = Panel->GetDockedTo();
-			if( DockedTo == AVREditorFloatingUI::EDockedTo::LeftHand ||  DockedTo == AVREditorFloatingUI::EDockedTo::RightHand )
+			if( DockedTo == AVREditorFloatingUI::EDockedTo::LeftHand ||  DockedTo == AVREditorFloatingUI::EDockedTo::RightHand  )
 			{
 				// Prevent the panel from spawning on the hand if there is already a UI on the hand
-				if( bFirstCheckOnHand )
+				// Iterate through all panels to check if there is already a panel on the hand, but only once and if we haven't found one yet
+				if( bFirstCheckOnHand && PanelOnHand == nullptr )
 				{
 					bFirstCheckOnHand = false;
 					for(AVREditorFloatingUI* OtherPanel : EditorUIPanels)
@@ -1139,6 +1140,7 @@ void FVREditorUISystem::TogglePanelsVisibility()
 				{
 					const uint32 HandIndex = Panel->GetDockedTo() == AVREditorFloatingUI::EDockedTo::LeftHand ? VREditorConstants::LeftHandIndex : VREditorConstants::RightHandIndex;
 					Owner.GetVirtualHand( HandIndex ).bHasUIInFront = bPanelVisibilityToggle;
+					PanelOnHand = Panel;
 				}
 				else
 				{
