VR Editor: When double-clicking an asset, the asset editor will now automatically appear in your hand

diff --git a/Engine/Source/Editor/VREditor/VREditorUISystem.cpp b/Engine/Source/Editor/VREditor/VREditorUISystem.cpp
--- a/Engine/Source/Editor/VREditor/VREditorUISystem.cpp
+++ b/Engine/Source/Editor/VREditor/VREditorUISystem.cpp
@@ -50,6 +50,7 @@ namespace VREd
 	static FAutoConsoleVariable MinDockDragDistance( TEXT( "VREd.MinDockDragDistance" ), 10.0f, TEXT( "Minimum amount of distance needed to drag dock from hand" ) );
 	static FAutoConsoleVariable DoubleClickTime( TEXT( "VREd.DoubleClickTime" ), 0.25f, TEXT( "Minimum duration between clicks for a double click event to fire" ) );
 	static FAutoConsoleVariable UIPressHapticFeedbackStrength( TEXT( "VREd.UIPressHapticFeedbackStrength" ), 0.4f, TEXT( "Strenth of haptic when clicking on the UI" ) );
+	static FAutoConsoleVariable UIAssetEditorSummonedOnHandHapticFeedbackStrength( TEXT( "VREd.UIAssetEditorSummonedOnHandHapticFeedbackStrength" ), 1.0f, TEXT( "Strenth of haptic to play to remind a user which hand an asset editor was spawned on" ) );
 	static FAutoConsoleVariable MaxDockWindowSize( TEXT( "VREd.MaxDockWindowSize" ), 250, TEXT( "Maximum size for dockable windows" ) );
 	static FAutoConsoleVariable MinDockWindowSize( TEXT( "VREd.MinDockWindowSize" ), 40, TEXT( "Minimum size for dockable windows" ) );
 
@@ -1178,6 +1179,23 @@ float FVREditorUISystem::GetMinDockWindowSize() const
 void FVREditorUISystem::OnProxyTabLaunched(TSharedPtr<SDockTab> NewTab)
 {
 	bRefocusViewport = true;
+
+	// A tab was opened, so make sure the "Asset" UI is visible.  That's where the user can interact
+	// with the newly-opened tab
+	if( !IsShowingEditorUIPanel( EEditorUIPanel::AssetEditor ) )
+	{
+		// Always spawn on a hand.  But which hand?  Well, we'll choose the hand that isn't actively clicking on something using a laser.
+		const int32 HandIndex = Owner.GetVirtualHand( VREditorConstants::LeftHandIndex ).bIsClickingOnUI ? VREditorConstants::RightHandIndex : VREditorConstants::LeftHandIndex;	// Hand that did not clicked with a laser
+		const bool bShouldShow = true;
+		const bool bShowOnHand = true;
+		ShowEditorUIPanel( EEditorUIPanel::AssetEditor, HandIndex, bShouldShow, bShowOnHand );
+
+		// Play haptic effect so user knows to look at their hand that now has UI on it!
+		const float Strength = VREd::UIAssetEditorSummonedOnHandHapticFeedbackStrength->GetFloat();
+		Owner.PlayHapticEffect(
+			HandIndex == VREditorConstants::LeftHandIndex ? Strength : 0.0f,
+			HandIndex == VREditorConstants::RightHandIndex ? Strength : 0.0f );
+	}
 }
 
 void FVREditorUISystem::TogglePanelVisibility( const EEditorUIPanel EditorUIPanel )
