VR Editor: Added floating help labels to Oculus Tocuh controllers when held up to your face
- Also, fixed the laser pointer overlapping your hand meshes

diff --git a/Engine/Source/Editor/VREditor/VREditorMode.cpp b/Engine/Source/Editor/VREditor/VREditorMode.cpp
--- a/Engine/Source/Editor/VREditor/VREditorMode.cpp
+++ b/Engine/Source/Editor/VREditor/VREditorMode.cpp
@@ -27,6 +27,8 @@ namespace VREd
 	static FAutoConsoleVariable AlternativeLaserPointerRotationOffset( TEXT( "VREd.LaserPointerRotationOffset" ), 90.0f, TEXT( "How much to rotate the laser pointer (pitch) when using the alternative laser pointer direction mode" ) );
 	static FAutoConsoleVariable LaserPointerLightPullBackDistance( TEXT( "VREd.LaserPointerLightPullBackDistance" ), 2.5f, TEXT( "How far to pull back our little hover light from the impact surface" ) );
 	static FAutoConsoleVariable LaserPointerLightRadius( TEXT( "VREd.LaserPointLightRadius" ), 20.0f, TEXT( "How big our hover light is" ) );
+	static FAutoConsoleVariable OculusLaserPointerStartOffset( TEXT( "VREd.OculusLaserPointerStartOffset" ), 2.8f, TEXT( "How far to offset the start of the laser pointer to avoid overlapping the hand mesh geometry (Oculus)" ) );
+	static FAutoConsoleVariable ViveLaserPointerStartOffset( TEXT( "VREd.ViveLaserPointerStartOffset" ), 2.0f, TEXT( "How far to offset the start of the laser pointer to avoid overlapping the hand mesh geometry (Vive)" ) );
 
 	static FAutoConsoleVariable UseMouseAsHandInForcedVRMode( TEXT( "VREd.UseMouseAsHandInForcedVRMode" ), 1, TEXT( "When in forced VR mode, enabling this setting uses the mouse cursor as a virtual hand instead of motion controllers" ) );
 
@@ -999,8 +1001,14 @@ void FVREditorMode::Tick( FEditorViewportClient* ViewportClient, float DeltaTime
 					LaserPointerImpactPoint = Hand.HoverLocation;
 				}
 
+				// Offset the beginning of the laser pointer a bit, so that it doesn't overlap the hand mesh
+				const float LaserPointerStartOffset =
+					GetWorldScaleFactor() *
+					( GetHMDDeviceType() == EHMDDeviceType::DT_OculusRift ? VREd::OculusLaserPointerStartOffset->GetFloat() : VREd::ViveLaserPointerStartOffset->GetFloat() );
+				Hand.LaserPointerMeshComponent->SetRelativeLocation( FVector( LaserPointerStartOffset, 0.0f, 0.0f ) );
+
 				const FVector LaserPointerDirection = ( LaserPointerImpactPoint - LaserPointerStart ).GetSafeNormal();
-				const float LaserPointerLength = FMath::Max( 0.000001f, ( LaserPointerImpactPoint - LaserPointerStart ).Size() );
+				const float LaserPointerLength = FMath::Max( 0.000001f, ( LaserPointerImpactPoint - LaserPointerStart ).Size() - LaserPointerStartOffset );
 
 				// The laser pointer needs to stay the same size relative to our tracking space, so we inverse compensate for world to meters scale here
 				float LaserPointerRadius = VREd::LaserPointerRadius->GetFloat() * GetWorldScaleFactor();
@@ -1014,7 +1022,6 @@ void FVREditorMode::Tick( FEditorViewportClient* ViewportClient, float DeltaTime
 					HoverMeshRadius *= 0.35f;	// @todo vreditor tweak
 				}
 
-
 				Hand.LaserPointerMeshComponent->SetRelativeScale3D( FVector( LaserPointerLength, LaserPointerRadius * 2.0f, LaserPointerRadius * 2.0f ) );
 
 				if( Hand.bIsHovering )
@@ -2246,7 +2253,15 @@ void FVREditorMode::UpdateHelpLabels()
 
 				UStaticMeshSocket* Socket = FindMeshSocketForKey( Hand.HandMeshComponent->StaticMesh, Key );
 				check( Socket != nullptr );
-				const FTransform SocketRelativeTransform( Socket->RelativeRotation, Socket->RelativeLocation, Socket->RelativeScale );
+				FTransform SocketRelativeTransform( Socket->RelativeRotation, Socket->RelativeLocation, Socket->RelativeScale );
+
+				// Oculus has asymmetrical controllers, so we the sock transform horizontally
+				if( HandIndex == VREditorConstants::RightHandIndex &&
+					GetHMDDeviceType() == EHMDDeviceType::DT_OculusRift )
+				{
+					const FVector Scale3D = SocketRelativeTransform.GetLocation();
+					SocketRelativeTransform.SetLocation( FVector( Scale3D.X, -Scale3D.Y, Scale3D.Z ) );
+				}
 
 				// Make sure the labels stay the same size even when the world is scaled
 				FTransform HandTransformWithWorldToMetersScaling = Hand.Transform;
