VR Editor: Fixed hovering over a radial selection while hovering on opposite controller causes perpetual sound play (UE-28047)

diff --git a/Engine/Source/Editor/VREditor/VREditorWorldInteraction.cpp b/Engine/Source/Editor/VREditor/VREditorWorldInteraction.cpp
--- a/Engine/Source/Editor/VREditor/VREditorWorldInteraction.cpp
+++ b/Engine/Source/Editor/VREditor/VREditorWorldInteraction.cpp
@@ -894,9 +894,10 @@ void FVREditorWorldInteraction::Tick( FEditorViewportClient* ViewportClient, con
 			}
 			else
 			{
-				if ( !Hand.bHasUIInFront && (Hand.bIsTrackpadPositionValid[0] && Hand.bIsTrackpadPositionValid[1] ) && Owner.GetUISystem().GetDraggingDockUIHandIndex() != HandIndex )
+				if ( !Hand.bHasUIInFront && (Hand.bIsTrackpadPositionValid[0] && Hand.bIsTrackpadPositionValid[1] ) )
 				{
-					if( Owner.GetHMDDeviceType() == EHMDDeviceType::DT_SteamVR && Hand.bIsTouchingTrackpad )
+					FVirtualHand& OtherHand = Owner.GetOtherHand( HandIndex );
+					if( Owner.GetHMDDeviceType() == EHMDDeviceType::DT_SteamVR && Hand.bIsTouchingTrackpad && !OtherHand.bIsTouchingTrackpad )
 					{
 						Owner.GetUISystem().TryToSpawnRadialMenu( HandIndex );
 					}
