VR Editor: Added selected actors snap to ground to the radial menu

diff --git a/Engine/Source/Editor/VREditor/VREditorMode.cpp b/Engine/Source/Editor/VREditor/VREditorMode.cpp
--- a/Engine/Source/Editor/VREditor/VREditorMode.cpp
+++ b/Engine/Source/Editor/VREditor/VREditorMode.cpp
@@ -17,6 +17,8 @@
 #include "IMotionController.h"
 #include "MotionControllerComponent.h"
 
+#include "Editor/LevelEditor/Public/LevelEditorActions.h"
+
 #define LOCTEXT_NAMESPACE "VREditorMode"
 
 namespace VREd
@@ -2506,6 +2508,15 @@ void FVREditorMode::Duplicate()
 }
 
 
+void FVREditorMode::SnapSelectedActorsToGround()
+{
+	FLevelEditorModule& LevelEditorModule = FModuleManager::GetModuleChecked<FLevelEditorModule>( TEXT( "LevelEditor" ) );
+	const FLevelEditorCommands& Commands = LevelEditorModule.GetLevelEditorCommands();
+	const TSharedPtr< FUICommandList >& CommandList = GetLevelViewportPossessedForVR().GetParentLevelEditor().Pin()->GetLevelEditorActions();
+	CommandList->ExecuteAction( Commands.SnapBottomCenterBoundsToFloor.ToSharedRef() );
+	WorldInteraction->SetupTransformablesForSelectedActors();
+}
+
 void FVREditorMode::SetLaserVisuals( const int32 HandIndex, const FLinearColor& NewColor, const float CrawlFade, const float CrawlSpeed )
 {
 	FVirtualHand Hand = GetVirtualHand( HandIndex );
diff --git a/Engine/Source/Editor/VREditor/VREditorMode.h b/Engine/Source/Editor/VREditor/VREditorMode.h
--- a/Engine/Source/Editor/VREditor/VREditorMode.h
+++ b/Engine/Source/Editor/VREditor/VREditorMode.h
@@ -263,6 +263,9 @@ public:
 	/** Duplicates the selected objects */
 	void Duplicate();
 
+	/** Snaps the selected objects to the ground */
+	void SnapSelectedActorsToGround();
+
 	/** Sets the visuals of the LaserPointer */
 	void SetLaserVisuals( const int32 HandIndex, const FLinearColor& NewColor, const float CrawlFade = 0.0f, const float CrawlSpeed = 0.0f );
 
diff --git a/Engine/Source/Editor/VREditor/VREditorRadialMenu.cpp b/Engine/Source/Editor/VREditor/VREditorRadialMenu.cpp
--- a/Engine/Source/Editor/VREditor/VREditorRadialMenu.cpp
+++ b/Engine/Source/Editor/VREditor/VREditorRadialMenu.cpp
@@ -10,6 +10,8 @@
 #include "VREditorTransformGizmo.h"
 #include "VirtualHand.h"
 
+#include "Editor/LevelEditor/Public/LevelEditor.h"
+
 #define LOCTEXT_NAMESPACE "VREditorRadialMenu"
 
 namespace VREd
@@ -24,6 +26,7 @@ UVREditorRadialMenu::UVREditorRadialMenu( const FObjectInitializer& ObjectInitia
 	TopRightItem(nullptr),
 	RightItem(nullptr),
 	BottomRightItem(nullptr),
+	BottomItem(nullptr),
 	LeftBottomItem(nullptr),
 	LeftItem(nullptr),
 	TopLeftItem(nullptr),
@@ -66,6 +69,10 @@ void UVREditorRadialMenu::Update( const FVirtualHand& Hand )
 			{
 				CurrentItem = BottomRightItem;
 			}
+			else if ( Angle > 157.5f && Angle <= 202.5f)
+			{
+				CurrentItem = BottomItem;
+			}
 			else if (Angle > 202.5f && Angle <= 247.5f)
 			{
 				CurrentItem = LeftBottomItem;
@@ -107,7 +114,8 @@ void UVREditorRadialMenu::Update( const FVirtualHand& Hand )
 	}
 }
 
-void UVREditorRadialMenu::SetButtons(UVREditorRadialMenuItem* InitTopItem, UVREditorRadialMenuItem* InitTopRightItem, UVREditorRadialMenuItem* InitRightItem, UVREditorRadialMenuItem* InitBottomRightItem, UVREditorRadialMenuItem* InitLeftBottomItem, UVREditorRadialMenuItem* InitLeftItem, UVREditorRadialMenuItem* InitTopLeftItem)
+void UVREditorRadialMenu::SetButtons( UVREditorRadialMenuItem* InitTopItem, UVREditorRadialMenuItem* InitTopRightItem, UVREditorRadialMenuItem* InitRightItem, UVREditorRadialMenuItem* InitBottomRightItem,
+	UVREditorRadialMenuItem* InitBottomItem, UVREditorRadialMenuItem* InitLeftBottomItem, UVREditorRadialMenuItem* InitLeftItem, UVREditorRadialMenuItem* InitTopLeftItem )
 {
 	TopItem = InitTopItem;
 	TopItem->OnPressedDelegate.AddUObject( this, &UVREditorRadialMenu::OnGizmoCycle );
@@ -125,6 +133,10 @@ void UVREditorRadialMenu::SetButtons(UVREditorRadialMenuItem* InitTopItem, UVREd
 	BottomRightItem->OnPressedDelegate.AddUObject(this, &UVREditorRadialMenu::OnRedoButtonClicked);
 	BottomRightItem->SetLabel(FString("Redo"));
 
+	BottomItem = InitBottomItem;
+	BottomItem->OnPressedDelegate.AddUObject( this, &UVREditorRadialMenu::OnSnapActorsToGroundClicked );
+	BottomItem->SetLabel( FString( "Snap To Ground" ) );
+
 	LeftBottomItem = InitLeftBottomItem;
 	LeftBottomItem->OnPressedDelegate.AddUObject(this, &UVREditorRadialMenu::OnUndoButtonClicked);
 	LeftBottomItem->SetLabel(FString("Undo"));
@@ -235,4 +247,10 @@ void UVREditorRadialMenu::OnPasteButtonClicked()
 	GetOwner()->GetOwner().GetOwner().Paste();
 }
 
+void UVREditorRadialMenu::OnSnapActorsToGroundClicked()
+{
+	GetOwner()->GetOwner().GetOwner().SnapSelectedActorsToGround();
+}
+
+
 #undef LOCTEXT_NAMESPACE
diff --git a/Engine/Source/Editor/VREditor/VREditorRadialMenu.h b/Engine/Source/Editor/VREditor/VREditorRadialMenu.h
--- a/Engine/Source/Editor/VREditor/VREditorRadialMenu.h
+++ b/Engine/Source/Editor/VREditor/VREditorRadialMenu.h
@@ -25,8 +25,8 @@ public:
 
 	/** To bind the buttons from UMG */
 	UFUNCTION(BlueprintCallable, Category = "RadialItems")
-	void SetButtons(UVREditorRadialMenuItem* InitTopItem, UVREditorRadialMenuItem* InitTopRightItem, UVREditorRadialMenuItem* InitRightItem, UVREditorRadialMenuItem* InitBottomRightItem,
-		UVREditorRadialMenuItem* InitLeftBottomItem, UVREditorRadialMenuItem* InitLeftItem, UVREditorRadialMenuItem* InitTopLeftItem);
+	void SetButtons( UVREditorRadialMenuItem* InitTopItem, UVREditorRadialMenuItem* InitTopRightItem, UVREditorRadialMenuItem* InitRightItem, UVREditorRadialMenuItem* InitBottomRightItem,
+		UVREditorRadialMenuItem* InitBottomItem, UVREditorRadialMenuItem* InitLeftBottomItem, UVREditorRadialMenuItem* InitLeftItem, UVREditorRadialMenuItem* InitTopLeftItem );
 
 	/** Resets the current item to not hovering */
 	void ResetItem();
@@ -56,6 +56,8 @@ private:
 	UPROPERTY()
 	UVREditorRadialMenuItem* BottomRightItem;
 	UPROPERTY()
+	UVREditorRadialMenuItem* BottomItem;
+	UPROPERTY()
 	UVREditorRadialMenuItem* LeftBottomItem;
 	UPROPERTY()
 	UVREditorRadialMenuItem* LeftItem;
@@ -81,4 +83,5 @@ private:
 	void OnRedoButtonClicked();
 	void OnCopyButtonClicked();
 	void OnPasteButtonClicked();
+	void OnSnapActorsToGroundClicked();
 };
\ No newline at end of file
