VREditor - Notifying when a tab is focused.

diff --git a/Engine/Source/Editor/VREditor/VREditorUISystem.cpp b/Engine/Source/Editor/VREditor/VREditorUISystem.cpp
--- a/Engine/Source/Editor/VREditor/VREditorUISystem.cpp
+++ b/Engine/Source/Editor/VREditor/VREditorUISystem.cpp
@@ -808,7 +808,7 @@ void FVREditorUISystem::CreateUIs()
 			const bool bWithSceneComponent = false;
 			AVREditorDockableWindow* TabManagerUI = GetOwner().SpawnTransientSceneActor< AVREditorDockableWindow >(TEXT("AssetEditor"), bWithSceneComponent);
 			TabManagerUI->SetSlateWidget(*this, SNullWidget::NullWidget, Resolution, VREd::EditorUISize->GetFloat() * 2.0f, AVREditorFloatingUI::EDockedTo::Nothing);
-			TabManagerUI->ShowUI( false );
+			TabManagerUI->ShowUI( true );
 
 			// @todo vreditor: Could use "Hovering" instead for better performance with many UIs, but needs to be manually refreshed in some cases
 			TabManagerUI->GetWidgetComponent()->SetDrawingPolicy(EVREditorWidgetDrawingPolicy::Always);
@@ -822,7 +822,9 @@ void FVREditorUISystem::CreateUIs()
 			TSharedPtr<SWindow> TabManagerWindow = TabManagerUI->GetWidgetComponent()->GetSlateWindow();
 			TSharedRef<SWindow> TabManagerWindowRef = TabManagerWindow.ToSharedRef();
 			ProxyTabManager = MakeShareable(new FProxyTabmanager(TabManagerWindowRef));
-			ProxyTabManager->OnProxyTabLaunched.Add(FOnProxyTabLaunched::FDelegate::CreateRaw(this, &FVREditorUISystem::OnProxyTabLaunched));
+
+			ProxyTabManager->OnTabOpened.Add(FOnTabEvent::FDelegate::CreateRaw(this, &FVREditorUISystem::OnProxyTabLaunched));
+			ProxyTabManager->OnAttentionDrawnToTab.Add(FOnTabEvent::FDelegate::CreateRaw(this, &FVREditorUISystem::OnAttentionDrawnToTab));
 
 			// We're going to start stealing tabs from the global tab manager inserting them into the world instead.
 			FGlobalTabmanager::Get()->SetProxyTabManager(ProxyTabManager);
@@ -1177,24 +1179,34 @@ float FVREditorUISystem::GetMinDockWindowSize() const
 }
 
 void FVREditorUISystem::OnProxyTabLaunched(TSharedPtr<SDockTab> NewTab)
+{
+	ShowAssetEditor();
+}
+
+void FVREditorUISystem::OnAttentionDrawnToTab(TSharedPtr<SDockTab> NewTab)
+{
+	ShowAssetEditor();
+}
+
+void FVREditorUISystem::ShowAssetEditor()
 {
 	bRefocusViewport = true;
 
 	// A tab was opened, so make sure the "Asset" UI is visible.  That's where the user can interact
 	// with the newly-opened tab
-	if( !IsShowingEditorUIPanel( EEditorUIPanel::AssetEditor ) )
+	if ( !IsShowingEditorUIPanel(EEditorUIPanel::AssetEditor) )
 	{
 		// Always spawn on a hand.  But which hand?  Well, we'll choose the hand that isn't actively clicking on something using a laser.
-		const int32 HandIndex = Owner.GetVirtualHand( VREditorConstants::LeftHandIndex ).bIsClickingOnUI ? VREditorConstants::RightHandIndex : VREditorConstants::LeftHandIndex;	// Hand that did not clicked with a laser
+		const int32 HandIndex = Owner.GetVirtualHand(VREditorConstants::LeftHandIndex).bIsClickingOnUI ? VREditorConstants::RightHandIndex : VREditorConstants::LeftHandIndex;	// Hand that did not clicked with a laser
 		const bool bShouldShow = true;
 		const bool bShowOnHand = true;
-		ShowEditorUIPanel( EEditorUIPanel::AssetEditor, HandIndex, bShouldShow, bShowOnHand );
+		ShowEditorUIPanel(EEditorUIPanel::AssetEditor, HandIndex, bShouldShow, bShowOnHand);
 
 		// Play haptic effect so user knows to look at their hand that now has UI on it!
 		const float Strength = VREd::UIAssetEditorSummonedOnHandHapticFeedbackStrength->GetFloat();
 		Owner.PlayHapticEffect(
 			HandIndex == VREditorConstants::LeftHandIndex ? Strength : 0.0f,
-			HandIndex == VREditorConstants::RightHandIndex ? Strength : 0.0f );
+			HandIndex == VREditorConstants::RightHandIndex ? Strength : 0.0f);
 	}
 }
 
@@ -1216,3 +1228,4 @@ void FVREditorUISystem::TogglePanelVisibility( const EEditorUIPanel EditorUIPane
 		}
 	}
 }
+
diff --git a/Engine/Source/Editor/VREditor/VREditorUISystem.h b/Engine/Source/Editor/VREditor/VREditorUISystem.h
--- a/Engine/Source/Editor/VREditor/VREditorUISystem.h
+++ b/Engine/Source/Editor/VREditor/VREditorUISystem.h
@@ -127,6 +127,12 @@ protected:
 	 */
 	void OnProxyTabLaunched(TSharedPtr<SDockTab> NewTab);
 
+	/** Called when the injected proxy tab manager reports needing to draw attention to a tab. */
+	void OnAttentionDrawnToTab(TSharedPtr<SDockTab> NewTab);
+
+	/** Show the asset editor if it's not currently visible. */
+	void ShowAssetEditor();
+
 protected:
 
 	/** Owning object */
diff --git a/Engine/Source/Runtime/Slate/Private/Framework/Docking/TabManager.cpp b/Engine/Source/Runtime/Slate/Private/Framework/Docking/TabManager.cpp
--- a/Engine/Source/Runtime/Slate/Private/Framework/Docking/TabManager.cpp
+++ b/Engine/Source/Runtime/Slate/Private/Framework/Docking/TabManager.cpp
@@ -1788,6 +1788,12 @@ void FGlobalTabmanager::DrawAttentionToTabManager( const TSharedRef<FTabManager>
 	if ( Tab.IsValid() )
 	{
 		this->DrawAttention(Tab.ToSharedRef());
+
+		// #HACK VREDITOR
+		if ( ProxyTabManager.IsValid() )
+		{
+			ProxyTabManager->DrawAttention(Tab.ToSharedRef());
+		}
 	}
 }
 
@@ -2029,7 +2035,16 @@ void FProxyTabmanager::OpenUnmanagedTab(FName PlaceholderId, const FSearchPrefer
 	NewlyOpenedTab->GetParent()->GetParentDockTabStack()->OpenTab(UnmanagedTab);
 	NewlyOpenedTab->RequestCloseTab();
 
-	OnProxyTabLaunched.Broadcast(UnmanagedTab);
+	MainNonCloseableTab = UnmanagedTab;
+
+	OnTabOpened.Broadcast(UnmanagedTab);
+}
+
+void FProxyTabmanager::DrawAttention(const TSharedRef<SDockTab>& TabToHighlight)
+{
+	FTabManager::DrawAttention(TabToHighlight);
+
+	OnAttentionDrawnToTab.Broadcast(TabToHighlight);
 }
 
 
diff --git a/Engine/Source/Runtime/Slate/Public/Framework/Docking/TabManager.h b/Engine/Source/Runtime/Slate/Public/Framework/Docking/TabManager.h
--- a/Engine/Source/Runtime/Slate/Public/Framework/Docking/TabManager.h
+++ b/Engine/Source/Runtime/Slate/Public/Framework/Docking/TabManager.h
@@ -607,7 +607,7 @@ class SLATE_API FTabManager : public TSharedFromThis<FTabManager>
 
 		void PopulateTabSpawnerMenu( FMenuBuilder &PopulateMe, const FName& TabType );
 
-		void DrawAttention( const TSharedRef<SDockTab>& TabToHighlight );
+		virtual void DrawAttention( const TSharedRef<SDockTab>& TabToHighlight );
 
 		struct ESearchPreference
 		{
@@ -992,7 +992,7 @@ private:
 
 //#HACK VREDITOR - Had to introduce the proxy tab manager to steal asset tabs.
 
-DECLARE_MULTICAST_DELEGATE_OneParam(FOnProxyTabLaunched, TSharedPtr<SDockTab>);
+DECLARE_MULTICAST_DELEGATE_OneParam(FOnTabEvent, TSharedPtr<SDockTab>);
 
 
 class SLATE_API FProxyTabmanager : public FTabManager
@@ -1003,12 +1003,15 @@ public:
 		: FTabManager(TSharedPtr<SDockTab>(), MakeShareable(new FTabSpawner()))
 		, ParentWindow(InParentWindow)
 	{
+		bCanDoDragOperation = false;
 	}
 
 	virtual void OpenUnmanagedTab(FName PlaceholderId, const FSearchPreference& SearchPreference, const TSharedRef<SDockTab>& UnmanagedTab) override;
+	virtual void DrawAttention(const TSharedRef<SDockTab>& TabToHighlight) override;
 
 public:
-	FOnProxyTabLaunched OnProxyTabLaunched;
+	FOnTabEvent OnTabOpened;
+	FOnTabEvent OnAttentionDrawnToTab;
 
 private:
 	TWeakPtr<SWindow> ParentWindow;
