VR Editor:
- Fixed Scaling UI Windows with thumbpad causes radial menu to appear (UE-28060)
- Fixed Trackpad scaling speed for UI windows

diff --git a/Engine/Source/Editor/VREditor/VREditorUISystem.cpp b/Engine/Source/Editor/VREditor/VREditorUISystem.cpp
--- a/Engine/Source/Editor/VREditor/VREditorUISystem.cpp
+++ b/Engine/Source/Editor/VREditor/VREditorUISystem.cpp
@@ -1001,6 +1001,7 @@ void FVREditorUISystem::TryToSpawnRadialMenu( const int32 HandIndex )
 		Hand.DraggingMode != EVREditorDraggingMode::ActorsFreely &&
 		Hand.DraggingMode != EVREditorDraggingMode::ActorsWithGizmo &&
 		Hand.DraggingMode != EVREditorDraggingMode::AssistingDrag &&
+		DraggingUIHandIndex != HandIndex &&
 		!Hand.bIsHoveringOverUI;	// Don't show radial menu when aiming at a UI  (too much clutter)
 
 	UVREditorRadialMenu* RadialMenu = QuickRadialMenu->GetUserWidget<UVREditorRadialMenu>();
diff --git a/Engine/Source/Editor/VREditor/VREditorWorldInteraction.cpp b/Engine/Source/Editor/VREditor/VREditorWorldInteraction.cpp
--- a/Engine/Source/Editor/VREditor/VREditorWorldInteraction.cpp
+++ b/Engine/Source/Editor/VREditor/VREditorWorldInteraction.cpp
@@ -829,18 +829,7 @@ void FVREditorWorldInteraction::Tick( FEditorViewportClient* ViewportClient, con
 				if( Hand.bIsTrackpadPositionValid[ 1 ] )
 				{
 					const bool bIsAbsolute = ( Owner.GetHMDDeviceType() == EHMDDeviceType::DT_SteamVR );
-					float SlideDelta = 0.0f;
-					if( Hand.bIsTouchingTrackpad || !bIsAbsolute )
-					{
-						if( bIsAbsolute )
-						{
-							SlideDelta = ( ( Hand.TrackpadPosition.Y - Hand.LastTrackpadPosition.Y ) * VREd::TrackpadAbsoluteDragSpeed->GetFloat() ) * WorldScaleFactor;
-						}
-						else
-						{
-							SlideDelta = ( Hand.TrackpadPosition.Y * VREd::TrackpadRelativeDragSpeed->GetFloat() ) * WorldScaleFactor;
-						}
-					}
+					float SlideDelta = GetTrackpadSlideDelta( HandIndex ) * WorldScaleFactor;
 
 					if( !FMath::IsNearlyZero( SlideDelta ) )
 					{
@@ -1219,18 +1208,7 @@ void FVREditorWorldInteraction::Tick( FEditorViewportClient* ViewportClient, con
 				{
 					//@todo: VR Editor - This should be in the UI system
 					const bool bIsAbsolute = ( Owner.GetHMDDeviceType() == EHMDDeviceType::DT_SteamVR );
-					float SlideDeltaY = 0.0f;
-					if (Hand.bIsTouchingTrackpad || !bIsAbsolute)
-					{
-						if (bIsAbsolute)
-						{
-							SlideDeltaY = ((Hand.TrackpadPosition.Y - Hand.LastTrackpadPosition.Y) * VREd::TrackpadAbsoluteDragSpeed->GetFloat()) * WorldScaleFactor;
-						}
-						else
-						{
-							SlideDeltaY = (Hand.TrackpadPosition.Y * VREd::TrackpadRelativeDragSpeed->GetFloat()) * WorldScaleFactor;
-						}
-					}
+					float SlideDeltaY = GetTrackpadSlideDelta( HandIndex );
 
 					float NewUIScale = DraggingUI->GetScale() + SlideDeltaY;
 					if( NewUIScale <= Owner.GetUISystem().GetMinDockWindowSize() )
@@ -2668,6 +2646,26 @@ void FVREditorWorldInteraction::SetupTransformablesForSelectedActors()
 }
 
 
+float FVREditorWorldInteraction::GetTrackpadSlideDelta( const int32 HandIndex )
+{
+	const bool bIsAbsolute = (Owner.GetHMDDeviceType() == EHMDDeviceType::DT_SteamVR);
+	float SlideDelta = 0.0f;
+	FVirtualHand& Hand = Owner.GetVirtualHand( HandIndex );
+	if (Hand.bIsTouchingTrackpad || !bIsAbsolute)
+	{
+		if (bIsAbsolute)
+		{
+			SlideDelta = ((Hand.TrackpadPosition.Y - Hand.LastTrackpadPosition.Y) * VREd::TrackpadAbsoluteDragSpeed->GetFloat());
+		}
+		else
+		{
+			SlideDelta = (Hand.TrackpadPosition.Y * VREd::TrackpadRelativeDragSpeed->GetFloat());
+		}
+	}
+
+	return SlideDelta;
+}
+
 bool FVREditorWorldInteraction::IsTransformingActor( AActor* Actor ) const
 {
 	bool bFoundActor = false;
diff --git a/Engine/Source/Editor/VREditor/VREditorWorldInteraction.h b/Engine/Source/Editor/VREditor/VREditorWorldInteraction.h
--- a/Engine/Source/Editor/VREditor/VREditorWorldInteraction.h
+++ b/Engine/Source/Editor/VREditor/VREditorWorldInteraction.h
@@ -49,6 +49,9 @@ public:
 	/** Creates a Transformable for each selected actor, wiping out the existing transformables */
 	void SetupTransformablesForSelectedActors();
 
+	/** Gets the Y delta of the trackpad */
+	float GetTrackpadSlideDelta( const int32 HandIndex );
+
 protected:
 
 	/** Called when the user presses a button on their motion controller device */
