VR Editor: Added a code comment (no effective change)

diff --git a/Engine/Source/Editor/VREditor/VREditorMode.cpp b/Engine/Source/Editor/VREditor/VREditorMode.cpp
--- a/Engine/Source/Editor/VREditor/VREditorMode.cpp
+++ b/Engine/Source/Editor/VREditor/VREditorMode.cpp
@@ -2513,7 +2513,10 @@ void FVREditorMode::SnapSelectedActorsToGround()
 	FLevelEditorModule& LevelEditorModule = FModuleManager::GetModuleChecked<FLevelEditorModule>( TEXT( "LevelEditor" ) );
 	const FLevelEditorCommands& Commands = LevelEditorModule.GetLevelEditorCommands();
 	const TSharedPtr< FUICommandList >& CommandList = GetLevelViewportPossessedForVR().GetParentLevelEditor().Pin()->GetLevelEditorActions();
+
 	CommandList->ExecuteAction( Commands.SnapBottomCenterBoundsToFloor.ToSharedRef() );
+	
+	// @todo vreditor: This should not be needed after to allow transformables to stop animating the transforms of actors after they come to rest
 	WorldInteraction->SetupTransformablesForSelectedActors();
 }
 
