VR Editor: Oculus fix for hovering over a radial selection while hovering on opposite controller causes perpetual sound play (UE-28047).

diff --git a/Engine/Source/Editor/VREditor/VREditorUISystem.cpp b/Engine/Source/Editor/VREditor/VREditorUISystem.cpp
--- a/Engine/Source/Editor/VREditor/VREditorUISystem.cpp
+++ b/Engine/Source/Editor/VREditor/VREditorUISystem.cpp
@@ -576,9 +576,9 @@ void FVREditorUISystem::Tick( FEditorViewportClient* ViewportClient, const float
 		for( int32 HandIndex = 0; HandIndex < VREditorConstants::NumVirtualHands; ++HandIndex )
 		{
 			const FVirtualHand& Hand = Owner.GetVirtualHand( HandIndex );
-
+			const FVirtualHand& OtherHand = Owner.GetOtherHand( HandIndex );
 			bool bShouldShowRadialMenu = false;
-			if( !Hand.bHasUIInFront && Hand.bIsTrackpadPositionValid[ 0 ] && Hand.bIsTrackpadPositionValid[ 1 ] && GetDraggingDockUIHandIndex() != HandIndex )
+			if( !Hand.bHasUIInFront && Hand.bIsTrackpadPositionValid[ 0 ] && Hand.bIsTrackpadPositionValid[ 1 ] && GetDraggingDockUIHandIndex() != HandIndex && !OtherHand.bIsTouchingTrackpad )
 			{
 				if( Hand.TrackpadPosition.Size() > 0.15f )	// @todo vreditor tweak
 				{
