VR Editor: Fixed bug when trying to use radial menu with analog sticks on both Oculus Touch controllers at the same time

diff --git a/Engine/Source/Editor/VREditor/VREditorUISystem.cpp b/Engine/Source/Editor/VREditor/VREditorUISystem.cpp
--- a/Engine/Source/Editor/VREditor/VREditorUISystem.cpp
+++ b/Engine/Source/Editor/VREditor/VREditorUISystem.cpp
@@ -578,12 +578,20 @@ void FVREditorUISystem::Tick( FEditorViewportClient* ViewportClient, const float
 			const FVirtualHand& Hand = Owner.GetVirtualHand( HandIndex );
 			const FVirtualHand& OtherHand = Owner.GetOtherHand( HandIndex );
 			bool bShouldShowRadialMenu = false;
-			if( !Hand.bHasUIInFront && Hand.bIsTrackpadPositionValid[ 0 ] && Hand.bIsTrackpadPositionValid[ 1 ] && GetDraggingDockUIHandIndex() != HandIndex && !OtherHand.bIsTouchingTrackpad )
+
+			const float MinJoystickOffsetBeforeRadialMenu = 0.15f;	// @todo vreditor twea
+
+			const bool bOtherHandAlreadyHasRadialMenu =
+				IsShowingRadialMenu( Owner.GetOtherHandIndex( HandIndex ) );
+
+			if( !Hand.bHasUIInFront &&
+				Hand.bIsTrackpadPositionValid[ 0 ] &&
+				Hand.bIsTrackpadPositionValid[ 1 ] &&
+				Hand.TrackpadPosition.Size() > MinJoystickOffsetBeforeRadialMenu &&
+				GetDraggingDockUIHandIndex() != HandIndex &&
+				!bOtherHandAlreadyHasRadialMenu )
 			{
-				if( Hand.TrackpadPosition.Size() > 0.15f )	// @todo vreditor tweak
-				{
-					bShouldShowRadialMenu = true;
-				}
+				bShouldShowRadialMenu = true;
 			}
 
 			if( bShouldShowRadialMenu )
