VR Editor: Don't set initial offset/rotation for Asset Editor UI (as it always starts out in your hand)

diff --git a/Engine/Source/Editor/VREditor/VREditorUISystem.cpp b/Engine/Source/Editor/VREditor/VREditorUISystem.cpp
--- a/Engine/Source/Editor/VREditor/VREditorUISystem.cpp
+++ b/Engine/Source/Editor/VREditor/VREditorUISystem.cpp
@@ -821,8 +821,6 @@ void FVREditorUISystem::CreateUIs()
 			// @todo vreditor: Could use "Hovering" instead for better performance with many UIs, but needs to be manually refreshed in some cases
 			TabManagerUI->GetWidgetComponent()->SetDrawingPolicy(EVREditorWidgetDrawingPolicy::Always);
 
-			TabManagerUI->SetRelativeOffset( FVector( 200.0f, 0.0f, 150.0f ) );
-			TabManagerUI->SetLocalRotation( FRotator( 0.0f, 180.0f, 0.0f ) );
 			FloatingUIs.Add(TabManagerUI);
 
 			EditorUIPanels[ (int32)EEditorUIPanel::AssetEditor ] = TabManagerUI;
