VREditor - Moving input processing into a slate input processor so that we no longer rely on the viewport to have focus to correctly use the motion controllers to drive the user's input.  This makes them function a lot more like a mouse would, where focus doesn't matter, your mouse always works.

diff --git a/Engine/Source/Editor/VREditor/VREditorInputProcessor.cpp b/Engine/Source/Editor/VREditor/VREditorInputProcessor.cpp
new file mode 100644
--- /dev/null
+++ b/Engine/Source/Editor/VREditor/VREditorInputProcessor.cpp
@@ -0,0 +1,47 @@
+// Copyright 1998-2015 Epic Games, Inc. All Rights Reserved.
+
+#include "VREditorModule.h"
+#include "VREditorInputProcessor.h"
+#include "VREditorMode.h"
+#include "SLevelViewport.h"
+
+FVREditorInputProcessor::FVREditorInputProcessor(FVREditorMode* InVRMode)
+	: VRMode(InVRMode)
+{
+}
+
+FVREditorInputProcessor::~FVREditorInputProcessor()
+{
+}
+
+void FVREditorInputProcessor::Tick(const float DeltaTime, FSlateApplication& SlateApp, TSharedRef<ICursor> Cursor)
+{
+}
+
+bool FVREditorInputProcessor::HandleKeyDownEvent(FSlateApplication& SlateApp, const FKeyEvent& InKeyEvent)
+{
+	FEditorViewportClient* ViewportClient = VRMode->GetLevelViewportPossessedForVR().GetViewportClient().Get();
+	FViewport* ActiveViewport = VRMode->GetLevelViewportPossessedForVR().GetActiveViewport();
+	return VRMode->HandleInputKey(ViewportClient, ActiveViewport, InKeyEvent.GetKey(), IE_Pressed);
+}
+
+bool FVREditorInputProcessor::HandleKeyUpEvent(FSlateApplication& SlateApp, const FKeyEvent& InKeyEvent)
+{
+	FEditorViewportClient* ViewportClient = VRMode->GetLevelViewportPossessedForVR().GetViewportClient().Get();
+	FViewport* ActiveViewport = VRMode->GetLevelViewportPossessedForVR().GetActiveViewport();
+	return VRMode->HandleInputKey(ViewportClient, ActiveViewport, InKeyEvent.GetKey(), IE_Released);
+}
+
+bool FVREditorInputProcessor::HandleAnalogInputEvent(FSlateApplication& SlateApp, const FAnalogInputEvent& InAnalogInputEvent)
+{
+	FEditorViewportClient* ViewportClient = VRMode->GetLevelViewportPossessedForVR().GetViewportClient().Get();
+	FViewport* ActiveViewport = VRMode->GetLevelViewportPossessedForVR().GetActiveViewport();
+	return VRMode->HandleInputAxis(ViewportClient, ActiveViewport, InAnalogInputEvent.GetUserIndex(), InAnalogInputEvent.GetKey(), InAnalogInputEvent.GetAnalogValue(), FApp::GetDeltaTime());
+}
+
+bool FVREditorInputProcessor::HandleMouseMoveEvent(FSlateApplication& SlateApp, const FPointerEvent& MouseEvent)
+{
+	FEditorViewportClient* ViewportClient = VRMode->GetLevelViewportPossessedForVR().GetViewportClient().Get();
+
+	return false;
+}
diff --git a/Engine/Source/Editor/VREditor/VREditorInputProcessor.h b/Engine/Source/Editor/VREditor/VREditorInputProcessor.h
new file mode 100644
--- /dev/null
+++ b/Engine/Source/Editor/VREditor/VREditorInputProcessor.h
@@ -0,0 +1,24 @@
+// Copyright 1998-2015 Epic Games, Inc. All Rights Reserved.
+
+#pragma once
+
+#include "Framework/Application/IInputProcessor.h"
+
+class FVREditorMode;
+
+class FVREditorInputProcessor : public IInputProcessor
+{
+public:
+	FVREditorInputProcessor(FVREditorMode* InVRMode);
+	virtual ~FVREditorInputProcessor();
+
+	virtual void Tick(const float DeltaTime, FSlateApplication& SlateApp, TSharedRef<ICursor> Cursor) override;
+
+	virtual bool HandleKeyDownEvent(FSlateApplication& SlateApp, const FKeyEvent& InKeyEvent) override;
+	virtual bool HandleKeyUpEvent(FSlateApplication& SlateApp, const FKeyEvent& InKeyEvent) override;
+	virtual bool HandleAnalogInputEvent(FSlateApplication& SlateApp, const FAnalogInputEvent& InAnalogInputEvent) override;
+	virtual bool HandleMouseMoveEvent(FSlateApplication& SlateApp, const FPointerEvent& MouseEvent) override;
+
+private:
+	FVREditorMode* VRMode;
+};
diff --git a/Engine/Source/Editor/VREditor/VREditorMode.cpp b/Engine/Source/Editor/VREditor/VREditorMode.cpp
--- a/Engine/Source/Editor/VREditor/VREditorMode.cpp
+++ b/Engine/Source/Editor/VREditor/VREditorMode.cpp
@@ -131,6 +131,9 @@ FVREditorMode::~FVREditorMode()
 
 void FVREditorMode::Enter()
 {
+	InputProcessor = MakeShareable(new FVREditorInputProcessor(this));
+	FSlateApplication::Get().SetInputPreProcessor(true, InputProcessor);
+
 	bIsFullyInitialized = false;
 	bWantsToExitMode = false;
 
@@ -431,6 +434,8 @@ void FVREditorMode::Exit()
 {
 	bIsFullyInitialized = false;
 
+	FSlateApplication::Get().SetInputPreProcessor(false);
+
 	// Kill subsystems
 	{
 		WorldInteraction.Reset();
@@ -1297,14 +1302,19 @@ void FVREditorMode::Tick( FEditorViewportClient* ViewportClient, float DeltaTime
 }
 
 
-bool FVREditorMode::InputKey( FEditorViewportClient* ViewportClient, FViewport* Viewport, FKey Key, EInputEvent Event )
+bool FVREditorMode::InputKey(FEditorViewportClient* ViewportClient, FViewport* Viewport, FKey Key, EInputEvent Event)
 {
 	// Only if this is our VR viewport. Remember, editor modes get a chance to tick and receive input for each active viewport.
-	if( ViewportClient != GetLevelViewportPossessedForVR().GetViewportClient().Get() )
+	if ( ViewportClient != GetLevelViewportPossessedForVR().GetViewportClient().Get() )
 	{
-		return InputKey( GetLevelViewportPossessedForVR().GetViewportClient().Get(), Viewport, Key, Event );
+		return InputKey(GetLevelViewportPossessedForVR().GetViewportClient().Get(), Viewport, Key, Event);
 	}
 
+	return FEdMode::InputKey(ViewportClient, Viewport, Key, Event);
+}
+
+bool FVREditorMode::HandleInputKey(FEditorViewportClient* ViewportClient, FViewport* Viewport, FKey Key, EInputEvent Event)
+{
 	bool bHandled = false;
 
 	FVRAction* KnownAction = KeyToActionMap.Find( Key );
@@ -1338,7 +1348,7 @@ bool FVREditorMode::InputKey( FEditorViewportClient* ViewportClient, FViewport*
 							// Synthesize an input key for releasing the light press
 							// NOTE: Intentionally re-entrant call here.
 							const EInputEvent InputEvent = IE_Released;
-							const bool bWasLightReleaseHandled = this->InputKey( ViewportClient, Viewport, VRAction.HandIndex == VREditorConstants::LeftHandIndex ? VREditorKeyNames::MotionController_Left_LightlyPressedTriggerAxis : VREditorKeyNames::MotionController_Right_LightlyPressedTriggerAxis, InputEvent );
+							const bool bWasLightReleaseHandled = HandleInputKey(ViewportClient, Viewport, VRAction.HandIndex == VREditorConstants::LeftHandIndex ? VREditorKeyNames::MotionController_Left_LightlyPressedTriggerAxis : VREditorKeyNames::MotionController_Right_LightlyPressedTriggerAxis, InputEvent);
 						}
 						else
 						{
@@ -1448,18 +1458,23 @@ bool FVREditorMode::InputKey( FEditorViewportClient* ViewportClient, FViewport*
 
 	StopOldHapticEffects();
 
-	return bHandled ? true : FEdMode::InputKey( ViewportClient, Viewport, Key, Event );
+	return bHandled;
 }
 
 
-bool FVREditorMode::InputAxis( FEditorViewportClient* ViewportClient, FViewport* Viewport, int32 ControllerId, FKey Key, float Delta, float DeltaTime )
+bool FVREditorMode::InputAxis(FEditorViewportClient* ViewportClient, FViewport* Viewport, int32 ControllerId, FKey Key, float Delta, float DeltaTime)
 {
 	// Only if this is our VR viewport. Remember, editor modes get a chance to tick and receive input for each active viewport.
-	if( ViewportClient != GetLevelViewportPossessedForVR().GetViewportClient().Get() )
+	if ( ViewportClient != GetLevelViewportPossessedForVR().GetViewportClient().Get() )
 	{
-		return InputAxis( GetLevelViewportPossessedForVR().GetViewportClient().Get(), Viewport, ControllerId, Key, Delta, DeltaTime );
+		return InputAxis(GetLevelViewportPossessedForVR().GetViewportClient().Get(), Viewport, ControllerId, Key, Delta, DeltaTime);
 	}
 
+	return FEdMode::InputAxis(ViewportClient, Viewport, ControllerId, Key, Delta, DeltaTime);
+}
+
+bool FVREditorMode::HandleInputAxis(FEditorViewportClient* ViewportClient, FViewport* Viewport, int32 ControllerId, FKey Key, float Delta, float DeltaTime)
+{
 	bool bHandled = false;
 
 	FVRAction* KnownAction = KeyToActionMap.Find( Key );
@@ -1487,7 +1502,7 @@ bool FVREditorMode::InputAxis( FEditorViewportClient* ViewportClient, FViewport*
 
 						// Synthesize an input key for this light press
 						const EInputEvent InputEvent = IE_Pressed;
-						const bool bWasLightPressHandled = this->InputKey( ViewportClient, Viewport, VRAction.HandIndex == VREditorConstants::LeftHandIndex ? VREditorKeyNames::MotionController_Left_LightlyPressedTriggerAxis : VREditorKeyNames::MotionController_Right_LightlyPressedTriggerAxis, InputEvent );
+						const bool bWasLightPressHandled = HandleInputKey(ViewportClient, Viewport, VRAction.HandIndex == VREditorConstants::LeftHandIndex ? VREditorKeyNames::MotionController_Left_LightlyPressedTriggerAxis : VREditorKeyNames::MotionController_Right_LightlyPressedTriggerAxis, InputEvent);
 					}
 					else if( Hand.bIsTriggerAtLeastLightlyPressed && Delta < VREd::TriggerLightlyPressedThreshold->GetFloat() )
 					{
@@ -1495,7 +1510,7 @@ bool FVREditorMode::InputAxis( FEditorViewportClient* ViewportClient, FViewport*
 
 						// Synthesize an input key for this light press
 						const EInputEvent InputEvent = IE_Released;
-						const bool bWasLightReleaseHandled = this->InputKey( ViewportClient, Viewport, VRAction.HandIndex == VREditorConstants::LeftHandIndex ? VREditorKeyNames::MotionController_Left_LightlyPressedTriggerAxis : VREditorKeyNames::MotionController_Right_LightlyPressedTriggerAxis, InputEvent );
+						const bool bWasLightReleaseHandled = HandleInputKey(ViewportClient, Viewport, VRAction.HandIndex == VREditorConstants::LeftHandIndex ? VREditorKeyNames::MotionController_Left_LightlyPressedTriggerAxis : VREditorKeyNames::MotionController_Right_LightlyPressedTriggerAxis, InputEvent);
 					}
 				}
 
@@ -1513,7 +1528,7 @@ bool FVREditorMode::InputAxis( FEditorViewportClient* ViewportClient, FViewport*
 
 						// Synthesize an input key for this full press
 						const EInputEvent InputEvent = IE_Pressed;
-						this->InputKey( ViewportClient, Viewport, VRAction.HandIndex == VREditorConstants::LeftHandIndex ? VREditorKeyNames::MotionController_Left_FullyPressedTriggerAxis : VREditorKeyNames::MotionController_Right_FullyPressedTriggerAxis, InputEvent );
+						HandleInputKey(ViewportClient, Viewport, VRAction.HandIndex == VREditorConstants::LeftHandIndex ? VREditorKeyNames::MotionController_Left_FullyPressedTriggerAxis : VREditorKeyNames::MotionController_Right_FullyPressedTriggerAxis, InputEvent);
 					}
 					else if( Hand.bIsTriggerFullyPressed && Delta < VREd::TriggerFullyPressedReleaseThreshold->GetFloat() )
 					{
@@ -1521,7 +1536,7 @@ bool FVREditorMode::InputAxis( FEditorViewportClient* ViewportClient, FViewport*
 
 						// Synthesize an input key for this full press
 						const EInputEvent InputEvent = IE_Released;
-						this->InputKey( ViewportClient, Viewport, VRAction.HandIndex == VREditorConstants::LeftHandIndex ? VREditorKeyNames::MotionController_Left_FullyPressedTriggerAxis : VREditorKeyNames::MotionController_Right_FullyPressedTriggerAxis, InputEvent );
+						HandleInputKey(ViewportClient, Viewport, VRAction.HandIndex == VREditorConstants::LeftHandIndex ? VREditorKeyNames::MotionController_Left_FullyPressedTriggerAxis : VREditorKeyNames::MotionController_Right_FullyPressedTriggerAxis, InputEvent);
 					}
 				}
 			}
@@ -1552,23 +1567,7 @@ bool FVREditorMode::InputAxis( FEditorViewportClient* ViewportClient, FViewport*
 	// ...
 	StopOldHapticEffects();
 
-	return bHandled ? true : FEdMode::InputAxis( ViewportClient, Viewport, ControllerId, Key, Delta, DeltaTime );
-}
-
-bool FVREditorMode::InputDelta( FEditorViewportClient* ViewportClient, FViewport* Viewport, FVector& InDrag, FRotator& InRot, FVector& InScale )
-{
-	// Only if this is our VR viewport. Remember, editor modes get a chance to tick and receive input for each active viewport.
-	if( ViewportClient != GetLevelViewportPossessedForVR().GetViewportClient().Get() )
-	{
-		return InputDelta( GetLevelViewportPossessedForVR().GetViewportClient().Get(), Viewport, InDrag, InRot, InScale );
-	}
-
-	bool bHandled = false;
-
-	// ... 
-	StopOldHapticEffects();
-
-	return bHandled ? true : FEdMode::InputDelta( ViewportClient,Viewport,InDrag, InRot, InScale);
+	return bHandled;
 }
 
 bool FVREditorMode::IsCompatibleWith(FEditorModeID OtherModeID) const
diff --git a/Engine/Source/Editor/VREditor/VREditorMode.h b/Engine/Source/Editor/VREditor/VREditorMode.h
--- a/Engine/Source/Editor/VREditor/VREditorMode.h
+++ b/Engine/Source/Editor/VREditor/VREditorMode.h
@@ -6,6 +6,7 @@
 #include "IVREditorModule.h"
 #include "VirtualHand.h"
 #include "HeadMountedDisplayTypes.h"	// For EHMDDeviceType::Type
+#include "VREditorInputProcessor.h"
 
 
 namespace VREd
@@ -66,11 +67,13 @@ public:
 	virtual void Tick( FEditorViewportClient* ViewportClient, float DeltaTime ) override;
 	virtual bool InputKey( FEditorViewportClient* ViewportClient, FViewport* Viewport, FKey Key, EInputEvent Event ) override;
 	virtual bool InputAxis( FEditorViewportClient* ViewportClient, FViewport* Viewport, int32 ControllerId, FKey Key, float Delta, float DeltaTime ) override;
-	virtual bool InputDelta(FEditorViewportClient* InViewportClient,FViewport* InViewport,FVector& InDrag,FRotator& InRot,FVector& InScale) override;
 	virtual bool IsCompatibleWith( FEditorModeID OtherModeID ) const override;
 	virtual void AddReferencedObjects( FReferenceCollector& Collector ) override;
 	virtual void Render( const FSceneView* SceneView, FViewport* Viewport, FPrimitiveDrawInterface* PDI ) override;	
 
+	bool HandleInputKey(FEditorViewportClient* ViewportClient, FViewport* Viewport, FKey Key, EInputEvent Event);
+	bool HandleInputAxis(FEditorViewportClient* ViewportClient, FViewport* Viewport, int32 ControllerId, FKey Key, float Delta, float DeltaTime);
+
 	// IVREditorMode interface
 	virtual bool GetLaserPointer( const int32 HandIndex, FVector& LaserPointerStart, FVector& LaserPointerEnd, const bool bEvenIfUIIsInFront = false, const float LaserLengthOverride = 0.0f) const override;
 	virtual bool IsInputCaptured( const FVRAction VRAction ) const;
@@ -263,9 +266,6 @@ public:
 	/** Duplicates the selected objects */
 	void Duplicate();
 
-	/** Snaps the selected objects to the ground */
-	void SnapSelectedActorsToGround();
-
 	/** Sets the visuals of the LaserPointer */
 	void SetLaserVisuals( const int32 HandIndex, const FLinearColor& NewColor, const float CrawlFade = 0.0f, const float CrawlSpeed = 0.0f );
 
@@ -520,6 +520,9 @@ protected:
 	/** The current Gizmo type that is used for the TransformGizmo Actor */
 	EGizmoHandleTypes CurrentGizmoType;
 
+	// Slate Input Processor
+	TSharedPtr<FVREditorInputProcessor> InputProcessor;
+
 	//
 	// Colors
 	//
