VR Editor: Fixed compile error

diff --git a/Engine/Source/Editor/VREditor/VREditorMode.h b/Engine/Source/Editor/VREditor/VREditorMode.h
--- a/Engine/Source/Editor/VREditor/VREditorMode.h
+++ b/Engine/Source/Editor/VREditor/VREditorMode.h
@@ -266,6 +266,9 @@ public:
 	/** Duplicates the selected objects */
 	void Duplicate();
 
+	/** Snaps the selected objects to the ground */
+	void SnapSelectedActorsToGround();
+
 	/** Sets the visuals of the LaserPointer */
 	void SetLaserVisuals( const int32 HandIndex, const FLinearColor& NewColor, const float CrawlFade = 0.0f, const float CrawlSpeed = 0.0f );
 
