VR Editor: Fixed button presses being spammed every frame (repeats not routed properly)

diff --git a/Engine/Source/Editor/VREditor/VREditorInputProcessor.cpp b/Engine/Source/Editor/VREditor/VREditorInputProcessor.cpp
--- a/Engine/Source/Editor/VREditor/VREditorInputProcessor.cpp
+++ b/Engine/Source/Editor/VREditor/VREditorInputProcessor.cpp
@@ -22,7 +22,7 @@ bool FVREditorInputProcessor::HandleKeyDownEvent(FSlateApplication& SlateApp, co
 {
 	FEditorViewportClient* ViewportClient = VRMode->GetLevelViewportPossessedForVR().GetViewportClient().Get();
 	FViewport* ActiveViewport = VRMode->GetLevelViewportPossessedForVR().GetActiveViewport();
-	return VRMode->HandleInputKey(ViewportClient, ActiveViewport, InKeyEvent.GetKey(), IE_Pressed);
+	return VRMode->HandleInputKey(ViewportClient, ActiveViewport, InKeyEvent.GetKey(), InKeyEvent.IsRepeat() ? IE_Repeat : IE_Pressed);
 }
 
 bool FVREditorInputProcessor::HandleKeyUpEvent(FSlateApplication& SlateApp, const FKeyEvent& InKeyEvent)
