VR Editor: Don't place Tutorial in the world in VREditorDemo mode
- This fixes it reappearing as soon as you press down on the trackpad

diff --git a/Engine/Source/Editor/VREditor/VREditorUISystem.cpp b/Engine/Source/Editor/VREditor/VREditorUISystem.cpp
--- a/Engine/Source/Editor/VREditor/VREditorUISystem.cpp
+++ b/Engine/Source/Editor/VREditor/VREditorUISystem.cpp
@@ -634,6 +634,8 @@ void FVREditorUISystem::CreateUIs()
 {
 	const FIntPoint DefaultResolution( VREd::EditorUIResolutionX->GetInt(), VREd::EditorUIResolutionY->GetInt() );
 
+	const bool bIsVREditorDemo = FParse::Param( FCommandLine::Get(), TEXT( "VREditorDemo" ) );	// @todo vreditor: Remove this when no longer needed
+																								
 	// @todo vreditor: Tweak
 	if ( QuickMenuWidgetClass != nullptr )
 	{
@@ -684,7 +686,6 @@ void FVREditorUISystem::CreateUIs()
 			Config.ThumbnailLabel = EThumbnailLabel::NoLabel;
 			Config.ThumbnailScale = 0.4f;
 
-			const bool bIsVREditorDemo = FParse::Param( FCommandLine::Get(), TEXT( "VREditorDemo" ) );	// @todo vreditor: Remove this when no longer needed
 			if( bIsVREditorDemo )
 			{
 				Config.bUsePathPicker = false;
@@ -797,13 +798,17 @@ void FVREditorUISystem::CreateUIs()
 
 		// Create the tutorial dockable window
 		{
+			const bool bShowAtStart = !bIsVREditorDemo;
+
 			AVREditorFloatingUI* TutorialUI = GetOwner().SpawnTransientSceneActor< AVREditorDockableWindow >( TEXT( "TutorialUI" ), false );
-			TutorialUI->SetUMGWidget( *this, TutorialWidgetClass, FIntPoint( VREd::TutorialUIResolutionX->GetFloat(), VREd::TutorialUIResolutionY->GetFloat() ), VREd::TutorialUISize->GetFloat(), AVREditorFloatingUI::EDockedTo::Room );
-			TutorialUI->SetRelativeOffset( FVector( VREd::TutorialUILocationX->GetFloat(), VREd::TutorialUILocationY->GetFloat(), VREd::TutorialUILocationZ->GetFloat() ) );
-			TutorialUI->SetLocalRotation( FRotator( VREd::TutorialUIPitch->GetFloat(), VREd::TutorialUIYaw->GetFloat(), 0 ) );
+			TutorialUI->SetUMGWidget( *this, TutorialWidgetClass, FIntPoint( VREd::TutorialUIResolutionX->GetFloat(), VREd::TutorialUIResolutionY->GetFloat() ), VREd::TutorialUISize->GetFloat(), bShowAtStart ? AVREditorFloatingUI::EDockedTo::Room : AVREditorFloatingUI::EDockedTo::Nothing );
+			if( bShowAtStart )
+			{
+				TutorialUI->SetRelativeOffset( FVector( VREd::TutorialUILocationX->GetFloat(), VREd::TutorialUILocationY->GetFloat(), VREd::TutorialUILocationZ->GetFloat() ) );
+				TutorialUI->SetLocalRotation( FRotator( VREd::TutorialUIPitch->GetFloat(), VREd::TutorialUIYaw->GetFloat(), 0 ) );
+			}
 			
-			const bool bIsVREditorDemo = FParse::Param( FCommandLine::Get(), TEXT( "VREditorDemo" ) );
-			TutorialUI->ShowUI( !bIsVREditorDemo );
+			TutorialUI->ShowUI( bShowAtStart );
 
 			FloatingUIs.Add( TutorialUI );
 
