VR Editor: Improved ergonomics for Vive and Oculus
- Laser pointer is now orientated at a more comfortable direction on Vive
- UI docked on hands is now oriented 45 degrees toward the viewer to make it easier to interact with
- Removed "Exit VR" from Quick Menu as no one ever uses that
- Cleaned up lots of sketchy math in UI docking code

diff --git a/Engine/Source/Editor/VREditor/VREditorDockableWindow.cpp b/Engine/Source/Editor/VREditor/VREditorDockableWindow.cpp
--- a/Engine/Source/Editor/VREditor/VREditorDockableWindow.cpp
+++ b/Engine/Source/Editor/VREditor/VREditorDockableWindow.cpp
@@ -279,7 +279,7 @@ void AVREditorDockableWindow::UpdateRelativeRoomTransform()
 	const FTransform WindowToWorldTransform = GetActorTransform();
 	const FTransform WindowToRoomTransform = WindowToWorldTransform * WorldToRoom;
 
-	const FVector RoomSpaceWindowLocation = WindowToRoomTransform.GetLocation() / GetOwner().GetOwner().GetWorldScaleFactor();
+	const FVector RoomSpaceWindowLocation = WindowToRoomTransform.GetLocation();
 	const FQuat RoomSpaceWindowRotation = WindowToRoomTransform.GetRotation();
 
 	RelativeOffset = RoomSpaceWindowLocation;
diff --git a/Engine/Source/Editor/VREditor/VREditorFloatingUI.cpp b/Engine/Source/Editor/VREditor/VREditorFloatingUI.cpp
--- a/Engine/Source/Editor/VREditor/VREditorFloatingUI.cpp
+++ b/Engine/Source/Editor/VREditor/VREditorFloatingUI.cpp
@@ -14,6 +14,7 @@ namespace VREd
 {
 	static FAutoConsoleVariable UIFadeSpeed( TEXT( "VREd.UIFadeSpeed" ), 6.0f, TEXT( "How fast UI should fade in and out" ) );
 	static FAutoConsoleVariable UIBrightness( TEXT( "VREd.UIBrightness" ), 0.3f, TEXT( "How bright the UI should be" ) );
+	static FAutoConsoleVariable UIOnHandRotationOffset( TEXT( "VREd.UIOnHandRotationOffset" ), 45.0f, TEXT( "Rotation offset for UI that's docked to your hand, to make it more comfortable to hold" ) );
 }
 
 
@@ -28,7 +29,6 @@ AVREditorFloatingUI::AVREditorFloatingUI()
 	  WidgetComponent( nullptr ),
 	  bShouldBeVisible(),
 	  FadeAlpha( 1.0f ),
-	  bRotateToHead( false ),
 	  LocalRotation( FRotator( 90.0f, 180.0f, 0.0f ) ),
 	  bCollisionOnShowUI( true ),
 	  FadeDelay( 0.0f )
@@ -232,12 +232,12 @@ void AVREditorFloatingUI::UpdateTransformIfDocked()
 		
 		if( DockedTo == EDockedTo::LeftHand )
 		{
-			const bool bOnArm = true;
+			const bool bOnArm = false;
 			NewTransform = MakeUITransformLockedToHand( VREditorConstants::LeftHandIndex, bOnArm );
 		}
 		else if( DockedTo == EDockedTo::RightHand )
 		{
-			const bool bOnArm = true;
+			const bool bOnArm = false;
 			NewTransform = MakeUITransformLockedToHand( VREditorConstants::RightHandIndex, bOnArm );
 		}
 		else if( DockedTo == EDockedTo::LeftArm )
@@ -254,10 +254,6 @@ void AVREditorFloatingUI::UpdateTransformIfDocked()
 		{
 			NewTransform = MakeUITransformLockedToRoom();
 		}
-		else if (DockedTo == EDockedTo::Head)
-		{
-			NewTransform = MakeUITransformLockedToHead();
-		}
 		else if( DockedTo == EDockedTo::Custom )
 		{
 			NewTransform = MakeCustomUITransform();
@@ -267,11 +263,6 @@ void AVREditorFloatingUI::UpdateTransformIfDocked()
 			check( 0 );
 		}
 
-		if( bRotateToHead )
-		{
-			RotateToHead( NewTransform );
-		}
-
 		SetTransform( NewTransform );
 	}
 }
@@ -309,11 +300,6 @@ void AVREditorFloatingUI::SetTransform( const FTransform& Transform )
 }
 
 
-void AVREditorFloatingUI::SetRotateToHead( const bool bInitRotateToHead )
-{
-	bRotateToHead = bInitRotateToHead;
-}
-
 void AVREditorFloatingUI::SetRelativeOffset( const FVector& InRelativeOffset )
 {
 	RelativeOffset = InRelativeOffset;
@@ -338,88 +324,32 @@ FTransform AVREditorFloatingUI::MakeUITransformLockedToHand( const int32 HandInd
 {
 	const float WorldScaleFactor = Owner->GetOwner().GetWorldScaleFactor();
 
-	const FTransform& HandTransform = Owner->GetOwner().GetVirtualHand( HandIndex ).Transform;
-
-	FVector UILocation;
-	FRotator UIRotation;
-	FTransform OnArmTransform;
-	
-	FTransform LocalTransform( LocalRotation.Quaternion(), FVector::ZeroVector );
-
-	if( bOnArm )
+	FTransform UIToHandTransform( LocalRotation, RelativeOffset );
+	if( !bOnArm )
 	{
-		OnArmTransform = LocalTransform * HandTransform;
-
-		const FVector HandForward = HandTransform.GetUnitAxis( EAxis::X );	
-		const FVector HandRight = HandTransform.GetUnitAxis( EAxis::Y );
-		const FVector HandUp = HandTransform.GetUnitAxis( EAxis::Z );
-
-		UILocation =
-			OnArmTransform.GetLocation() +
-			HandForward * RelativeOffset.X * WorldScaleFactor +
-			HandRight * RelativeOffset.Y * WorldScaleFactor +
-			HandUp * RelativeOffset.Z * WorldScaleFactor;
+		UIToHandTransform *= FTransform( FRotator( VREd::UIOnHandRotationOffset->GetFloat(), 0.0f, 0.0f ).Quaternion(), FVector::ZeroVector );
 	}
 
-	if ( !bRotateToHead )
-	{
-		UIRotation = ( OnArmTransform.GetRotation() ).Rotator();
-	}
+	const FTransform HandToWorldTransform = Owner->GetOwner().GetVirtualHand( HandIndex ).Transform;
 
-	const FTransform UITransform =
-		FTransform(
-			UIRotation,
-			UILocation,
-			FVector( Scale * WorldScaleFactor ) );
+	FTransform UIToWorldTransform = UIToHandTransform * HandToWorldTransform;
+	UIToWorldTransform.SetScale3D( FVector( Scale * WorldScaleFactor ) );
 
-	return UITransform;
-}
-
-void AVREditorFloatingUI::RotateToHead( FTransform& Transform )
-{
-	FRotator UIRotation;
-	const FVector HeadToUIDirection = ( Owner->GetOwner().GetHeadTransform().GetLocation() - Transform.GetLocation() ).GetSafeNormal();
-	UIRotation = HeadToUIDirection.ToOrientationRotator();
-	UIRotation.Roll = 0.0f;
-
-	Transform.SetRotation( FQuat( UIRotation ) );
+	return UIToWorldTransform;
 }
 
 FTransform AVREditorFloatingUI::MakeUITransformLockedToRoom()
 {
-	FVector UILocation;
 	const float WorldScaleFactor = Owner->GetOwner().GetWorldScaleFactor();
-	const FTransform RoomTransform = Owner->GetOwner().GetRoomTransform();
-	
-	const FVector RoomForward = RoomTransform.GetUnitAxis( EAxis::X );
-	const FVector RoomRight = RoomTransform.GetUnitAxis( EAxis::Y );
-	const FVector RoomUp = RoomTransform.GetUnitAxis( EAxis::Z );
-
-	UILocation =
-		RoomTransform.GetLocation() +
-		RoomForward * RelativeOffset.X * WorldScaleFactor +
-		RoomRight * RelativeOffset.Y * WorldScaleFactor +
-		RoomUp * RelativeOffset.Z * WorldScaleFactor;
-
-	return FTransform(
-		LocalRotation + RoomTransform.GetRotation().Rotator(),
-		UILocation,
-		FVector( Scale * WorldScaleFactor ) );
-}
 
-FTransform AVREditorFloatingUI::MakeUITransformLockedToHead()
-{
-	FTransform Result;
-	const float WorldScaleFactor = Owner->GetOwner().GetWorldScaleFactor();
+	const FTransform UIToRoomTransform( LocalRotation, RelativeOffset );
 
-	FVector HeadWorldLocation = Owner->GetOwner().GetHeadTransform().GetLocation();
-	FVector RoomWorldLocation = Owner->GetOwner().GetRoomTransform().GetLocation();
-	FVector NewLocation = FVector( HeadWorldLocation.X, HeadWorldLocation.Y, RoomWorldLocation.Z );
-	
-	Result.SetLocation( NewLocation + ( RelativeOffset * WorldScaleFactor ) );
-	Result.SetScale3D( FVector( Scale * WorldScaleFactor ) );
-	
-	return Result;
+	const FTransform RoomToWorldTransform = Owner->GetOwner().GetRoomTransform();
+
+	FTransform UIToWorldTransform = UIToRoomTransform * RoomToWorldTransform;
+	UIToWorldTransform.SetScale3D( FVector( Scale * WorldScaleFactor ) );
+
+	return UIToWorldTransform;
 }
 
 void AVREditorFloatingUI::SetDockedTo( const EDockedTo NewDockedTo )
diff --git a/Engine/Source/Editor/VREditor/VREditorFloatingUI.h b/Engine/Source/Editor/VREditor/VREditorFloatingUI.h
--- a/Engine/Source/Editor/VREditor/VREditorFloatingUI.h
+++ b/Engine/Source/Editor/VREditor/VREditorFloatingUI.h
@@ -28,7 +28,6 @@ public:
 		LeftArm,
 		RightArm,
 		Room,
-		Head,
 		Custom,
 		Dragging,
 	};
@@ -109,9 +108,6 @@ public:
 		return CastChecked<T>( UserWidget );
 	}
 
-	/** Sets if the UI should rotate towards head */
-	void SetRotateToHead( const bool bInitRotateToHead );
-
 	/** Sets the relative offset of the UI */
 	void SetRelativeOffset( const FVector& InRelativeOffset );
 
@@ -146,14 +142,9 @@ private:
 	/** Given a hand to lock to, returns a transform to place UI at that hand's location and orientation */
 	FTransform MakeUITransformLockedToHand( const int32 HandIndex, const bool bOnArm);
 
-	/** Rotates the UI towards the head */
-	void RotateToHead( FTransform& Transform );
-
 	/** Create a room transform with the RelativeOffset and LocalRotation */
 	FTransform MakeUITransformLockedToRoom();
 
-	/** Create a transform with the RelativeOffset and LocalRotation that moves with the head */
-	FTransform MakeUITransformLockedToHead();
 
 protected:
 
diff --git a/Engine/Source/Editor/VREditor/VREditorMode.cpp b/Engine/Source/Editor/VREditor/VREditorMode.cpp
--- a/Engine/Source/Editor/VREditor/VREditorMode.cpp
+++ b/Engine/Source/Editor/VREditor/VREditorMode.cpp
@@ -26,11 +26,12 @@ namespace VREd
 	static FAutoConsoleVariable LaserPointerMaxLength( TEXT( "VREd.LaserPointerMaxLength" ), 10000.0f, TEXT( "Maximum length of the laser pointer line" ) );
 	static FAutoConsoleVariable LaserPointerRadius( TEXT( "VREd.LaserPointerRadius" ), 0.5f, TEXT( "Radius of the laser pointer line" ) );
 	static FAutoConsoleVariable LaserPointerHoverBallRadius( TEXT( "VREd.LaserPointerHoverBallRadius" ), 1.5f, TEXT( "Radius of the visual cue for a hovered object along the laser pointer ray" ) );
-	static FAutoConsoleVariable AlternativeLaserPointerRotationOffset( TEXT( "VREd.LaserPointerRotationOffset" ), 90.0f, TEXT( "How much to rotate the laser pointer (pitch) when using the alternative laser pointer direction mode" ) );
 	static FAutoConsoleVariable LaserPointerLightPullBackDistance( TEXT( "VREd.LaserPointerLightPullBackDistance" ), 2.5f, TEXT( "How far to pull back our little hover light from the impact surface" ) );
 	static FAutoConsoleVariable LaserPointerLightRadius( TEXT( "VREd.LaserPointLightRadius" ), 20.0f, TEXT( "How big our hover light is" ) );
+	static FAutoConsoleVariable OculusLaserPointerRotationOffset( TEXT( "VREd.OculusLaserPointerRotationOffset" ), 0.0f, TEXT( "How much to rotate the laser pointer (pitch) relative to the forward vector of the controller (Oculus)" ) );
+	static FAutoConsoleVariable ViveLaserPointerRotationOffset( TEXT( "VREd.ViveLaserPointerRotationOffset" ), -57.8f, TEXT( "How much to rotate the laser pointer (pitch) relative to the forward vector of the controller (Vive)" ) );
 	static FAutoConsoleVariable OculusLaserPointerStartOffset( TEXT( "VREd.OculusLaserPointerStartOffset" ), 2.8f, TEXT( "How far to offset the start of the laser pointer to avoid overlapping the hand mesh geometry (Oculus)" ) );
-	static FAutoConsoleVariable ViveLaserPointerStartOffset( TEXT( "VREd.ViveLaserPointerStartOffset" ), 2.0f, TEXT( "How far to offset the start of the laser pointer to avoid overlapping the hand mesh geometry (Vive)" ) );
+	static FAutoConsoleVariable ViveLaserPointerStartOffset( TEXT( "VREd.ViveLaserPointerStartOffset" ), 8.5f, TEXT( "How far to offset the start of the laser pointer to avoid overlapping the hand mesh geometry (Vive)" ) );
 
 	static FAutoConsoleVariable UseMouseAsHandInForcedVRMode( TEXT( "VREd.UseMouseAsHandInForcedVRMode" ), 1, TEXT( "When in forced VR mode, enabling this setting uses the mouse cursor as a virtual hand instead of motion controllers" ) );
 
@@ -109,7 +110,6 @@ FVREditorMode::FVREditorMode()
 	  CurrentScaleProgressMeshComponent( nullptr ),
 	  UserScaleIndicatorText( nullptr ),
 	  CurrentGizmoType( EGizmoHandleTypes::All ),
-	  UseAlternativeLaserPointerDirection( TEXT( "VREd.UseAlternativeLaserPointerDirection" ), 0, TEXT( "When enabled, the laser pointer direction will be rotated to make interactions more comfortable when seated" ) ),
 	  UISystem(),
 	  WorldInteraction()
 {
@@ -1008,15 +1008,21 @@ void FVREditorMode::Tick( FEditorViewportClient* ViewportClient, float DeltaTime
 					LaserPointerImpactPoint = Hand.HoverLocation;
 				}
 
+				// Apply rotation offset to the laser direction
+				const float LaserPointerRotationOffset = GetHMDDeviceType() == EHMDDeviceType::DT_OculusRift ? VREd::OculusLaserPointerRotationOffset->GetFloat() : VREd::ViveLaserPointerRotationOffset->GetFloat();
+				Hand.LaserPointerMeshComponent->SetRelativeRotation( FRotator( LaserPointerRotationOffset, 0.0f, 0.0f ) );
+
+				const FVector LaserPointerDirection = ( LaserPointerImpactPoint - LaserPointerStart ).GetSafeNormal();
+
 				// Offset the beginning of the laser pointer a bit, so that it doesn't overlap the hand mesh
 				const float LaserPointerStartOffset =
 					GetWorldScaleFactor() *
 					( GetHMDDeviceType() == EHMDDeviceType::DT_OculusRift ? VREd::OculusLaserPointerStartOffset->GetFloat() : VREd::ViveLaserPointerStartOffset->GetFloat() );
-				Hand.LaserPointerMeshComponent->SetRelativeLocation( FVector( LaserPointerStartOffset, 0.0f, 0.0f ) );
 
-				const FVector LaserPointerDirection = ( LaserPointerImpactPoint - LaserPointerStart ).GetSafeNormal();
 				const float LaserPointerLength = FMath::Max( 0.000001f, ( LaserPointerImpactPoint - LaserPointerStart ).Size() - LaserPointerStartOffset );
 
+				Hand.LaserPointerMeshComponent->SetRelativeLocation( FRotator( LaserPointerRotationOffset, 0.0f, 0.0f ).RotateVector( FVector( LaserPointerStartOffset, 0.0f, 0.0f ) ) );
+
 				// The laser pointer needs to stay the same size relative to our tracking space, so we inverse compensate for world to meters scale here
 				float LaserPointerRadius = VREd::LaserPointerRadius->GetFloat() * GetWorldScaleFactor();
 				float HoverMeshRadius = VREd::LaserPointerHoverBallRadius->GetFloat() * GetWorldScaleFactor();
@@ -1665,14 +1671,8 @@ bool FVREditorMode::GetHandTransformAndForwardVector( int32 HandIndex, FTransfor
 	{
 		OutHandTransform = Hand.Transform;
 
-		if( UseAlternativeLaserPointerDirection->GetInt() > 0 )
-		{
-			OutForwardVector = OutHandTransform.GetRotation().RotateVector( FRotator( VREd::AlternativeLaserPointerRotationOffset->GetFloat(), 0.0f, 0.0f ).RotateVector( FVector( 1.0f, 0.0f, 0.0f ) ) );
-		}
-		else
-		{
-			OutForwardVector = OutHandTransform.GetRotation().RotateVector( FVector( 1.0f, 0.0f, 0.0f ) );
-		}
+		const float LaserPointerRotationOffset = GetHMDDeviceType() == EHMDDeviceType::DT_OculusRift ? VREd::OculusLaserPointerRotationOffset->GetFloat() : VREd::ViveLaserPointerRotationOffset->GetFloat();
+		OutForwardVector = OutHandTransform.GetRotation().RotateVector( FRotator( LaserPointerRotationOffset, 0.0f, 0.0f ).RotateVector( FVector( 1.0f, 0.0f, 0.0f ) ) );
 
 		return true;
 	}
diff --git a/Engine/Source/Editor/VREditor/VREditorMode.h b/Engine/Source/Editor/VREditor/VREditorMode.h
--- a/Engine/Source/Editor/VREditor/VREditorMode.h
+++ b/Engine/Source/Editor/VREditor/VREditorMode.h
@@ -472,15 +472,6 @@ protected:
 	/** Post process for drawing VR-specific post effects */
 	class UPostProcessComponent* PostProcessComponent;
 
-
-	//
-	// 3D interactions
-	//
-
-	/** Toggles motion controller laser sights between two orientations, one which may be better-suited for standing up vs. sitting down */
-	FAutoConsoleVariable UseAlternativeLaserPointerDirection;
-
-
 	//
 	// Input
 	//
diff --git a/Engine/Source/Editor/VREditor/VREditorUISystem.cpp b/Engine/Source/Editor/VREditor/VREditorUISystem.cpp
--- a/Engine/Source/Editor/VREditor/VREditorUISystem.cpp
+++ b/Engine/Source/Editor/VREditor/VREditorUISystem.cpp
@@ -37,7 +37,7 @@ namespace VREd
 	static FAutoConsoleVariable EditorUIResolutionX( TEXT( "VREd.EditorUIResolutionX" ), 1024, TEXT( "Horizontal resolution to use for VR editor UI render targets" ) );
 	static FAutoConsoleVariable EditorUIResolutionY( TEXT( "VREd.EditorUIResolutionY" ), 1024, TEXT( "Vertical resolution to use for VR editor UI render targets" ) );
 	static FAutoConsoleVariable QuickMenuUIResolutionX( TEXT( "VREd.QuickMenuUIResolutionX" ), 1024, TEXT( "Horizontal resolution to use for Quick Menu VR UI render targets" ) );
-	static FAutoConsoleVariable QuickMenuUIResolutionY( TEXT( "VREd.QuickMenuUIResolutionY" ), 1130, TEXT( "Vertical resolution to use for Quick Menu VR UI render targets" ) );
+	static FAutoConsoleVariable QuickMenuUIResolutionY( TEXT( "VREd.QuickMenuUIResolutionY" ), 900, TEXT( "Vertical resolution to use for Quick Menu VR UI render targets" ) );
 	static FAutoConsoleVariable ContentBrowserUISize( TEXT( "VREd.ContentBrowserUISize" ), 60.0f, TEXT( "How big content browser UIs should be" ) );
 	static FAutoConsoleVariable EditorUISize( TEXT( "VREd.EditorUISize" ), 50.0f, TEXT( "How big editor UIs should be" ) );
 	static FAutoConsoleVariable ContentBrowserUIScale( TEXT( "VREd.ContentBrowserUIScale" ), 2.0f, TEXT( "How much to scale up (or down) the content browser for VR" ) );
@@ -532,7 +532,7 @@ void FVREditorUISystem::Tick( FEditorViewportClient* ViewportClient, const float
 
 			// @todo vreditor tweak: Weird to hard code this here.  Probably should be an accessor on the hand itself, and based on the actual device type
 			const FTransform UICapsuleTransform = OtherHand.Transform;
-			const FVector UICapsuleStart = FVector( -8.0f, 0.0f, 0.0f ) * WorldScaleFactor;
+			const FVector UICapsuleStart = FVector( -9.0f, 0.0f, 0.0f ) * WorldScaleFactor;
 			const FVector UICapsuleEnd = FVector( -18.0f, 0.0f, 0.0f ) * WorldScaleFactor;
 			const float UICapsuleLocalRadius = 6.0f * WorldScaleFactor;
 			const float MinDistanceToUICapsule = 8.0f * WorldScaleFactor;	// @todo vreditor tweak
@@ -645,7 +645,7 @@ void FVREditorUISystem::CreateUIs()
 		QuickMenuUI = GetOwner().SpawnTransientSceneActor< AVREditorFloatingUI >(TEXT("QuickMenu"), bWithSceneComponent);
 		QuickMenuUI->SetUMGWidget( *this, QuickMenuWidgetClass, Resolution, 30.0f, AVREditorFloatingUI::EDockedTo::Nothing );
 		QuickMenuUI->ShowUI( false );
-		QuickMenuUI->SetRelativeOffset( FVector( -12.0f, 0.0f, 3.0f ) );
+		QuickMenuUI->SetRelativeOffset( FVector( -11.0f, 0.0f, 3.0f ) );
 		FloatingUIs.Add( QuickMenuUI );
 	}
 
@@ -657,8 +657,6 @@ void FVREditorUISystem::CreateUIs()
 		QuickRadialMenu->SetUMGWidget( *this, QuickRadialWidgetClass, DefaultResolution, 40.0f, AVREditorFloatingUI::EDockedTo::Nothing );
 		QuickRadialMenu->ShowUI( false );
 		QuickRadialMenu->SetActorEnableCollision( false );
-		QuickRadialMenu->SetRotateToHead( false );
-		QuickRadialMenu->SetLocalRotation( FRotator( 45.0f, 180.0f, 0.0f) );
 		QuickRadialMenu->SetRelativeOffset( FVector( 10.0f, 0.0f, 10.f ) );
 		QuickRadialMenu->SetCollisionOnShow( false );
 		FloatingUIs.Add(QuickRadialMenu);
