VR Editor: Fixed regression with world scale and UI docked to your hands

diff --git a/Engine/Source/Editor/VREditor/VREditorFloatingUI.cpp b/Engine/Source/Editor/VREditor/VREditorFloatingUI.cpp
--- a/Engine/Source/Editor/VREditor/VREditorFloatingUI.cpp
+++ b/Engine/Source/Editor/VREditor/VREditorFloatingUI.cpp
@@ -324,7 +324,7 @@ FTransform AVREditorFloatingUI::MakeUITransformLockedToHand( const int32 HandInd
 {
 	const float WorldScaleFactor = Owner->GetOwner().GetWorldScaleFactor();
 
-	FTransform UIToHandTransform( LocalRotation, RelativeOffset );
+	FTransform UIToHandTransform( LocalRotation, RelativeOffset * WorldScaleFactor );
 	if( !bOnArm )
 	{
 		UIToHandTransform *= FTransform( FRotator( VREd::UIOnHandRotationOffset->GetFloat(), 0.0f, 0.0f ).Quaternion(), FVector::ZeroVector );
@@ -342,7 +342,7 @@ FTransform AVREditorFloatingUI::MakeUITransformLockedToRoom()
 {
 	const float WorldScaleFactor = Owner->GetOwner().GetWorldScaleFactor();
 
-	const FTransform UIToRoomTransform( LocalRotation, RelativeOffset );
+	const FTransform UIToRoomTransform( LocalRotation, RelativeOffset * WorldScaleFactor );
 
 	const FTransform RoomToWorldTransform = Owner->GetOwner().GetRoomTransform();
 
