VR Editor: Asset editor UI is now much wider than other editor panels (to accomodate Blueprint editor, etc.)

diff --git a/Engine/Source/Editor/VREditor/VREditorUISystem.cpp b/Engine/Source/Editor/VREditor/VREditorUISystem.cpp
--- a/Engine/Source/Editor/VREditor/VREditorUISystem.cpp
+++ b/Engine/Source/Editor/VREditor/VREditorUISystem.cpp
@@ -34,8 +34,8 @@ namespace VREd
 {
 	static FAutoConsoleVariable ContentBrowserUIResolutionX( TEXT( "VREd.ContentBrowserUIResolutionX" ), 1920, TEXT( "Horizontal resolution to use for content browser UI render targets" ) );
 	static FAutoConsoleVariable ContentBrowserUIResolutionY( TEXT( "VREd.ContentBrowserUIResolutionY" ), 1200, TEXT( "Vertical resolution to use for content browser UI render targets" ) );
-	static FAutoConsoleVariable EditorUIResolutionX( TEXT( "VREd.EditorUIResolutionX" ), 1024, TEXT( "Horizontal resolution to use for VR editor UI render targets" ) );
-	static FAutoConsoleVariable EditorUIResolutionY( TEXT( "VREd.EditorUIResolutionY" ), 1024, TEXT( "Vertical resolution to use for VR editor UI render targets" ) );
+	static FAutoConsoleVariable DefaultEditorUIResolutionX( TEXT( "VREd.DefaultEditorUIResolutionX" ), 1024, TEXT( "Horizontal resolution to use for VR editor UI render targets" ) );
+	static FAutoConsoleVariable DefaultEditorUIResolutionY( TEXT( "VREd.DefaultEditorUIResolutionY" ), 1024, TEXT( "Vertical resolution to use for VR editor UI render targets" ) );
 	static FAutoConsoleVariable QuickMenuUIResolutionX( TEXT( "VREd.QuickMenuUIResolutionX" ), 1024, TEXT( "Horizontal resolution to use for Quick Menu VR UI render targets" ) );
 	static FAutoConsoleVariable QuickMenuUIResolutionY( TEXT( "VREd.QuickMenuUIResolutionY" ), 900, TEXT( "Vertical resolution to use for Quick Menu VR UI render targets" ) );
 	static FAutoConsoleVariable ContentBrowserUISize( TEXT( "VREd.ContentBrowserUISize" ), 60.0f, TEXT( "How big content browser UIs should be" ) );
@@ -43,6 +43,8 @@ namespace VREd
 	static FAutoConsoleVariable ContentBrowserUIScale( TEXT( "VREd.ContentBrowserUIScale" ), 2.0f, TEXT( "How much to scale up (or down) the content browser for VR" ) );
 	static FAutoConsoleVariable EditorUIScale( TEXT( "VREd.EditorUIScale" ), 2.0f, TEXT( "How much to scale up (or down) editor UIs for VR" ) );
 	static FAutoConsoleVariable EditorUIOffsetFromHand( TEXT( "VREd.EditorUIOffsetFromHand" ), 12.0f, TEXT( "How far to offset editor UIs from the origin of the hand mesh" ) );
+	static FAutoConsoleVariable AssetEditorUIResolutionX( TEXT( "VREd.AssetEditorUIResolutionX" ), 1920, TEXT( "Horizontal resolution to use for VR editor asset editor UI render targets" ) );
+	static FAutoConsoleVariable AssetEditorUIResolutionY( TEXT( "VREd.AssetEditorUIResolutionY" ), 1080, TEXT( "Vertical resolution to use for VR editor asset editor UI render targets" ) );
 	static FAutoConsoleVariable RadialMenuFadeDelay( TEXT( "VREd.RadialMenuFadeDelay" ), 0.2f, TEXT( "The delay for the radial menu after selecting a button" ) );
 	static FAutoConsoleVariable UIAbsoluteScrollSpeed( TEXT( "VREd.UIAbsoluteScrollSpeed" ), 8.0f, TEXT( "How fast the UI scrolls when dragging the touchpad" ) );
 	static FAutoConsoleVariable UIRelativeScrollSpeed( TEXT( "VREd.UIRelativeScrollSpeed" ), 0.75f, TEXT( "How fast the UI scrolls when holding an analog stick" ) );
@@ -632,7 +634,7 @@ void FVREditorUISystem::Render( const FSceneView* SceneView, FViewport* Viewport
 
 void FVREditorUISystem::CreateUIs()
 {
-	const FIntPoint DefaultResolution( VREd::EditorUIResolutionX->GetInt(), VREd::EditorUIResolutionY->GetInt() );
+	const FIntPoint DefaultResolution( VREd::DefaultEditorUIResolutionX->GetInt(), VREd::DefaultEditorUIResolutionY->GetInt() );
 
 	const bool bIsVREditorDemo = FParse::Param( FCommandLine::Get(), TEXT( "VREditorDemo" ) );	// @todo vreditor: Remove this when no longer needed
 																								
@@ -814,7 +816,7 @@ void FVREditorUISystem::CreateUIs()
 		}
 
 		{
-			const FIntPoint Resolution(VREd::EditorUIResolutionX->GetInt(), VREd::EditorUIResolutionY->GetInt());
+			const FIntPoint Resolution(VREd::AssetEditorUIResolutionX->GetInt(), VREd::AssetEditorUIResolutionY->GetInt());
 
 			const bool bWithSceneComponent = false;
 			AVREditorDockableWindow* TabManagerUI = GetOwner().SpawnTransientSceneActor< AVREditorDockableWindow >(TEXT("AssetEditor"), bWithSceneComponent);
