VREditor - now embedding the asset editor into a SDPIScaler.

diff --git a/Engine/Source/Editor/VREditor/VREditorUISystem.cpp b/Engine/Source/Editor/VREditor/VREditorUISystem.cpp
--- a/Engine/Source/Editor/VREditor/VREditorUISystem.cpp
+++ b/Engine/Source/Editor/VREditor/VREditorUISystem.cpp
@@ -1192,6 +1192,20 @@ float FVREditorUISystem::GetMinDockWindowSize() const
 
 void FVREditorUISystem::OnProxyTabLaunched(TSharedPtr<SDockTab> NewTab)
 {
+	AVREditorFloatingUI* Panel = EditorUIPanels[(int32)EEditorUIPanel::AssetEditor];
+	if ( Panel != nullptr )
+	{
+		TSharedPtr<SWindow> AssetEditorWindow = Panel->GetWidgetComponent()->GetSlateWindow();
+		TSharedPtr<SWidget> AssetEditorWidget = const_cast<SWidget&>( AssetEditorWindow->GetContent().Get() ).AsShared();
+
+		Panel->GetWidgetComponent()->SetSlateWidget(
+			SNew(SDPIScaler)
+			.DPIScale(VREd::ContentBrowserUIScale->GetFloat())
+			[
+				AssetEditorWidget.ToSharedRef()
+			]);
+	}
+
 	ShowAssetEditor();
 }
 
