VR Editor: Fixed window docking regression (UE-28125)

diff --git a/Engine/Source/Editor/VREditor/VREditorDockableWindow.cpp b/Engine/Source/Editor/VREditor/VREditorDockableWindow.cpp
--- a/Engine/Source/Editor/VREditor/VREditorDockableWindow.cpp
+++ b/Engine/Source/Editor/VREditor/VREditorDockableWindow.cpp
@@ -279,7 +279,7 @@ void AVREditorDockableWindow::UpdateRelativeRoomTransform()
 	const FTransform WindowToWorldTransform = GetActorTransform();
 	const FTransform WindowToRoomTransform = WindowToWorldTransform * WorldToRoom;
 
-	const FVector RoomSpaceWindowLocation = WindowToRoomTransform.GetLocation();
+	const FVector RoomSpaceWindowLocation = WindowToRoomTransform.GetLocation() / GetOwner().GetOwner().GetWorldScaleFactor();;
 	const FQuat RoomSpaceWindowRotation = WindowToRoomTransform.GetRotation();
 
 	RelativeOffset = RoomSpaceWindowLocation;
