VR Editor: -  UI Snap to hand movement
- Improved toggle showing and hiding all dockable windows
- Sounds for UI

diff --git a/Engine/Source/Editor/VREditor/VREditorFloatingUI.cpp b/Engine/Source/Editor/VREditor/VREditorFloatingUI.cpp
--- a/Engine/Source/Editor/VREditor/VREditorFloatingUI.cpp
+++ b/Engine/Source/Editor/VREditor/VREditorFloatingUI.cpp
@@ -31,7 +31,12 @@ AVREditorFloatingUI::AVREditorFloatingUI()
 	  FadeAlpha( 1.0f ),
 	  LocalRotation( FRotator( 90.0f, 180.0f, 0.0f ) ),
 	  bCollisionOnShowUI( true ),
-	  FadeDelay( 0.0f )
+	  FadeDelay( 0.0f ),
+	  bIsMoving( false ),
+	  MoveToTransform( FTransform() ),
+	  StartMoveToTransform( FTransform() ),
+	  MoveToAlpha( 0.0f ),
+	  MoveToTime( 0.0f )
 {
 	const bool bTransient = true;
 	USceneComponent* SceneComponent = CreateDefaultSubobject<USceneComponent>( TEXT( "SceneComponent" ), bTransient );
@@ -167,8 +172,15 @@ void AVREditorFloatingUI::TickManually( float DeltaTime )
 	// Update fading state
 	UpdateFadingState( DeltaTime );
 
-	// Update transform
-	UpdateTransformIfDocked();
+	if ( bIsMoving )
+	{
+		TickMoveTo( DeltaTime );
+	}
+	else
+	{
+		// Update transform
+		UpdateTransformIfDocked();
+	}
 }
 
 
@@ -321,11 +333,16 @@ float AVREditorFloatingUI::GetInitialScale() const
 }
 
 FTransform AVREditorFloatingUI::MakeUITransformLockedToHand( const int32 HandIndex, const bool bOnArm )
+{
+	return MakeUITransformLockedToHand( HandIndex, bOnArm, RelativeOffset, LocalRotation );
+}
+
+FTransform AVREditorFloatingUI::MakeUITransformLockedToHand( const int32 HandIndex, const bool bOnArm, const FVector& InRelativeOffset, const FRotator& InLocalRotation )
 {
 	const float WorldScaleFactor = Owner->GetOwner().GetWorldScaleFactor();
 
-	FTransform UIToHandTransform( LocalRotation, RelativeOffset * WorldScaleFactor );
-	if( !bOnArm )
+	FTransform UIToHandTransform( InLocalRotation, InRelativeOffset * WorldScaleFactor );
+	if (!bOnArm)
 	{
 		UIToHandTransform *= FTransform( FRotator( VREd::UIOnHandRotationOffset->GetFloat(), 0.0f, 0.0f ).Quaternion(), FVector::ZeroVector );
 	}
@@ -338,6 +355,23 @@ FTransform AVREditorFloatingUI::MakeUITransformLockedToHand( const int32 HandInd
 	return UIToWorldTransform;
 }
 
+void AVREditorFloatingUI::MoveTo( const FTransform& ResultTransform, const float TotalMoveToTime, const EDockedTo ResultDock )
+{
+	MoveToTime = TotalMoveToTime;
+	bIsMoving = true;
+	StartMoveToTransform = GetActorTransform();
+	MoveToTransform = ResultTransform;
+	MoveToAlpha = 0.0f;
+	MoveToResultDock = ResultDock;
+}
+
+void AVREditorFloatingUI::StopMoveTo()
+{
+	bIsMoving = false;
+	SetTransform( MoveToTransform );
+	MoveToTransform = FTransform();
+}
+
 FTransform AVREditorFloatingUI::MakeUITransformLockedToRoom()
 {
 	const float WorldScaleFactor = Owner->GetOwner().GetWorldScaleFactor();
@@ -352,6 +386,27 @@ FTransform AVREditorFloatingUI::MakeUITransformLockedToRoom()
 	return UIToWorldTransform;
 }
 
+void AVREditorFloatingUI::TickMoveTo( const float DeltaTime )
+{
+	const float WorldScaleFactor = Owner->GetOwner().GetWorldScaleFactor();
+
+	MoveToAlpha += DeltaTime;
+	const float LerpTime = MoveToTime;
+	if (MoveToAlpha > MoveToTime)
+	{
+		MoveToAlpha = MoveToTime;
+		bIsMoving = false;
+		SetDockedTo( MoveToResultDock );
+	}
+
+	const float CurrentALpha = MoveToAlpha / LerpTime;
+	const FVector NewLocation = FMath::Lerp( StartMoveToTransform.GetLocation(), MoveToTransform.GetLocation(), CurrentALpha );
+	const FQuat NewRotation = FMath::Lerp( StartMoveToTransform.GetRotation(), MoveToTransform.GetRotation(), CurrentALpha );
+	FTransform NewTransform( NewRotation, NewLocation, FVector( Scale * WorldScaleFactor ) );
+
+	SetTransform( NewTransform );
+}
+
 void AVREditorFloatingUI::SetDockedTo( const EDockedTo NewDockedTo )
 {
 	PreviousDockedTo = DockedTo;
diff --git a/Engine/Source/Editor/VREditor/VREditorFloatingUI.h b/Engine/Source/Editor/VREditor/VREditorFloatingUI.h
--- a/Engine/Source/Editor/VREditor/VREditorFloatingUI.h
+++ b/Engine/Source/Editor/VREditor/VREditorFloatingUI.h
@@ -120,6 +120,15 @@ public:
 	/** Gets the initial size of this UI */
 	float GetInitialScale() const;
 
+	/** Get a relative transform from the hand */
+	FTransform MakeUITransformLockedToHand( const int32 HandIndex, const bool bOnArm, const FVector& InRelativeOffset, const FRotator& InLocalRotation );
+
+	/** Start moving towards a transform */
+	void MoveTo( const FTransform& ResultTransform, const float TotalMoveToTime, const EDockedTo ResultDock );
+
+	/** Abort moving to transform and move instant to MoveToTransform */
+	void StopMoveTo();
+
 protected:
 
 	/** Returns a scale to use for this widget that takes into account animation */
@@ -145,8 +154,10 @@ private:
 	/** Create a room transform with the RelativeOffset and LocalRotation */
 	FTransform MakeUITransformLockedToRoom();
 
+	/** Updates the lerp movement */
+	void TickMoveTo( const float DeltaTime );
 
-protected:
+public:
 
 	/** Local rotation of the UI */
 	FRotator LocalRotation;
@@ -154,6 +165,7 @@ protected:
 	/** Relative offset of the UI */
 	FVector RelativeOffset;
 
+protected:
 	/** Slate widget we're drawing, or null if we're drawing a UMG user widget */
 	TSharedPtr<SWidget> SlateWidget;
 	
@@ -204,5 +216,23 @@ private:
 
 	/** The starting scale of this UI */
 	float InitialScale;
+
+	/** If the UI is currently moving */
+	bool bIsMoving;
+
+	/** The end transform to move to */
+	FTransform MoveToTransform;
+
+	/** The actor transform when started moving */
+	FTransform StartMoveToTransform;
+
+	/** Current lerp alpha */
+	float MoveToAlpha;
+
+	/** Total time to move to the end transform */
+	float MoveToTime;
+
+	/** Dock state to switch to when finished moving */
+	EDockedTo MoveToResultDock;
 };
 
diff --git a/Engine/Source/Editor/VREditor/VREditorUISystem.cpp b/Engine/Source/Editor/VREditor/VREditorUISystem.cpp
--- a/Engine/Source/Editor/VREditor/VREditorUISystem.cpp
+++ b/Engine/Source/Editor/VREditor/VREditorUISystem.cpp
@@ -55,6 +55,8 @@ namespace VREd
 	static FAutoConsoleVariable UIAssetEditorSummonedOnHandHapticFeedbackStrength( TEXT( "VREd.UIAssetEditorSummonedOnHandHapticFeedbackStrength" ), 1.0f, TEXT( "Strenth of haptic to play to remind a user which hand an asset editor was spawned on" ) );
 	static FAutoConsoleVariable MaxDockWindowSize( TEXT( "VREd.MaxDockWindowSize" ), 250, TEXT( "Maximum size for dockable windows" ) );
 	static FAutoConsoleVariable MinDockWindowSize( TEXT( "VREd.MinDockWindowSize" ), 40, TEXT( "Minimum size for dockable windows" ) );
+	static FAutoConsoleVariable DockSnapTime( TEXT( "VREd.DockSnapTime" ), 0.3f, TEXT( "Lerp time for snapping dock to hand and back" ) );
+	static FAutoConsoleVariable DockDragSpeed( TEXT( "VREd.DockDragSpeed" ), 0.05f, TEXT( "Dock dragging time delay" ) );
 
 	// Tutorial UI commands
 	static FAutoConsoleVariable TutorialUIResolutionX( TEXT( "VREd.TutorialUI.Resolution.X" ), 1920, TEXT( "The X resolution for the tutorial UI panel" ) );
@@ -78,7 +80,6 @@ FVREditorUISystem::FVREditorUISystem( FVREditorMode& InitOwner )
 	  DraggingUI( nullptr ),
 	  DraggingUIHandIndex( INDEX_NONE ),
 	  DraggingUIOffsetTransform( FTransform::Identity ),
-	  bPanelVisibilityToggle( false ),
 	  RadialMenuHideDelayTime( 0.0f ),
 	  bRefocusViewport(false)
 {
@@ -87,7 +88,7 @@ FVREditorUISystem::FVREditorUISystem( FVREditorMode& InitOwner )
 	Owner.OnVRHoverUpdate().AddRaw( this, &FVREditorUISystem::OnVRHoverUpdate );
 
 	EditorUIPanels.SetNumZeroed( (int32)EEditorUIPanel::TotalCount );
-
+	
 	QuickMenuWidgetClass = LoadClass<UVREditorQuickMenu>( nullptr, TEXT( "/Engine/VREditor/UI/VRQuickMenu.VRQuickMenu_C" ) );
 	check( QuickMenuWidgetClass != nullptr );
 
@@ -125,6 +126,11 @@ FVREditorUISystem::~FVREditorUISystem()
 
 	QuickMenuWidgetClass = nullptr;
 	QuickRadialWidgetClass = nullptr;
+	TutorialWidgetClass = nullptr;
+	StartDragUISound = nullptr;
+	StopDragUISound = nullptr;
+	HideUISound = nullptr;
+	ShowUISound = nullptr;
 }
 
 
@@ -307,19 +313,25 @@ void FVREditorUISystem::OnVRAction( FEditorViewportClient& ViewportClient, const
 		// Put the Dock back on the hand it came from or leave it where it is in the room
 		if (DraggingUI != nullptr && DraggingUIHandIndex == VRAction.HandIndex)
 		{
-			FVirtualHand& Hand = Owner.GetVirtualHand( VRAction.HandIndex );
 			FVirtualHand& OtherHand = Owner.GetOtherHand( VRAction.HandIndex );
-			const float Distance = FVector::Dist( Hand.HoverLocation, OtherHand.Transform.GetLocation() ) / Owner.GetWorldScaleFactor();
-			bool bOnHand = true;
+			const int32 OtherHandIndex = Owner.GetOtherHandIndex( DraggingUIHandIndex );
+
+			bool bOnHand = false;
+			const float Distance = FVector::Dist( LastDraggingHoverLocation, OtherHand.Transform.GetLocation() ) / Owner.GetWorldScaleFactor();
 			if (Distance > VREd::MinDockDragDistance->GetFloat())
 			{
 				DraggingUI->SetDockedTo( AVREditorFloatingUI::EDockedTo::Room );
-				bOnHand = false;
+			}
+			else
+			{
+				bOnHand = true;
+				const FVector EditorUIRelativeOffset( DraggingUI->GetSize().Y * 0.5f + VREd::EditorUIOffsetFromHand->GetFloat(), 0.0f, 0.0f );
+				const FTransform MovingUIGoalTransform = DraggingUI->MakeUITransformLockedToHand( OtherHandIndex, false, EditorUIRelativeOffset, FRotator( 90.0f, 180.0f, 0.0f ) );
+				const AVREditorFloatingUI::EDockedTo NewDockedTo = OtherHandIndex == VREditorConstants::LeftHandIndex ? AVREditorFloatingUI::EDockedTo::LeftHand : AVREditorFloatingUI::EDockedTo::RightHand;
+				DraggingUI->MoveTo( MovingUIGoalTransform, VREd::DockSnapTime->GetFloat(), NewDockedTo );
 			}
 
-			const int32 OtherHandIndex = Owner.GetOtherHandIndex( DraggingUIHandIndex );
-			ShowEditorUIPanel( DraggingUI, OtherHandIndex, true, bOnHand, true );
-			
+			ShowEditorUIPanel( DraggingUI, OtherHandIndex, true, bOnHand );
 			StopDraggingDockUI();
 		}
 	}
@@ -910,7 +922,7 @@ bool FVREditorUISystem::IsShowingEditorUIPanel( const EEditorUIPanel EditorUIPan
 }
 
 
-void FVREditorUISystem::ShowEditorUIPanel( const UWidgetComponent* WidgetComponent, const int32 HandIndex, const bool bShouldShow, const bool OnHand, const bool bRefreshQuickMenu )
+void FVREditorUISystem::ShowEditorUIPanel( const UWidgetComponent* WidgetComponent, const int32 HandIndex, const bool bShouldShow, const bool OnHand, const bool bRefreshQuickMenu, const bool bPlaySound )
 {
 	AVREditorFloatingUI* Panel = nullptr;
 	for( AVREditorFloatingUI* CurrentPanel : EditorUIPanels )
@@ -922,17 +934,17 @@ void FVREditorUISystem::ShowEditorUIPanel( const UWidgetComponent* WidgetCompone
 		}
 	}
 
-	ShowEditorUIPanel( Panel, HandIndex, bShouldShow, OnHand, bRefreshQuickMenu );
+	ShowEditorUIPanel( Panel, HandIndex, bShouldShow, OnHand, bRefreshQuickMenu, bPlaySound );
 }
 
 
-void FVREditorUISystem::ShowEditorUIPanel( const EEditorUIPanel EditorUIPanel, const int32 HandIndex, const bool bShouldShow, const bool OnHand, const bool bRefreshQuickMenu )
+void FVREditorUISystem::ShowEditorUIPanel( const EEditorUIPanel EditorUIPanel, const int32 HandIndex, const bool bShouldShow, const bool OnHand, const bool bRefreshQuickMenu, const bool bPlaySound )
 {
 	AVREditorFloatingUI* Panel = EditorUIPanels[ (int32)EditorUIPanel ];
-	ShowEditorUIPanel( Panel, HandIndex, bShouldShow, OnHand, bRefreshQuickMenu );
+	ShowEditorUIPanel( Panel, HandIndex, bShouldShow, OnHand, bRefreshQuickMenu, bPlaySound );
 }
 
-void FVREditorUISystem::ShowEditorUIPanel( AVREditorFloatingUI* Panel, const int32 HandIndex, const bool bShouldShow, const bool OnHand, const bool bRefreshQuickMenu )
+void FVREditorUISystem::ShowEditorUIPanel( AVREditorFloatingUI* Panel, const int32 HandIndex, const bool bShouldShow, const bool OnHand, const bool bRefreshQuickMenu, const bool bPlaySound )
 {
 	if( Panel != nullptr )
 	{
@@ -965,7 +977,7 @@ void FVREditorUISystem::ShowEditorUIPanel( AVREditorFloatingUI* Panel, const int
 
 			const FVector EditorUIRelativeOffset( Panel->GetSize().Y * 0.5f + VREd::EditorUIOffsetFromHand->GetFloat(), 0.0f, 0.0f );
 			Panel->SetRelativeOffset( EditorUIRelativeOffset );
-			Panel->SetLocalRotation( FRotator( 90.0f, 180.0f, 0.0f ) );
+			Panel->SetLocalRotation( FRotator( 90.0f, 180.0f, 0.0f ) ); // Todo needs initial rotation 
 		}
 
 		Panel->ShowUI( bShouldShow );
@@ -975,7 +987,10 @@ void FVREditorUISystem::ShowEditorUIPanel( AVREditorFloatingUI* Panel, const int
 			Owner.GetVirtualHand( HandIndex ).bHasUIInFront = bShouldShow;
 		}
 
-		UGameplayStatics::PlaySound2D( Owner.GetWorld(), bShouldShow ? ShowUISound : HideUISound );
+		if( bPlaySound )
+		{
+			UGameplayStatics::PlaySound2D( Owner.GetWorld(), bShouldShow ? ShowUISound : HideUISound );
+		}
 
 		if (bRefreshQuickMenu && QuickMenuUI)
 		{
@@ -1060,7 +1075,11 @@ FTransform FVREditorUISystem::MakeDockableUITransformOnLaser( AVREditorDockableW
 {
 	const FVirtualHand& Hand = Owner.GetVirtualHand( HandIndex );
 
-	const FVector NewLocation = Hand.Transform.GetLocation() + ( Hand.Transform.GetRotation().Vector().GetSafeNormal() * DockSelectDistance );
+	FTransform HandTransform;
+	FVector HandForward;
+	Owner.GetHandTransformAndForwardVector( HandIndex, HandTransform, HandForward );
+
+	const FVector NewLocation = Hand.Transform.GetLocation() + ( HandForward.GetSafeNormal() * DockSelectDistance );
 	
 	FRotator NewRotation = ( Hand.Transform.GetLocation() - NewLocation ).ToOrientationRotator();
 	NewRotation.Roll = -Hand.Transform.GetRotation().Rotator().Roll;
@@ -1070,12 +1089,47 @@ FTransform FVREditorUISystem::MakeDockableUITransformOnLaser( AVREditorDockableW
 }
 
 
-FTransform FVREditorUISystem::MakeDockableUITransform( AVREditorDockableWindow* InitDraggingDockUI, const int32 HandIndex, const float DockSelectDistance ) const
+FTransform FVREditorUISystem::MakeDockableUITransform( AVREditorDockableWindow* InitDraggingDockUI, const int32 HandIndex, const float DockSelectDistance )
 {
 	const FTransform UIOnLaserToWorld = MakeDockableUITransformOnLaser( DraggingUI, HandIndex, DockSelectDistance );
 	const FTransform UIToUIOnLaser = DraggingUIOffsetTransform;
 	
 	const FTransform UIToWorld = UIToUIOnLaser * UIOnLaserToWorld;
+
+	if( DraggingUI != nullptr &&  DraggingUI->GetDockedTo() == AVREditorFloatingUI::EDockedTo::Dragging )
+	{
+		
+		FVirtualHand& Hand = Owner.GetVirtualHand( DraggingUIHandIndex );
+		FVirtualHand& OtherHand = Owner.GetOtherHand( DraggingUIHandIndex );
+		
+		const float WorldScaleFactor = Owner.GetWorldScaleFactor();
+		const FVector DraggingUILocation = Hand.HoverLocation;
+		const FVector OtherHandLocation =  OtherHand.Transform.GetLocation();
+		LastDraggingHoverLocation = UIOnLaserToWorld.GetLocation();
+
+		const float Distance = FVector::Dist( UIOnLaserToWorld.GetLocation(), OtherHandLocation ) / WorldScaleFactor;
+
+		// if dragged passed the threshold since starting dragging
+		if ( Distance > VREd::MinDockDragDistance->GetFloat() && !bDraggedDockFromHandPassedThreshold && ( DraggingUI->GetPreviouslyDockedTo() == AVREditorFloatingUI::EDockedTo::LeftHand || DraggingUI->GetPreviouslyDockedTo() == AVREditorFloatingUI::EDockedTo::RightHand ) )
+		{
+			bDraggedDockFromHandPassedThreshold = true;
+		}
+
+		// Snapping to a hand when in range
+		if ( Distance <= VREd::MinDockDragDistance->GetFloat() && ( bDraggedDockFromHandPassedThreshold || DraggingUI->GetPreviouslyDockedTo() == AVREditorFloatingUI::EDockedTo::Room ) )
+		{
+			const int32 OtherHandIndex = Owner.GetOtherHandIndex( HandIndex );
+			const FVector EditorUIRelativeOffset( DraggingUI->GetSize().Y * 0.5f + VREd::EditorUIOffsetFromHand->GetFloat(), 0.0f, 0.0f );
+			const FTransform MovingUIGoalTransform = DraggingUI->MakeUITransformLockedToHand( OtherHandIndex, false, EditorUIRelativeOffset, FRotator( 90.0f, 180.0f, 0.0f ) );
+			const AVREditorFloatingUI::EDockedTo NewDockedTo = OtherHandIndex == VREditorConstants::LeftHandIndex ? AVREditorFloatingUI::EDockedTo::LeftHand : AVREditorFloatingUI::EDockedTo::RightHand;
+			DraggingUI->MoveTo( MovingUIGoalTransform, VREd::DockDragSpeed->GetFloat(), NewDockedTo );
+		}
+		else
+		{
+			DraggingUI->MoveTo( UIToWorld, VREd::DockDragSpeed->GetFloat(), AVREditorFloatingUI::EDockedTo::Room );
+		}
+	}
+
 	return UIToWorld;
 }
 
@@ -1085,6 +1139,7 @@ AVREditorFloatingUI* FVREditorUISystem::StartDraggingDockUI( AVREditorDockableWi
 	AVREditorFloatingUI::EDockedTo DockTo = InitDraggingDockUI->GetDockedTo();
 	if( DockTo == AVREditorFloatingUI::EDockedTo::LeftHand || DockTo == AVREditorFloatingUI::EDockedTo::RightHand )
 	{
+		bDraggedDockFromHandPassedThreshold = false;
 		Owner.GetOtherHand( HandIndex ).bHasUIInFront = false;
 	}
 
@@ -1134,58 +1189,74 @@ bool FVREditorUISystem::IsDraggingDockUI()
 
 void FVREditorUISystem::TogglePanelsVisibility()
 {
-	bPanelVisibilityToggle = !bPanelVisibilityToggle;
-	
-	UGameplayStatics::PlaySound2D( Owner.GetWorld(), bPanelVisibilityToggle ? ShowUISound : HideUISound );
-
-	bool bFirstCheckOnHand = true;
+	bool bAnyPanelsVisible = false;
 	AVREditorFloatingUI* PanelOnHand = nullptr;
-	for( AVREditorFloatingUI* Panel : EditorUIPanels )
+
+	// Check if there is any panel visible and if any is docked to a hand
+	for (AVREditorFloatingUI* Panel : EditorUIPanels)
 	{
-		if( Panel != nullptr && Panel->IsUIVisible() != bPanelVisibilityToggle )
+		if (Panel && Panel->IsUIVisible())
 		{
-			bool bShouldSetNewVisibility = true;
+			bAnyPanelsVisible = true;
+
 			const AVREditorFloatingUI::EDockedTo DockedTo = Panel->GetDockedTo();
-			if( DockedTo == AVREditorFloatingUI::EDockedTo::LeftHand ||  DockedTo == AVREditorFloatingUI::EDockedTo::RightHand  )
+			if (DockedTo == AVREditorFloatingUI::EDockedTo::LeftHand || DockedTo == AVREditorFloatingUI::EDockedTo::RightHand)
 			{
-				// Prevent the panel from spawning on the hand if there is already a UI on the hand
-				// Iterate through all panels to check if there is already a panel on the hand, but only once and if we haven't found one yet
-				if( bFirstCheckOnHand && PanelOnHand == nullptr )
+				PanelOnHand = Panel;
+				break;
+			}
+		}
+	}
+
+	// Hide if there is any UI visible
+	const bool bShowUI = !bAnyPanelsVisible;
+	
+	UGameplayStatics::PlaySound2D( Owner.GetWorld(), bShowUI ? ShowUISound : HideUISound );
+
+	for (AVREditorFloatingUI* Panel : EditorUIPanels)
+	{
+		if (Panel != nullptr && Panel->IsUIVisible() != bShowUI)
+		{
+			if (bShowUI)
+			{
+				bool SetNewVisibility = true;
+				const AVREditorFloatingUI::EDockedTo DockedTo = Panel->GetDockedTo();
+				if (DockedTo == AVREditorFloatingUI::EDockedTo::LeftHand || DockedTo == AVREditorFloatingUI::EDockedTo::RightHand)
 				{
-					bFirstCheckOnHand = false;
-					for(AVREditorFloatingUI* OtherPanel : EditorUIPanels)
+					if (!PanelOnHand)
 					{
-						if( OtherPanel != Panel && OtherPanel->IsUIVisible() && ( OtherPanel->GetDockedTo() == AVREditorFloatingUI::EDockedTo::LeftHand || OtherPanel->GetDockedTo() == AVREditorFloatingUI::EDockedTo::RightHand ) )
-						{
-							PanelOnHand = OtherPanel;
-							break;
-						}
+						const uint32 HandIndex = Panel->GetDockedTo() == AVREditorFloatingUI::EDockedTo::LeftHand ? VREditorConstants::LeftHandIndex : VREditorConstants::RightHandIndex;
+						Owner.GetVirtualHand( HandIndex ).bHasUIInFront = true;
+						PanelOnHand = Panel;
+					}
+					else
+					{
+						SetNewVisibility = false;
 					}
 				}
-
-				if( !PanelOnHand )
+				else if (DockedTo == AVREditorFloatingUI::EDockedTo::Nothing)
 				{
-					const uint32 HandIndex = Panel->GetDockedTo() == AVREditorFloatingUI::EDockedTo::LeftHand ? VREditorConstants::LeftHandIndex : VREditorConstants::RightHandIndex;
-					Owner.GetVirtualHand( HandIndex ).bHasUIInFront = bPanelVisibilityToggle;
-					PanelOnHand = Panel;
+					SetNewVisibility = false;
 				}
-				else
+
+				if (SetNewVisibility)
 				{
-					bShouldSetNewVisibility = false;
+					Panel->ShowUI( true );
 				}
 			}
-			else if ( DockedTo == AVREditorFloatingUI::EDockedTo::Nothing )
-			{
-				bShouldSetNewVisibility = false;
-			}
-
-			if( bShouldSetNewVisibility )
+			else
 			{
-				Panel->ShowUI( bPanelVisibilityToggle );
+				Panel->ShowUI( false );
 			}
 		}
 	}
 
+	if( !bShowUI && PanelOnHand )
+	{
+		const uint32 HandIndex = PanelOnHand->GetDockedTo() == AVREditorFloatingUI::EDockedTo::LeftHand ? VREditorConstants::LeftHandIndex : VREditorConstants::RightHandIndex;
+		Owner.GetVirtualHand( HandIndex ).bHasUIInFront = false;
+	}
+
 	if (QuickMenuUI)
 	{
 		QuickMenuUI->GetUserWidget<UVREditorQuickMenu>()->RefreshUI();
@@ -1204,20 +1275,6 @@ float FVREditorUISystem::GetMinDockWindowSize() const
 
 void FVREditorUISystem::OnProxyTabLaunched(TSharedPtr<SDockTab> NewTab)
 {
-	AVREditorFloatingUI* Panel = EditorUIPanels[(int32)EEditorUIPanel::AssetEditor];
-	if ( Panel != nullptr )
-	{
-		TSharedPtr<SWindow> AssetEditorWindow = Panel->GetWidgetComponent()->GetSlateWindow();
-		TSharedPtr<SWidget> AssetEditorWidget = const_cast<SWidget&>( AssetEditorWindow->GetContent().Get() ).AsShared();
-
-		Panel->GetWidgetComponent()->SetSlateWidget(
-			SNew(SDPIScaler)
-			.DPIScale(VREd::ContentBrowserUIScale->GetFloat())
-			[
-				AssetEditorWidget.ToSharedRef()
-			]);
-	}
-
 	ShowAssetEditor();
 }
 
@@ -1266,4 +1323,3 @@ void FVREditorUISystem::TogglePanelVisibility( const EEditorUIPanel EditorUIPane
 		}
 	}
 }
-
diff --git a/Engine/Source/Editor/VREditor/VREditorUISystem.h b/Engine/Source/Editor/VREditor/VREditorUISystem.h
--- a/Engine/Source/Editor/VREditor/VREditorUISystem.h
+++ b/Engine/Source/Editor/VREditor/VREditorUISystem.h
@@ -61,9 +61,9 @@ public:
 	bool IsShowingEditorUIPanel( const EEditorUIPanel EditorUIPanel ) const;
 
 	/** Sets whether the specified editor UI panel should be visible.  Any other UI floating off this hand will be dismissed when showing it. */
-	void ShowEditorUIPanel( const class UWidgetComponent* WidgetComponent, const int32 HandIndex, const bool bShouldShow, const bool OnHand = false, const bool bRefreshQuickMenu = true );
-	void ShowEditorUIPanel( const EEditorUIPanel EditorUIPanel, const int32 HandIndex, const bool bShouldShow, const bool OnHand = false, const bool bRefreshQuickMenu = true );
-	void ShowEditorUIPanel( class AVREditorFloatingUI* Panel, const int32 HandIndex, const bool bShouldShow, const bool OnHand = false, const bool bRefreshQuickMenu = true );
+	void ShowEditorUIPanel( const class UWidgetComponent* WidgetComponent, const int32 HandIndex, const bool bShouldShow, const bool OnHand = false, const bool bRefreshQuickMenu = true, const bool bPlaySound = true );
+	void ShowEditorUIPanel( const EEditorUIPanel EditorUIPanel, const int32 HandIndex, const bool bShouldShow, const bool OnHand = false, const bool bRefreshQuickMenu = true, const bool bPlaySound = true );
+	void ShowEditorUIPanel( class AVREditorFloatingUI* Panel, const int32 HandIndex, const bool bShouldShow, const bool OnHand = false, const bool bRefreshQuickMenu = true, const bool bPlaySound = true );
 
 	/** Returns true if the radial menu is visible on this hand */
 	bool IsShowingRadialMenu( const int32 HandIndex ) const;
@@ -84,7 +84,7 @@ public:
 	FTransform MakeDockableUITransformOnLaser( AVREditorDockableWindow* InitDraggingDockUI, const int32 HandIndex, const float DockSelectDistance ) const;
 
 	/** Makes up a transform for a dockable UI when dragging it that includes the original offset from the laser's impact point */
-	FTransform MakeDockableUITransform( AVREditorDockableWindow* InitDraggingDockUI, const int32 HandIndex, const float DockSelectDistance ) const;
+	FTransform MakeDockableUITransform( AVREditorDockableWindow* InitDraggingDockUI, const int32 HandIndex, const float DockSelectDistance );
 	
 	/** Returns the current dragged dock window, nullptr if none */
 	class AVREditorDockableWindow* GetDraggingDockUI() const;
@@ -178,10 +178,6 @@ protected:
 	/** The current UI that is being dragged */
 	class AVREditorDockableWindow* DraggingUI;
 
-	/** If the panels should become visible or invisible next toggle*/
-	bool bPanelVisibilityToggle;
-
-
 	//
 	// Tab Manager UI
 	//
@@ -208,5 +204,11 @@ protected:
 
 	/** Show UI sound */
 	class USoundCue* ShowUISound;
+
+	/** If the current dragged dock passed a certain distance if dragged from a hand */
+	bool bDraggedDockFromHandPassedThreshold;
+
+	/** The last dragged hover location by the laser */
+	FVector LastDraggingHoverLocation;
 };
 
