VR Editor: Green laser when dragging UI

diff --git a/Engine/Source/Editor/VREditor/VREditorMode.cpp b/Engine/Source/Editor/VREditor/VREditorMode.cpp
--- a/Engine/Source/Editor/VREditor/VREditorMode.cpp
+++ b/Engine/Source/Editor/VREditor/VREditorMode.cpp
@@ -1193,7 +1193,8 @@ void FVREditorMode::Tick( FEditorViewportClient* ViewportClient, float DeltaTime
 			else if( Hand.DraggingMode == EVREditorDraggingMode::ActorsAtLaserImpact ||
 				Hand.DraggingMode == EVREditorDraggingMode::ActorsFreely ||
 				Hand.DraggingMode == EVREditorDraggingMode::ActorsWithGizmo ||
-				Hand.DraggingMode == EVREditorDraggingMode::AssistingDrag )
+				Hand.DraggingMode == EVREditorDraggingMode::AssistingDrag || 
+				( GetUISystem().IsDraggingDockUI() && HandIndex == GetUISystem().GetDraggingDockUIHandIndex() ) )
 			{
 				ResultColor = GetColor( EColors::SelectionColor );
 			}
