VR Editor: Sound when dropping textures or materials

diff --git a/Engine/Source/Editor/VREditor/VREditorWorldInteraction.cpp b/Engine/Source/Editor/VREditor/VREditorWorldInteraction.cpp
--- a/Engine/Source/Editor/VREditor/VREditorWorldInteraction.cpp
+++ b/Engine/Source/Editor/VREditor/VREditorWorldInteraction.cpp
@@ -289,10 +289,11 @@ FVREditorWorldInteraction::FVREditorWorldInteraction( FVREditorMode& InitOwner )
 	FEditorDelegates::OnAssetDragStarted.AddRaw( this, &FVREditorWorldInteraction::OnAssetDragStartedFromContentBrowser );
 
 	// Load sounds
-	{
-		TeleportSound = LoadObject<USoundCue>( nullptr, TEXT( "/Engine/VREditor/Sounds/VR_teleport_Cue" ) );
-		check( TeleportSound != nullptr );
-	}
+	TeleportSound = LoadObject<USoundCue>( nullptr, TEXT( "/Engine/VREditor/Sounds/VR_teleport_Cue" ) );
+	check( TeleportSound != nullptr );
+
+	DropMaterialOrMaterialSound = LoadObject<USoundCue>( nullptr, TEXT( "/Engine/VREditor/Sounds/VR_grab_Cue" ) );
+	check( DropMaterialOrMaterialSound != nullptr );
 }
 
 
@@ -302,6 +303,8 @@ FVREditorWorldInteraction::~FVREditorWorldInteraction()
 	Owner.DestroyTransientActor( TransformGizmoActor );
 	TransformGizmoActor = nullptr;
 	FloatingUIAssetDraggedFrom = nullptr;
+	TeleportSound = nullptr;
+	DropMaterialOrMaterialSound = nullptr;
 
 	FEditorDelegates::OnAssetDragStarted.RemoveAll( this );
 
@@ -317,6 +320,8 @@ void FVREditorWorldInteraction::AddReferencedObjects( FReferenceCollector& Colle
 	Collector.AddReferencedObject( TransformGizmoActor );
 	Collector.AddReferencedObject( PlacingMaterialOrTextureAsset );
 	Collector.AddReferencedObject( FloatingUIAssetDraggedFrom );
+	Collector.AddReferencedObject( TeleportSound );
+	Collector.AddReferencedObject( DropMaterialOrMaterialSound );
 }
 
 
@@ -605,7 +610,7 @@ void FVREditorWorldInteraction::OnVRAction( FEditorViewportClient& ViewportClien
 					if( Hand.DraggingMode == EVREditorDraggingMode::ActorsAtLaserImpact ||
 						Hand.DraggingMode == EVREditorDraggingMode::Material )
 					{
-						Owner.GetUISystem().ShowEditorUIPanel( FloatingUIAssetDraggedFrom, Owner.GetOtherHandIndex( VRAction.HandIndex ), true );
+						Owner.GetUISystem().ShowEditorUIPanel( FloatingUIAssetDraggedFrom, Owner.GetOtherHandIndex( VRAction.HandIndex ), true, true, false, false );
 						FloatingUIAssetDraggedFrom = nullptr;
 					}
 
@@ -2739,7 +2744,7 @@ void FVREditorWorldInteraction::OnAssetDragStartedFromContentBrowser( const TArr
 
 		FloatingUIAssetDraggedFrom = Hand.HoveringOverWidgetComponent;
 		// Hide the UI panel that's being used to drag
-		Owner.GetUISystem().ShowEditorUIPanel( FloatingUIAssetDraggedFrom, OtherHandIndex, false );
+		Owner.GetUISystem().ShowEditorUIPanel( FloatingUIAssetDraggedFrom, OtherHandIndex, false, false, true, false );
 
 		TArray< UObject* > DroppedObjects;
 		TArray< AActor* > AllNewActors;
@@ -2975,7 +2980,10 @@ void FVREditorWorldInteraction::PlaceDraggedMaterialOrTexture( const int32 HandI
 				const bool bPlaced = FComponentEditorUtils::AttemptApplyMaterialToComponent( HitComponent, DroppedObjAsMaterial, TargetMaterialSlot );
 			}
 
-			// @todo vreditor sounds: Needs a sound effect on apply!
+			if( DroppedObjAsMaterial || DroppedObjAsTexture )
+			{
+				UGameplayStatics::PlaySound2D( Owner.GetWorld(), DropMaterialOrMaterialSound );
+			}
 		}
 	}
 
diff --git a/Engine/Source/Editor/VREditor/VREditorWorldInteraction.h b/Engine/Source/Editor/VREditor/VREditorWorldInteraction.h
--- a/Engine/Source/Editor/VREditor/VREditorWorldInteraction.h
+++ b/Engine/Source/Editor/VREditor/VREditorWorldInteraction.h
@@ -224,6 +224,9 @@ protected:
 	/** Teleport sound */
 	class USoundCue* TeleportSound;
 
+	/** Sound for dropping materials and textures */
+	class USoundCue* DropMaterialOrMaterialSound;
+
 	/** The distance when starting dragging docks */
 	float DockSelectDistance;
 
