VR Editor: Dragging UI State

diff --git a/Engine/Source/Editor/VREditor/VREditorUISystem.cpp b/Engine/Source/Editor/VREditor/VREditorUISystem.cpp
--- a/Engine/Source/Editor/VREditor/VREditorUISystem.cpp
+++ b/Engine/Source/Editor/VREditor/VREditorUISystem.cpp
@@ -55,7 +55,7 @@ namespace VREd
 	static FAutoConsoleVariable UIAssetEditorSummonedOnHandHapticFeedbackStrength( TEXT( "VREd.UIAssetEditorSummonedOnHandHapticFeedbackStrength" ), 1.0f, TEXT( "Strenth of haptic to play to remind a user which hand an asset editor was spawned on" ) );
 	static FAutoConsoleVariable MaxDockWindowSize( TEXT( "VREd.MaxDockWindowSize" ), 250, TEXT( "Maximum size for dockable windows" ) );
 	static FAutoConsoleVariable MinDockWindowSize( TEXT( "VREd.MinDockWindowSize" ), 40, TEXT( "Minimum size for dockable windows" ) );
-	static FAutoConsoleVariable DockSnapTime( TEXT( "VREd.DockSnapTime" ), 0.3f, TEXT( "Lerp time for snapping dock to hand and back" ) );
+	static FAutoConsoleVariable DockSnapTime( TEXT( "VREd.DockSnapTime" ), 0.1f, TEXT( "Lerp time for snapping dock to hand and back" ) );
 	static FAutoConsoleVariable DockDragSpeed( TEXT( "VREd.DockDragSpeed" ), 0.05f, TEXT( "Dock dragging time delay" ) );
 
 	// Tutorial UI commands
@@ -1176,15 +1176,16 @@ void FVREditorUISystem::StopDraggingDockUI()
 {
 	DraggingUI = nullptr;
 	FVirtualHand& Hand = Owner.GetVirtualHand( DraggingUIHandIndex );
-	Hand.DraggingMode = EVREditorDraggingMode::Nothing;
+	Hand.DraggingMode = EVREditorDraggingMode::Nothing;	
 	DraggingUIHandIndex = INDEX_NONE;
 
 	UGameplayStatics::PlaySound2D( Owner.GetWorld(), StopDragUISound );
+
 }
 
 bool FVREditorUISystem::IsDraggingDockUI()
 {
-	return DraggingUI != nullptr &&  DraggingUI->GetDockedTo() == AVREditorFloatingUI::EDockedTo::Nothing;
+	return DraggingUI != nullptr && DraggingUI->GetDockedTo() == AVREditorFloatingUI::EDockedTo::Dragging;
 }
 
 void FVREditorUISystem::TogglePanelsVisibility()
