VR Editor: Set default layout when pressing in the middle of the trackpad or analog stick the first time

diff --git a/Engine/Source/Editor/VREditor/VREditorUISystem.cpp b/Engine/Source/Editor/VREditor/VREditorUISystem.cpp
--- a/Engine/Source/Editor/VREditor/VREditorUISystem.cpp
+++ b/Engine/Source/Editor/VREditor/VREditorUISystem.cpp
@@ -67,6 +67,11 @@ namespace VREd
 	static FAutoConsoleVariable TutorialUILocationX( TEXT( "VREd.TutorialUI.Location.X" ), 0, TEXT( "The X location for the tutorial UI panel" ) );
 	static FAutoConsoleVariable TutorialUILocationY( TEXT( "VREd.TutorialUI.Location.Y" ), 200, TEXT( "The Y location for the tutorial UI panel" ) );
 	static FAutoConsoleVariable TutorialUILocationZ( TEXT( "VREd.TutorialUI.Location.Z" ), 40, TEXT( "The Z location for the tutorial UI panel" ) );
+
+	static const FTransform DefaultContentBrowserTransform( FRotator( 50, 177, 0 ), FVector( 25, 0, 75 ), FVector( 1 ) );
+	static const FTransform DefaultWorldOutlinerTransform( FRotator( 10, -153, 0 ), FVector( 73, 67, 154 ), FVector( 1 ) );
+	static const FTransform DefaultActorDetailsTransform( FRotator( 30, -153, 0 ), FVector( 53, 60, 104 ), FVector( 1 ) );
+	static const FTransform DefaultModesTransform( FRotator( 30, 153, 0 ), FVector( 50, -60, 107 ), FVector( 1 ) );
 }
 
 
@@ -81,7 +86,8 @@ FVREditorUISystem::FVREditorUISystem( FVREditorMode& InitOwner )
 	  DraggingUIHandIndex( INDEX_NONE ),
 	  DraggingUIOffsetTransform( FTransform::Identity ),
 	  RadialMenuHideDelayTime( 0.0f ),
-	  bRefocusViewport(false)
+	  bRefocusViewport(false),
+	  bSetDefaultLayout( true )
 {
 	// Register to find out about VR events
 	Owner.OnVRAction().AddRaw( this, &FVREditorUISystem::OnVRAction );
@@ -89,6 +95,12 @@ FVREditorUISystem::FVREditorUISystem( FVREditorMode& InitOwner )
 
 	EditorUIPanels.SetNumZeroed( (int32)EEditorUIPanel::TotalCount );
 	
+	// Set default layout transform in correct order of enum
+	DefaultWindowTransforms.Add( VREd::DefaultContentBrowserTransform );
+	DefaultWindowTransforms.Add( VREd::DefaultWorldOutlinerTransform );
+	DefaultWindowTransforms.Add( VREd::DefaultActorDetailsTransform );
+	DefaultWindowTransforms.Add( VREd::DefaultModesTransform );
+
 	QuickMenuWidgetClass = LoadClass<UVREditorQuickMenu>( nullptr, TEXT( "/Engine/VREditor/UI/VRQuickMenu.VRQuickMenu_C" ) );
 	check( QuickMenuWidgetClass != nullptr );
 
@@ -124,6 +136,8 @@ FVREditorUISystem::~FVREditorUISystem()
 
 	CleanUpActorsBeforeMapChangeOrSimulate();
 
+	DefaultWindowTransforms.Empty();
+
 	QuickMenuWidgetClass = nullptr;
 	QuickRadialWidgetClass = nullptr;
 	TutorialWidgetClass = nullptr;
@@ -1214,41 +1228,49 @@ void FVREditorUISystem::TogglePanelsVisibility()
 	
 	UGameplayStatics::PlaySound2D( Owner.GetWorld(), bShowUI ? ShowUISound : HideUISound );
 
-	for (AVREditorFloatingUI* Panel : EditorUIPanels)
+	if (!bAnyPanelsVisible && bSetDefaultLayout)
+	{
+		bSetDefaultLayout = false;
+		SetDefaultWindowLayout();
+	}
+	else
 	{
-		if (Panel != nullptr && Panel->IsUIVisible() != bShowUI)
+		for (AVREditorFloatingUI* Panel : EditorUIPanels)
 		{
-			if (bShowUI)
+			if (Panel != nullptr && Panel->IsUIVisible() != bShowUI)
 			{
-				bool SetNewVisibility = true;
-				const AVREditorFloatingUI::EDockedTo DockedTo = Panel->GetDockedTo();
-				if (DockedTo == AVREditorFloatingUI::EDockedTo::LeftHand || DockedTo == AVREditorFloatingUI::EDockedTo::RightHand)
+				if (bShowUI)
 				{
-					if (!PanelOnHand)
+					bool SetNewVisibility = true;
+					const AVREditorFloatingUI::EDockedTo DockedTo = Panel->GetDockedTo();
+					if (DockedTo == AVREditorFloatingUI::EDockedTo::LeftHand || DockedTo == AVREditorFloatingUI::EDockedTo::RightHand)
 					{
-						const uint32 HandIndex = Panel->GetDockedTo() == AVREditorFloatingUI::EDockedTo::LeftHand ? VREditorConstants::LeftHandIndex : VREditorConstants::RightHandIndex;
-						Owner.GetVirtualHand( HandIndex ).bHasUIInFront = true;
-						PanelOnHand = Panel;
+						if (!PanelOnHand)
+						{
+							const uint32 HandIndex = Panel->GetDockedTo() == AVREditorFloatingUI::EDockedTo::LeftHand ? VREditorConstants::LeftHandIndex : VREditorConstants::RightHandIndex;
+							Owner.GetVirtualHand( HandIndex ).bHasUIInFront = true;
+							PanelOnHand = Panel;
+						}
+						else
+						{
+							SetNewVisibility = false;
+						}
 					}
-					else
+					else if (DockedTo == AVREditorFloatingUI::EDockedTo::Nothing)
 					{
 						SetNewVisibility = false;
 					}
-				}
-				else if (DockedTo == AVREditorFloatingUI::EDockedTo::Nothing)
-				{
-					SetNewVisibility = false;
-				}
 
-				if (SetNewVisibility)
+					if (SetNewVisibility)
+					{
+						Panel->ShowUI( true );
+					}
+				}
+				else
 				{
-					Panel->ShowUI( true );
+					Panel->ShowUI( false );
 				}
 			}
-			else
-			{
-				Panel->ShowUI( false );
-			}
 		}
 	}
 
@@ -1324,3 +1346,20 @@ void FVREditorUISystem::TogglePanelVisibility( const EEditorUIPanel EditorUIPane
 		}
 	}
 }
+
+void FVREditorUISystem::SetDefaultWindowLayout()
+{
+	for (int32 PanelIndex = 0; PanelIndex < DefaultWindowTransforms.Num(); ++PanelIndex)
+	{
+		AVREditorFloatingUI* Panel = EditorUIPanels[PanelIndex];
+		if (Panel)
+		{
+			Panel->SetDockedTo( AVREditorFloatingUI::EDockedTo::Room );
+			Panel->ShowUI( true );
+
+			const FTransform NewTransform = DefaultWindowTransforms[PanelIndex];
+			Panel->SetLocalRotation( NewTransform.GetRotation().Rotator() );
+			Panel->SetRelativeOffset( NewTransform.GetTranslation() );
+		}
+	}
+}
\ No newline at end of file
diff --git a/Engine/Source/Editor/VREditor/VREditorUISystem.h b/Engine/Source/Editor/VREditor/VREditorUISystem.h
--- a/Engine/Source/Editor/VREditor/VREditorUISystem.h
+++ b/Engine/Source/Editor/VREditor/VREditorUISystem.h
@@ -136,6 +136,9 @@ protected:
 	/** Called when an asset editor opens an asset while in VR Editor Mode. */
 	void OnAssetEditorOpened(UObject* Asset);
 
+	/** Sets the main windows to their default transform */
+	void SetDefaultWindowLayout();
+
 protected:
 
 	/** Owning object */
@@ -210,5 +213,11 @@ protected:
 
 	/** The last dragged hover location by the laser */
 	FVector LastDraggingHoverLocation;
+
+	/** Default transforms */
+	TArray<FTransform> DefaultWindowTransforms;
+
+	/** If this is the first time using TogglePanelsVisibility */
+	bool bSetDefaultLayout;
 };
 
