VR Editor: Fixed movement grid not drawing at tracking space floor when using Oculus

diff --git a/Engine/Source/Editor/VREditor/VREditorMode.cpp b/Engine/Source/Editor/VREditor/VREditorMode.cpp
--- a/Engine/Source/Editor/VREditor/VREditorMode.cpp
+++ b/Engine/Source/Editor/VREditor/VREditorMode.cpp
@@ -11,7 +11,6 @@
 #include "CameraController.h"
 #include "DynamicMeshBuilder.h"
 #include "Features/IModularFeatures.h"
-#include "IHeadMountedDisplay.h"
 #include "LevelEditor.h"
 #include "SLevelViewport.h"
 #include "IMotionController.h"
@@ -372,6 +371,12 @@ void FVREditorMode::Enter()
 			SavedEditorState.bOnScreenMessages = GAreScreenMessagesEnabled;
 			GAreScreenMessagesEnabled = false;
 
+			if( bActuallyUsingVR )
+			{
+				SavedEditorState.TrackingOrigin = GEngine->HMDDevice->GetTrackingOrigin();
+				GEngine->HMDDevice->SetTrackingOrigin( IHeadMountedDisplay::Floor );
+			}
+
 			// Make the new viewport the active level editing viewport right away
 			GCurrentLevelEditingViewportClient = &VRViewportClient;
 		}
@@ -519,6 +524,11 @@ void FVREditorMode::Exit()
 
 				GNearClippingPlane = SavedEditorState.NearClipPlane;
 				GAreScreenMessagesEnabled = SavedEditorState.bOnScreenMessages;
+
+				if( bActuallyUsingVR )
+				{
+					GEngine->HMDDevice->SetTrackingOrigin( SavedEditorState.TrackingOrigin );
+				}
 			}
 
 			if( bActuallyUsingVR )
diff --git a/Engine/Source/Editor/VREditor/VREditorMode.h b/Engine/Source/Editor/VREditor/VREditorMode.h
--- a/Engine/Source/Editor/VREditor/VREditorMode.h
+++ b/Engine/Source/Editor/VREditor/VREditorMode.h
@@ -6,6 +6,7 @@
 #include "IVREditorModule.h"
 #include "VirtualHand.h"
 #include "HeadMountedDisplayTypes.h"	// For EHMDDeviceType::Type
+#include "IHeadMountedDisplay.h"	// For IHeadMountedDisplay::ETrackingOrigin
 #include "VREditorInputProcessor.h"
 
 
@@ -375,6 +376,7 @@ protected:
 		bool bRealTime;
 		float DragTriggerDistance;
 		bool bOnScreenMessages;
+		IHeadMountedDisplay::ETrackingOrigin TrackingOrigin;
 
 		FSavedEditorState()
 			: ViewportType( LVT_Perspective ),
@@ -387,7 +389,8 @@ protected:
 			  NearClipPlane( 0.0f ),
 			  bRealTime( false ),
 			  DragTriggerDistance( 0.0f ),
-			  bOnScreenMessages( false )
+			  bOnScreenMessages( false ),
+			  TrackingOrigin( IHeadMountedDisplay::Eye )
 		{
 		}
 		
