VR Editor: Slightly increased joystick threshold before radial menu pops up on Oculus

diff --git a/Engine/Source/Editor/VREditor/VREditorRadialMenu.cpp b/Engine/Source/Editor/VREditor/VREditorRadialMenu.cpp
--- a/Engine/Source/Editor/VREditor/VREditorRadialMenu.cpp
+++ b/Engine/Source/Editor/VREditor/VREditorRadialMenu.cpp
@@ -16,7 +16,7 @@
 
 namespace VREd
 {
-	static FAutoConsoleVariable RadialMenuRiftSelectionRadius( TEXT( "VREd.RadialMenuRiftSelectionRadius" ), 0.25f, TEXT( "The radius that is used for selecting buttons on the radial menu on Oculus Rift" ) );
+	static FAutoConsoleVariable RadialMenuRiftSelectionRadius( TEXT( "VREd.RadialMenuRiftSelectionRadius" ), 0.3f, TEXT( "The radius that is used for selecting buttons on the radial menu on Oculus Rift" ) );
 	static FAutoConsoleVariable RadialMenuViveSelectionRadius( TEXT( "VREd.RadialMenuViveSelectionRadius" ), 0.3f, TEXT( "The radius that is used for selecting buttons on the radial menu on Vive" ) );
 }
 
