VR Editor: Fixed bug that caused a panel to always spawn back to a hand when draggging content from that panel

diff --git a/Engine/Source/Editor/VREditor/VREditorWorldInteraction.cpp b/Engine/Source/Editor/VREditor/VREditorWorldInteraction.cpp
--- a/Engine/Source/Editor/VREditor/VREditorWorldInteraction.cpp
+++ b/Engine/Source/Editor/VREditor/VREditorWorldInteraction.cpp
@@ -610,7 +610,7 @@ void FVREditorWorldInteraction::OnVRAction( FEditorViewportClient& ViewportClien
 					if( Hand.DraggingMode == EVREditorDraggingMode::ActorsAtLaserImpact ||
 						Hand.DraggingMode == EVREditorDraggingMode::Material )
 					{
-						Owner.GetUISystem().ShowEditorUIPanel( FloatingUIAssetDraggedFrom, Owner.GetOtherHandIndex( VRAction.HandIndex ), true, true, false, false );
+						Owner.GetUISystem().ShowEditorUIPanel( FloatingUIAssetDraggedFrom, Owner.GetOtherHandIndex( VRAction.HandIndex ), true, false, false, false );
 						FloatingUIAssetDraggedFrom = nullptr;
 					}
 
