VR Editor: Fixed initial location of UI panels when clicking the summon UI button

diff --git a/Engine/Source/Editor/VREditor/VREditorUISystem.cpp b/Engine/Source/Editor/VREditor/VREditorUISystem.cpp
--- a/Engine/Source/Editor/VREditor/VREditorUISystem.cpp
+++ b/Engine/Source/Editor/VREditor/VREditorUISystem.cpp
@@ -68,10 +68,10 @@ namespace VREd
 	static FAutoConsoleVariable TutorialUILocationY( TEXT( "VREd.TutorialUI.Location.Y" ), 200, TEXT( "The Y location for the tutorial UI panel" ) );
 	static FAutoConsoleVariable TutorialUILocationZ( TEXT( "VREd.TutorialUI.Location.Z" ), 40, TEXT( "The Z location for the tutorial UI panel" ) );
 
-	static const FTransform DefaultContentBrowserTransform( FRotator( 50, 177, 0 ), FVector( 25, 0, 75 ), FVector( 1 ) );
-	static const FTransform DefaultWorldOutlinerTransform( FRotator( 10, -153, 0 ), FVector( 73, 67, 154 ), FVector( 1 ) );
-	static const FTransform DefaultActorDetailsTransform( FRotator( 30, -153, 0 ), FVector( 53, 60, 104 ), FVector( 1 ) );
-	static const FTransform DefaultModesTransform( FRotator( 30, 153, 0 ), FVector( 50, -60, 107 ), FVector( 1 ) );
+	static const FTransform DefaultModesTransform(			FRotator( -20,  153, 0 ), FVector( 40, -50,  40 ), FVector( 1 ) );	// Top left
+	static const FTransform DefaultWorldOutlinerTransform(	FRotator( -20, -153, 0 ), FVector( 40,  50,  40 ), FVector( 1 ) );	// Top right
+	static const FTransform DefaultContentBrowserTransform( FRotator(  30,  153, 0 ), FVector( 40, -50, -40 ), FVector( 1 ) );  // Bottom left
+	static const FTransform DefaultActorDetailsTransform(	FRotator(  30, -153, 0 ), FVector( 40,  50, -40 ), FVector( 1 ) );	// Bottom right
 }
 
 
@@ -1357,9 +1357,16 @@ void FVREditorUISystem::SetDefaultWindowLayout()
 			Panel->SetDockedTo( AVREditorFloatingUI::EDockedTo::Room );
 			Panel->ShowUI( true );
 
-			const FTransform NewTransform = DefaultWindowTransforms[PanelIndex];
+			// Make sure the UIs are centered around the direction your head is looking (yaw only!)
+			const FVector RoomSpaceHeadLocation = Owner.GetRoomSpaceHeadTransform().GetLocation();
+			FRotator RoomSpaceHeadYawRotator = Owner.GetRoomSpaceHeadTransform().GetRotation().Rotator();
+			RoomSpaceHeadYawRotator.Pitch = 0.0f;
+			RoomSpaceHeadYawRotator.Roll = 0.0f;
+
+			FTransform NewTransform = DefaultWindowTransforms[PanelIndex];
+			NewTransform *= FTransform( RoomSpaceHeadYawRotator.Quaternion(), FVector::ZeroVector );
 			Panel->SetLocalRotation( NewTransform.GetRotation().Rotator() );
-			Panel->SetRelativeOffset( NewTransform.GetTranslation() );
+			Panel->SetRelativeOffset( RoomSpaceHeadLocation + NewTransform.GetTranslation() );
 		}
 	}
 }
\ No newline at end of file
