VR Editor: UI positioning tweaks

diff --git a/Engine/Source/Editor/VREditor/VREditorUISystem.cpp b/Engine/Source/Editor/VREditor/VREditorUISystem.cpp
--- a/Engine/Source/Editor/VREditor/VREditorUISystem.cpp
+++ b/Engine/Source/Editor/VREditor/VREditorUISystem.cpp
@@ -68,10 +68,10 @@ namespace VREd
 	static FAutoConsoleVariable TutorialUILocationY( TEXT( "VREd.TutorialUI.Location.Y" ), 200, TEXT( "The Y location for the tutorial UI panel" ) );
 	static FAutoConsoleVariable TutorialUILocationZ( TEXT( "VREd.TutorialUI.Location.Z" ), 40, TEXT( "The Z location for the tutorial UI panel" ) );
 
-	static const FTransform DefaultModesTransform(			FRotator( -20,  153, 0 ), FVector( 40, -50,  40 ), FVector( 1 ) );	// Top left
-	static const FTransform DefaultWorldOutlinerTransform(	FRotator( -20, -153, 0 ), FVector( 40,  50,  40 ), FVector( 1 ) );	// Top right
-	static const FTransform DefaultContentBrowserTransform( FRotator(  30,  153, 0 ), FVector( 40, -50, -40 ), FVector( 1 ) );  // Bottom left
-	static const FTransform DefaultActorDetailsTransform(	FRotator(  30, -153, 0 ), FVector( 40,  50, -40 ), FVector( 1 ) );	// Bottom right
+	static const FTransform DefaultModesTransform(			FRotator( -30,  153, 0 ), FVector( 40, -45,  40 ), FVector( 1 ) );	// Top left
+	static const FTransform DefaultWorldOutlinerTransform(	FRotator( -30, -153, 0 ), FVector( 40,  45,  40 ), FVector( 1 ) );	// Top right
+	static const FTransform DefaultContentBrowserTransform( FRotator(  20,  153, 0 ), FVector( 50, -50, -40 ), FVector( 1 ) );  // Bottom left
+	static const FTransform DefaultActorDetailsTransform(	FRotator(  20, -153, 0 ), FVector( 50,  50, -40 ), FVector( 1 ) );	// Bottom right
 }
 
 
