VR Editor: - Fix for crash when hovering over bsp
- Dragging materials/textures laser color

diff --git a/Engine/Source/Editor/VREditor/VREditorMode.cpp b/Engine/Source/Editor/VREditor/VREditorMode.cpp
--- a/Engine/Source/Editor/VREditor/VREditorMode.cpp
+++ b/Engine/Source/Editor/VREditor/VREditorMode.cpp
@@ -1201,6 +1201,7 @@ void FVREditorMode::Tick( FEditorViewportClient* ViewportClient, float DeltaTime
 				}
 			}
 			else if( Hand.DraggingMode == EVREditorDraggingMode::ActorsAtLaserImpact ||
+				Hand.DraggingMode == EVREditorDraggingMode::Material ||
 				Hand.DraggingMode == EVREditorDraggingMode::ActorsFreely ||
 				Hand.DraggingMode == EVREditorDraggingMode::ActorsWithGizmo ||
 				Hand.DraggingMode == EVREditorDraggingMode::AssistingDrag || 
@@ -1779,7 +1780,7 @@ FHitResult FVREditorMode::GetHitResultFromLaserPointer( int32 HandIndex, TArray<
 				// ECC_EditorGizmo (some gizmo handles are opaque and use ECC_Visibility as their collision channel)
 				// NOTE: We're treating components of floating UI actors as "gizmos" for the purpose of hit testing as long as the component is not the actual UI widget component
 				const bool bHitResultIsUI = UISystem->IsWidgetAnEditorUIWidget( HitResult.GetComponent() );
-				const bool bHitResultIsGizmo = ( HitResult.GetActor() == WorldInteraction->GetTransformGizmoActor() ) || ( !bHitResultIsUI && HitResult.GetActor()->IsA( AVREditorFloatingUI::StaticClass() ) );
+				const bool bHitResultIsGizmo = HitResult.GetActor() != nullptr && ( ( HitResult.GetActor() == WorldInteraction->GetTransformGizmoActor() ) || ( !bHitResultIsUI && HitResult.GetActor()->IsA( AVREditorFloatingUI::StaticClass() ) ) );
 				const bool bBestHitResultSoFarIsGizmo = BestHitResult.GetActor() != nullptr && ( ( BestHitResult.GetActor() == WorldInteraction->GetTransformGizmoActor() ) || ( !bBestHitResultSoFarIsUI && BestHitResult.GetActor()->IsA( AVREditorFloatingUI::StaticClass() ) ) );
 				if( BestHitResult.GetActor() == nullptr ||	// Don't have anything yet?
 					( bHitResultIsUI && bBestHitResultSoFarIsGizmo ) ||
