VR Editor: Added command for showing/hiding the movement grid

diff --git a/Engine/Source/Editor/VREditor/VREditorMode.cpp b/Engine/Source/Editor/VREditor/VREditorMode.cpp
--- a/Engine/Source/Editor/VREditor/VREditorMode.cpp
+++ b/Engine/Source/Editor/VREditor/VREditorMode.cpp
@@ -63,6 +63,7 @@ namespace VREd
 
 	static FAutoConsoleVariable InvertTrackpadVertical( TEXT( "VREd.InvertTrackpadVertical" ), 1, TEXT( "Toggles inverting the touch pad vertical axis" ) );
 	static FAutoConsoleVariable ForceOculusMirrorMode( TEXT( "VREd.ForceOculusMirrorMode" ), 3, TEXT( "Which Oculus display mirroring mode to use (see MirrorWindowModeType in OculusRiftHMD.h)" ) );
+	static FAutoConsoleVariable ShowMovementGrid( TEXT( "VREd.ShowMovementGrid" ), 1, TEXT( "Showing the ground movement grid" ) );
 }
 
 FEditorModeID FVREditorMode::VREditorModeID( "VREditor" );
@@ -1216,50 +1217,57 @@ void FVREditorMode::Tick( FEditorViewportClient* ViewportClient, float DeltaTime
 
 	// Updating the opacity and visibility of the grid according to the controllers
 	{
-		// Show the grid full opacity when dragging or scaling
-		float GoalOpacity = 0.f;
-		if ( GetVirtualHand( VREditorConstants::LeftHandIndex ).DraggingMode == EVREditorDraggingMode::World || GetVirtualHand( VREditorConstants::RightHandIndex ).DraggingMode == EVREditorDraggingMode::World )
+		if( VREd::ShowMovementGrid->GetInt() == 1)
 		{
-			GoalOpacity = 1.0f;
-		}
-		else if ( ( GetVirtualHand( VREditorConstants::LeftHandIndex ).LastDraggingMode == EVREditorDraggingMode::World && !GetVirtualHand( VREditorConstants::LeftHandIndex ).DragTranslationVelocity.IsNearlyZero( VREd::GridMovementTolerance->GetFloat() ) ) )
-		{
-			GoalOpacity = FMath::Clamp( GetVirtualHand( VREditorConstants::LeftHandIndex ).DragTranslationVelocity.Size() / VREd::GridFadeStartVelocity->GetFloat(), 0.0f, 1.0f );
-		}
-		else if( ( GetVirtualHand( VREditorConstants::RightHandIndex ).LastDraggingMode == EVREditorDraggingMode::World && !GetVirtualHand( VREditorConstants::RightHandIndex ).DragTranslationVelocity.IsNearlyZero( VREd::GridMovementTolerance->GetFloat() ) ) )
-		{
-			GoalOpacity = FMath::Clamp( GetVirtualHand( VREditorConstants::RightHandIndex ).DragTranslationVelocity.Size() / VREd::GridFadeStartVelocity->GetFloat(), 0.0f, 1.0f );
-		}
-
-		// Check the current opacity and add or subtract to reach the goal
-		float NewOpacity = WorldMovementGridOpacity;
-		if( NewOpacity < GoalOpacity )
-		{
-			NewOpacity = FMath::Min( GoalOpacity, NewOpacity + DeltaTime * VREd::GridFadeMultiplier->GetFloat() );
-		}
-		else if( NewOpacity > GoalOpacity )
-		{
-			NewOpacity = FMath::Max( GoalOpacity, NewOpacity - DeltaTime * VREd::GridFadeMultiplier->GetFloat() );
-		}
-
-		// Only update when the opacity is different
-		if( !FMath::IsNearlyEqual( WorldMovementGridOpacity, GoalOpacity ) )
-		{
-			WorldMovementGridOpacity = NewOpacity;
+			// Show the grid full opacity when dragging or scaling
+			float GoalOpacity = 0.f;
+			if ( GetVirtualHand( VREditorConstants::LeftHandIndex ).DraggingMode == EVREditorDraggingMode::World || GetVirtualHand( VREditorConstants::RightHandIndex ).DraggingMode == EVREditorDraggingMode::World )
+			{
+				GoalOpacity = 1.0f;
+			}
+			else if ( ( GetVirtualHand( VREditorConstants::LeftHandIndex ).LastDraggingMode == EVREditorDraggingMode::World && !GetVirtualHand( VREditorConstants::LeftHandIndex ).DragTranslationVelocity.IsNearlyZero( VREd::GridMovementTolerance->GetFloat() ) ) )
+			{
+				GoalOpacity = FMath::Clamp( GetVirtualHand( VREditorConstants::LeftHandIndex ).DragTranslationVelocity.Size() / VREd::GridFadeStartVelocity->GetFloat(), 0.0f, 1.0f );
+			}
+			else if( ( GetVirtualHand( VREditorConstants::RightHandIndex ).LastDraggingMode == EVREditorDraggingMode::World && !GetVirtualHand( VREditorConstants::RightHandIndex ).DragTranslationVelocity.IsNearlyZero( VREd::GridMovementTolerance->GetFloat() ) ) )
+			{
+				GoalOpacity = FMath::Clamp( GetVirtualHand( VREditorConstants::RightHandIndex ).DragTranslationVelocity.Size() / VREd::GridFadeStartVelocity->GetFloat(), 0.0f, 1.0f );
+			}
 
-			// See if the opacity is almost zero and if the goal opacity is lower than the new opacity set it to zero if so. Otherwise it will flicker
-			if( WorldMovementGridOpacity < SMALL_NUMBER )
+			// Check the current opacity and add or subtract to reach the goal
+			float NewOpacity = WorldMovementGridOpacity;
+			if( NewOpacity < GoalOpacity )
 			{
-				WorldMovementGridOpacity = 0.f;
-				WorldMovementGridMeshComponent->SetVisibility( false );
+				NewOpacity = FMath::Min( GoalOpacity, NewOpacity + DeltaTime * VREd::GridFadeMultiplier->GetFloat() );
 			}
-			else
+			else if( NewOpacity > GoalOpacity )
 			{
-				WorldMovementGridMeshComponent->SetVisibility( true );
+				NewOpacity = FMath::Max( GoalOpacity, NewOpacity - DeltaTime * VREd::GridFadeMultiplier->GetFloat() );
 			}
 
-			static const FName OpacityName( "OpacityParam" );
-			WorldMovementGridMID->SetScalarParameterValue( OpacityName, WorldMovementGridOpacity * VREd::GridMaxOpacity->GetFloat() );
+			// Only update when the opacity is different
+			if( !FMath::IsNearlyEqual( WorldMovementGridOpacity, GoalOpacity ) )
+			{
+				WorldMovementGridOpacity = NewOpacity;
+
+				// See if the opacity is almost zero and if the goal opacity is lower than the new opacity set it to zero if so. Otherwise it will flicker
+				if( WorldMovementGridOpacity < SMALL_NUMBER )
+				{
+					WorldMovementGridOpacity = 0.f;
+					WorldMovementGridMeshComponent->SetVisibility( false );
+				}
+				else
+				{
+					WorldMovementGridMeshComponent->SetVisibility( true );
+				}
+
+				static const FName OpacityName( "OpacityParam" );
+				WorldMovementGridMID->SetScalarParameterValue( OpacityName, WorldMovementGridOpacity * VREd::GridMaxOpacity->GetFloat() );
+			}
+		}
+		else
+		{
+			WorldMovementGridMeshComponent->SetVisibility( false );
 		}
 	}
 
