VR Editor: Added analytics for VREditor being activated

diff --git a/Engine/Source/Editor/VREditor/VREditorMode.cpp b/Engine/Source/Editor/VREditor/VREditorMode.cpp
--- a/Engine/Source/Editor/VREditor/VREditorMode.cpp
+++ b/Engine/Source/Editor/VREditor/VREditorMode.cpp
@@ -15,6 +15,7 @@
 #include "SLevelViewport.h"
 #include "IMotionController.h"
 #include "MotionControllerComponent.h"
+#include "EngineAnalytics.h"
 
 #include "Editor/LevelEditor/Public/LevelEditorActions.h"
 
@@ -143,6 +144,12 @@ void FVREditorMode::Enter()
 	bIsDrawingWorldMovementPostProcess = false;
 	LastFrameNumberInputWasPolled = 0;
 
+	// Take note of VREditor activation
+	if( FEngineAnalytics::IsAvailable() )
+	{
+		FEngineAnalytics::GetProvider().RecordEvent( TEXT( "Editor.Usage.EnterVREditorMode" ) );
+	}
+
 	// Fully reset all hand state.  We don't want anything carrying over from the previous session.
 	for( int32 HandIndex = 0; HandIndex < VREditorConstants::NumVirtualHands; ++HandIndex )
 	{
