VR Editor: Fixed tutorial window not visible in front of you (it was to your right) initially

diff --git a/Engine/Source/Editor/VREditor/VREditorUISystem.cpp b/Engine/Source/Editor/VREditor/VREditorUISystem.cpp
--- a/Engine/Source/Editor/VREditor/VREditorUISystem.cpp
+++ b/Engine/Source/Editor/VREditor/VREditorUISystem.cpp
@@ -62,10 +62,10 @@ namespace VREd
 	static FAutoConsoleVariable TutorialUIResolutionX( TEXT( "VREd.TutorialUI.Resolution.X" ), 1920, TEXT( "The X resolution for the tutorial UI panel" ) );
 	static FAutoConsoleVariable TutorialUIResolutionY( TEXT( "VREd.TutorialUI.Resolution.Y" ), 1080, TEXT( "The Y resolution for the tutorial UI panel" ) );
 	static FAutoConsoleVariable TutorialUISize( TEXT( "VREd.TutorialUI.Size" ), 200, TEXT( "The room space size for the tutorial UI panel" ) );
-	static FAutoConsoleVariable TutorialUIYaw( TEXT( "VREd.TutorialUI.Yaw" ), 270, TEXT( "The yaw rotation for the tutorial UI panel" ) );
+	static FAutoConsoleVariable TutorialUIYaw( TEXT( "VREd.TutorialUI.Yaw" ), 180, TEXT( "The yaw rotation for the tutorial UI panel" ) );
 	static FAutoConsoleVariable TutorialUIPitch( TEXT( "VREd.TutorialUI.Pitch" ), 45, TEXT( "The pitch rotation for the tutorial UI panel" ) );
-	static FAutoConsoleVariable TutorialUILocationX( TEXT( "VREd.TutorialUI.Location.X" ), 0, TEXT( "The X location for the tutorial UI panel" ) );
-	static FAutoConsoleVariable TutorialUILocationY( TEXT( "VREd.TutorialUI.Location.Y" ), 200, TEXT( "The Y location for the tutorial UI panel" ) );
+	static FAutoConsoleVariable TutorialUILocationX( TEXT( "VREd.TutorialUI.Location.X" ), 200, TEXT( "The X location for the tutorial UI panel" ) );
+	static FAutoConsoleVariable TutorialUILocationY( TEXT( "VREd.TutorialUI.Location.Y" ), 0, TEXT( "The Y location for the tutorial UI panel" ) );
 	static FAutoConsoleVariable TutorialUILocationZ( TEXT( "VREd.TutorialUI.Location.Z" ), 40, TEXT( "The Z location for the tutorial UI panel" ) );
 
 	static const FTransform DefaultModesTransform(			FRotator( -30,  153, 0 ), FVector( 40, -45,  40 ), FVector( 1 ) );	// Top left
