VR Editor: Default layout was not set the first time because of the tutorial being visible

diff --git a/Engine/Source/Editor/VREditor/VREditorUISystem.cpp b/Engine/Source/Editor/VREditor/VREditorUISystem.cpp
--- a/Engine/Source/Editor/VREditor/VREditorUISystem.cpp
+++ b/Engine/Source/Editor/VREditor/VREditorUISystem.cpp
@@ -1228,7 +1228,7 @@ void FVREditorUISystem::TogglePanelsVisibility()
 	
 	UGameplayStatics::PlaySound2D( Owner.GetWorld(), bShowUI ? ShowUISound : HideUISound );
 
-	if (!bAnyPanelsVisible && bSetDefaultLayout)
+	if ( bSetDefaultLayout)
 	{
 		bSetDefaultLayout = false;
 		SetDefaultWindowLayout();
