VR Editor: Fixed UI initially spawning in wrong location if world scale was not 1.0 when clicking down trackpad

diff --git a/Engine/Source/Editor/VREditor/VREditorUISystem.cpp b/Engine/Source/Editor/VREditor/VREditorUISystem.cpp
--- a/Engine/Source/Editor/VREditor/VREditorUISystem.cpp
+++ b/Engine/Source/Editor/VREditor/VREditorUISystem.cpp
@@ -1358,7 +1358,7 @@ void FVREditorUISystem::SetDefaultWindowLayout()
 			Panel->ShowUI( true );
 
 			// Make sure the UIs are centered around the direction your head is looking (yaw only!)
-			const FVector RoomSpaceHeadLocation = Owner.GetRoomSpaceHeadTransform().GetLocation();
+			const FVector RoomSpaceHeadLocation = Owner.GetRoomSpaceHeadTransform().GetLocation() / Owner.GetWorldScaleFactor();
 			FRotator RoomSpaceHeadYawRotator = Owner.GetRoomSpaceHeadTransform().GetRotation().Rotator();
 			RoomSpaceHeadYawRotator.Pitch = 0.0f;
 			RoomSpaceHeadYawRotator.Roll = 0.0f;
