VR Editor: Command for showing movement post processing

diff --git a/Engine/Source/Editor/VREditor/VREditorMode.cpp b/Engine/Source/Editor/VREditor/VREditorMode.cpp
--- a/Engine/Source/Editor/VREditor/VREditorMode.cpp
+++ b/Engine/Source/Editor/VREditor/VREditorMode.cpp
@@ -64,7 +64,9 @@ namespace VREd
 
 	static FAutoConsoleVariable InvertTrackpadVertical( TEXT( "VREd.InvertTrackpadVertical" ), 1, TEXT( "Toggles inverting the touch pad vertical axis" ) );
 	static FAutoConsoleVariable ForceOculusMirrorMode( TEXT( "VREd.ForceOculusMirrorMode" ), 3, TEXT( "Which Oculus display mirroring mode to use (see MirrorWindowModeType in OculusRiftHMD.h)" ) );
+
 	static FAutoConsoleVariable ShowMovementGrid( TEXT( "VREd.ShowMovementGrid" ), 1, TEXT( "Showing the ground movement grid" ) );
+	static FAutoConsoleVariable ShowWorldMovementPostProcess( TEXT( "VREd.ShowWorldMovementPostProcess" ), 1, TEXT( "Showing the movement post processing" ) );
 }
 
 FEditorModeID FVREditorMode::VREditorModeID( "VREditor" );
@@ -870,7 +872,7 @@ void FVREditorMode::SpawnAvatarMeshActor()
 			PostProcessComponent->RegisterComponent();
 
 			// Unlimited size
-			PostProcessComponent->bEnabled = true;
+			PostProcessComponent->bEnabled = VREd::ShowWorldMovementPostProcess->GetInt() == 1 ? true : false;
 			PostProcessComponent->bUnbound = true;
 		}
 	}
@@ -1283,6 +1285,8 @@ void FVREditorMode::Tick( FEditorViewportClient* ViewportClient, float DeltaTime
 	// that are not nearby, and completely hides the skybox and other very distant objects.  This is to help
 	// prevent simulation sickness when moving/rotating/scaling the entire world around them.
 	{
+		PostProcessComponent->bEnabled = VREd::ShowWorldMovementPostProcess->GetInt() == 1 ? true : false;
+
 		// Make sure our world size is reflected in the post process material
 		static FName WorldScaleFactorParameterName( "WorldScaleFactor" );
 		WorldMovementPostProcessMaterial->SetScalarParameterValue( WorldScaleFactorParameterName, GetWorldScaleFactor() );
