VR Editor: Allow user to shrink the world down smaller (1:60 instead of 1:40)

diff --git a/Engine/Source/Editor/VREditor/VREditorMode.h b/Engine/Source/Editor/VREditor/VREditorMode.h
--- a/Engine/Source/Editor/VREditor/VREditorMode.h
+++ b/Engine/Source/Editor/VREditor/VREditorMode.h
@@ -12,7 +12,7 @@
 
 namespace VREd
 {
-	static FAutoConsoleVariable ScaleMax( TEXT( "VREd.ScaleMax" ), 4000.0f, TEXT( "Maximum world scale in centimeters" ) );
+	static FAutoConsoleVariable ScaleMax( TEXT( "VREd.ScaleMax" ), 6000.0f, TEXT( "Maximum world scale in centimeters" ) );
 	static FAutoConsoleVariable ScaleMin( TEXT( "VREd.ScaleMin" ), 10.0f, TEXT( "Minimum world scale in centimeters" ) );
 }
 
