VR Editor: Removed static commands from VREditorMode.h

diff --git a/Engine/Source/Editor/VREditor/VREditorMode.cpp b/Engine/Source/Editor/VREditor/VREditorMode.cpp
--- a/Engine/Source/Editor/VREditor/VREditorMode.cpp
+++ b/Engine/Source/Editor/VREditor/VREditorMode.cpp
@@ -67,6 +67,9 @@ namespace VREd
 
 	static FAutoConsoleVariable ShowMovementGrid( TEXT( "VREd.ShowMovementGrid" ), 1, TEXT( "Showing the ground movement grid" ) );
 	static FAutoConsoleVariable ShowWorldMovementPostProcess( TEXT( "VREd.ShowWorldMovementPostProcess" ), 1, TEXT( "Showing the movement post processing" ) );
+
+	static FAutoConsoleVariable ScaleMax( TEXT( "VREd.ScaleMax" ), 6000.0f, TEXT( "Maximum world scale in centimeters" ) );
+	static FAutoConsoleVariable ScaleMin( TEXT( "VREd.ScaleMin" ), 10.0f, TEXT( "Minimum world scale in centimeters" ) );
 }
 
 FEditorModeID FVREditorMode::VREditorModeID( "VREditor" );
@@ -1712,6 +1715,16 @@ bool FVREditorMode::GetHandTransformAndForwardVector( int32 HandIndex, FTransfor
 }
 
 
+float FVREditorMode::GetMaxScale()
+{
+	return VREd::ScaleMax->GetFloat();
+}
+
+float FVREditorMode::GetMinScale()
+{
+	return VREd::ScaleMin->GetFloat();
+}
+
 bool FVREditorMode::GetLaserPointer( const int32 HandIndex, FVector& LaserPointerStart, FVector& LaserPointerEnd, const bool bEvenIfUIIsInFront, const float LaserLengthOverride ) const
 {
 	const FVirtualHand& Hand = VirtualHands[ HandIndex ];
diff --git a/Engine/Source/Editor/VREditor/VREditorMode.h b/Engine/Source/Editor/VREditor/VREditorMode.h
--- a/Engine/Source/Editor/VREditor/VREditorMode.h
+++ b/Engine/Source/Editor/VREditor/VREditorMode.h
@@ -9,13 +9,6 @@
 #include "IHeadMountedDisplay.h"	// For IHeadMountedDisplay::ETrackingOrigin
 #include "VREditorInputProcessor.h"
 
-
-namespace VREd
-{
-	static FAutoConsoleVariable ScaleMax( TEXT( "VREd.ScaleMax" ), 6000.0f, TEXT( "Maximum world scale in centimeters" ) );
-	static FAutoConsoleVariable ScaleMin( TEXT( "VREd.ScaleMin" ), 10.0f, TEXT( "Minimum world scale in centimeters" ) );
-}
-
 // Forward declare the GizmoHandleTypes that is defined in VREditorTransformGizmo
 enum class EGizmoHandleTypes : uint8;
 
@@ -305,6 +298,12 @@ public:
 	 */
 	bool GetHandTransformAndForwardVector( int32 HandIndex, FTransform& OutHandTransform, FVector& OutForwardVector ) const;
 
+	/** Returns the maximum user scale */
+	float GetMaxScale();
+
+	/** Returns the minimum user scale */
+	float GetMinScale();
+
 private:
 
 	/** Called when someone closes a standalone VR Editor window */
diff --git a/Engine/Source/Editor/VREditor/VREditorWorldInteraction.cpp b/Engine/Source/Editor/VREditor/VREditorWorldInteraction.cpp
--- a/Engine/Source/Editor/VREditor/VREditorWorldInteraction.cpp
+++ b/Engine/Source/Editor/VREditor/VREditorWorldInteraction.cpp
@@ -2305,7 +2305,7 @@ void FVREditorWorldInteraction::UpdateDragging(
 
 				// NOTE: Instead of clamping, we simply skip changing the W2M this frame if it's out of bounds.  Clamping makes our math more complicated.
 				if( NewWorldToMetersScale != OldWorldToMetersScale &&
-					NewWorldToMetersScale >= VREd::ScaleMin->GetFloat() && NewWorldToMetersScale <= VREd::ScaleMax->GetFloat() )
+					NewWorldToMetersScale >= Owner.GetMinScale() && NewWorldToMetersScale <= Owner.GetMaxScale() )
 				{
 					// @todo vreditor: This is bad because we're clobbering the world settings which will be saved with the map.  Instead we need to 
 					// be able to apply an override before the scene view gets it
