VR Editor: Replaced haptic feedback commands for different types of hovers to one general command

diff --git a/Engine/Source/Editor/VREditor/VREditorWorldInteraction.cpp b/Engine/Source/Editor/VREditor/VREditorWorldInteraction.cpp
--- a/Engine/Source/Editor/VREditor/VREditorWorldInteraction.cpp
+++ b/Engine/Source/Editor/VREditor/VREditorWorldInteraction.cpp
@@ -52,10 +52,8 @@ namespace VREd
 	static FAutoConsoleVariable GizmoHandleHoverAnimationDuration( TEXT( "VREd.GizmoHandleHoverAnimationDuration" ), 0.1f, TEXT( "How quickly to animate gizmo handle hover state" ) );
 	static FAutoConsoleVariable SnapGridSize( TEXT( "VREd.SnapGridSize" ), 3.0f, TEXT( "How big the snap grid should be, in multiples of the gizmo bounding box size" ) );
 	static FAutoConsoleVariable SnapGridLineWidth( TEXT( "VREd.SnapGridLineWidth" ), 3.0f, TEXT( "Width of the grid lines on the snap grid" ) );
-	static FAutoConsoleVariable GizmoHoverHapticFeedbackStrength( TEXT( "VREd.GizmoHoverHapticFeedbackStrength" ), 0.1f, TEXT( "Default strength for haptic feedback when hovering over gizmos" ) );
-	static FAutoConsoleVariable GizmoHoverHapticFeedbackTime( TEXT( "VREd.GizmoHoverHapticFeedbackTime" ), 1.0f, TEXT( "The minimum time between haptic feedback between hovering on a gizmo handle" ) );
-	static FAutoConsoleVariable WindowHoverHapticFeedbackStrength( TEXT( "VREd.WindowHoverHapticFeedbackStrength" ), 0.4f, TEXT( "Width of the grid lines on the snap grid" ) );
-	static FAutoConsoleVariable WindowHoverHapticFeedbackTime( TEXT( "VREd.WindowHoverHapticFeedbackTime" ), 0.2f, TEXT( "The minimum time between haptic feedback between hovering on a dockable window" ) );
+	static FAutoConsoleVariable HoverHapticFeedbackStrength( TEXT( "VREd.HoverHapticFeedbackStrength" ), 0.1f, TEXT( "Default strength for haptic feedback when hovering" ) );
+	static FAutoConsoleVariable HoverHapticFeedbackTime( TEXT( "VREd.HoverHapticFeedbackTime" ), 0.2f, TEXT( "The minimum time between haptic feedback for hovering" ) );
 }
 
 
@@ -709,8 +707,8 @@ void FVREditorWorldInteraction::OnVRHoverUpdate( FEditorViewportClient& Viewport
 
 					if( Hand.HoveringOverTransformGizmoComponent != PreviousHoverGizmoComponent )
 					{
-						HapticFeedbackStrength = VREd::GizmoHoverHapticFeedbackStrength->GetFloat();
-						HapticFeedbackTimeBuffer = VREd::GizmoHoverHapticFeedbackTime->GetFloat();
+						HapticFeedbackStrength = VREd::HoverHapticFeedbackStrength->GetFloat();
+						HapticFeedbackTimeBuffer = VREd::HoverHapticFeedbackTime->GetFloat();
 					}
 
 					Hand.HoverHapticCheckLastHoveredGizmoComponent = HitResult.GetComponent();
@@ -726,8 +724,8 @@ void FVREditorWorldInteraction::OnVRHoverUpdate( FEditorViewportClient& Viewport
 							DockableWindow->OnEnterHover( HitResult, HandIndex );
 						}
 						
-						HapticFeedbackStrength = VREd::WindowHoverHapticFeedbackStrength->GetFloat();
-						HapticFeedbackTimeBuffer = VREd::WindowHoverHapticFeedbackTime->GetFloat();
+						HapticFeedbackStrength = VREd::HoverHapticFeedbackStrength->GetFloat();
+						HapticFeedbackTimeBuffer = VREd::HoverHapticFeedbackTime->GetFloat();
 					}
 				}
 
