VR Editor: Fixed lighting not being properly invalidated when dragging actors in VR
- This fixes shadows not drawing on static objects lit by stationary lights after dragging them

diff --git a/Engine/Source/Editor/VREditor/VREditorWorldInteraction.cpp b/Engine/Source/Editor/VREditor/VREditorWorldInteraction.cpp
--- a/Engine/Source/Editor/VREditor/VREditorWorldInteraction.cpp
+++ b/Engine/Source/Editor/VREditor/VREditorWorldInteraction.cpp
@@ -1397,11 +1397,18 @@ void FVREditorWorldInteraction::Tick( FEditorViewportClient* ViewportClient, con
 				AActor* Actor = Cast<AActor>( Transformable.Object.Get() );
 				if( Actor != nullptr )
 				{
-					if( !Actor->GetActorTransform().Equals( InterpolatedTransform, 0.0f ) )
+					const FTransform& ExistingTransform = Actor->GetActorTransform();
+					if( !ExistingTransform.Equals( InterpolatedTransform, 0.0f ) )
 					{
+						const bool bOnlyTranslationChanged =
+							ExistingTransform.GetRotation() == InterpolatedTransform.GetRotation() &&
+							ExistingTransform.GetScale3D() == InterpolatedTransform.GetScale3D();
+
 						Actor->SetActorTransform( InterpolatedTransform, bSweep );
 						//GWarn->Logf( TEXT( "SMOOTH: Actor %s to %s" ), *Actor->GetName(), *Transformable.TargetTransform.ToString() );
 
+						Actor->InvalidateLightingCacheDetailed( bOnlyTranslationChanged );
+
 						const bool bFinished = false;
 						Actor->PostEditMove( bFinished );
 					}
@@ -2399,8 +2406,16 @@ void FVREditorWorldInteraction::UpdateDragging(
 				AActor* Actor = Cast<AActor>( Transformable.Object.Get() );
 				if( Actor != nullptr )
 				{
+					const FTransform& ExistingTransform = Actor->GetTransform();
+
+					const bool bOnlyTranslationChanged =
+						ExistingTransform.GetRotation() == Transformable.TargetTransform.GetRotation() &&
+						ExistingTransform.GetScale3D() == Transformable.TargetTransform.GetScale3D();
+
 					Actor->SetActorTransform( Transformable.TargetTransform, bSweep );
 
+					Actor->InvalidateLightingCacheDetailed( bOnlyTranslationChanged );
+
 					const bool bFinished = false;
 					Actor->PostEditMove( bFinished );
 				}
