VR Editor: Fix for Laser goes away when hiding all windows with a window docked to controller (UE-28340)

diff --git a/Engine/Source/Editor/VREditor/VREditorUISystem.cpp b/Engine/Source/Editor/VREditor/VREditorUISystem.cpp
--- a/Engine/Source/Editor/VREditor/VREditorUISystem.cpp
+++ b/Engine/Source/Editor/VREditor/VREditorUISystem.cpp
@@ -1354,6 +1354,13 @@ void FVREditorUISystem::SetDefaultWindowLayout()
 		AVREditorFloatingUI* Panel = EditorUIPanels[PanelIndex];
 		if (Panel)
 		{
+			const AVREditorFloatingUI::EDockedTo DockedTo = Panel->GetDockedTo();
+			if (DockedTo == AVREditorFloatingUI::EDockedTo::LeftHand || DockedTo == AVREditorFloatingUI::EDockedTo::RightHand)
+			{
+				const int32 HandIndex = DockedTo == AVREditorFloatingUI::EDockedTo::LeftHand ? VREditorConstants::LeftHandIndex : VREditorConstants::RightHandIndex;
+				Owner.GetVirtualHand( HandIndex ).bHasUIInFront = false;
+			}
+
 			Panel->SetDockedTo( AVREditorFloatingUI::EDockedTo::Room );
 			Panel->ShowUI( true );
 
