VR Editor: Cleaned up hardcoded RadialMenu

diff --git a/Engine/Source/Editor/VREditor/VREditorRadialMenu.cpp b/Engine/Source/Editor/VREditor/VREditorRadialMenu.cpp
--- a/Engine/Source/Editor/VREditor/VREditorRadialMenu.cpp
+++ b/Engine/Source/Editor/VREditor/VREditorRadialMenu.cpp
@@ -21,17 +21,9 @@ namespace VREd
 }
 
 UVREditorRadialMenu::UVREditorRadialMenu( const FObjectInitializer& ObjectInitializer )
-	: Super( ObjectInitializer ),
-	TopItem( nullptr ),
-	TopRightItem(nullptr),
-	RightItem(nullptr),
-	BottomRightItem(nullptr),
-	BottomItem(nullptr),
-	LeftBottomItem(nullptr),
-	LeftItem(nullptr),
-	TopLeftItem(nullptr),
-	CurrentItem(nullptr),
-	PreviousItem(nullptr)
+: Super( ObjectInitializer ),
+CurrentItem( nullptr ),
+PreviousItem( nullptr )
 {
 }
 
@@ -39,67 +31,43 @@ UVREditorRadialMenu::UVREditorRadialMenu( const FObjectInitializer& ObjectInitia
 void UVREditorRadialMenu::Update( const FVirtualHand& Hand )
 {
 	TrackpadPosition = Hand.TrackpadPosition;
-	
-	if( !GetOwner()->bHidden )
+
+	if (!GetOwner()->bHidden)
 	{
 		// Check if position is from the center
-		if( !IsInMenuRadius() )
+		if (!IsInMenuRadius())
 		{
-			float Angle = FRotator::NormalizeAxis( FMath::RadiansToDegrees(FMath::Atan2( TrackpadPosition.X, TrackpadPosition.Y ) ) );
+			const float AnglePerItem = 360 / MenuItems.Num();
+			float Angle = FRotator::NormalizeAxis( FMath::RadiansToDegrees( FMath::Atan2( TrackpadPosition.X, TrackpadPosition.Y ) ) );
 			TrackpadAngle = Angle;
-			if( Angle < 0)
-			{
-				Angle = 360 + Angle;
-			}
 
-			// @todo vreditor : really hardcoded! This should be calculated from the amount of widgets in the current group
-			if( Angle > 337.5f && Angle <= 360|| Angle > 0 && Angle <= 22.5f )
-			{
-				CurrentItem = TopItem;
-			}
-			else if( Angle > 22.5f && Angle <= 67.5f)
-			{
-				CurrentItem = TopRightItem;
-			}
-			else if (Angle > 67.5f && Angle <= 112.5f)
-			{
-				CurrentItem = RightItem;
-			}
-			else if (Angle > 112.5f && Angle <= 157.5f)
+			Angle += AnglePerItem / 2.0f;
+			if (Angle < 0)
 			{
-				CurrentItem = BottomRightItem;
+				Angle = Angle + 360;
 			}
-			else if ( Angle > 157.5f && Angle <= 202.5f)
-			{
-				CurrentItem = BottomItem;
-			}
-			else if (Angle > 202.5f && Angle <= 247.5f)
-			{
-				CurrentItem = LeftBottomItem;
-			}
-			else if (Angle > 247.5f && Angle <= 292.5f)
-			{
-				CurrentItem = LeftItem;
-			}
-			else if (Angle > 292.5f && Angle <= 337.5f)
+
+			const int32 Index = (Angle / AnglePerItem);
+			UVREditorRadialMenuItem* NewMenuItem = MenuItems[Index];
+			if (NewMenuItem)
 			{
-				CurrentItem = TopLeftItem;
+				CurrentItem = NewMenuItem;
 			}
 
 			// Update the visuals
 			if (CurrentItem != PreviousItem)
 			{
-				if( CurrentItem )
+				if (CurrentItem)
 				{
 					CurrentItem->OnEnterHover();
 				}
 
-				if ( PreviousItem )
+				if (PreviousItem)
 				{
 					PreviousItem->OnLeaveHover();
 				}
 			}
-		
+
 			PreviousItem = CurrentItem;
 		}
 		else
@@ -117,37 +85,37 @@ void UVREditorRadialMenu::Update( const FVirtualHand& Hand )
 void UVREditorRadialMenu::SetButtons( UVREditorRadialMenuItem* InitTopItem, UVREditorRadialMenuItem* InitTopRightItem, UVREditorRadialMenuItem* InitRightItem, UVREditorRadialMenuItem* InitBottomRightItem,
 	UVREditorRadialMenuItem* InitBottomItem, UVREditorRadialMenuItem* InitLeftBottomItem, UVREditorRadialMenuItem* InitLeftItem, UVREditorRadialMenuItem* InitTopLeftItem )
 {
-	TopItem = InitTopItem;
-	TopItem->OnPressedDelegate.AddUObject( this, &UVREditorRadialMenu::OnGizmoCycle );
+	InitTopItem->OnPressedDelegate.AddUObject( this, &UVREditorRadialMenu::OnGizmoCycle );
+	MenuItems.Add( InitTopItem );
 	UpdateGizmoTypeLabel();
 
-	TopRightItem = InitTopRightItem;
-	TopRightItem->OnPressedDelegate.AddUObject(this, &UVREditorRadialMenu::OnCopyButtonClicked);
-	TopRightItem->SetLabel(FString("Copy"));
+	InitTopRightItem->OnPressedDelegate.AddUObject( this, &UVREditorRadialMenu::OnCopyButtonClicked );
+	InitTopRightItem->SetLabel( FString( "Copy" ) );
+	MenuItems.Add( InitTopRightItem );
 
-	RightItem = InitRightItem;
-	RightItem->OnPressedDelegate.AddUObject(this, &UVREditorRadialMenu::OnPasteButtonClicked);
-	RightItem->SetLabel(FString("Paste"));
+	InitRightItem->OnPressedDelegate.AddUObject( this, &UVREditorRadialMenu::OnPasteButtonClicked );
+	InitRightItem->SetLabel( FString( "Paste" ) );
+	MenuItems.Add( InitRightItem );
 
-	BottomRightItem = InitBottomRightItem;
-	BottomRightItem->OnPressedDelegate.AddUObject(this, &UVREditorRadialMenu::OnRedoButtonClicked);
-	BottomRightItem->SetLabel(FString("Redo"));
+	InitBottomRightItem->OnPressedDelegate.AddUObject( this, &UVREditorRadialMenu::OnRedoButtonClicked );
+	InitBottomRightItem->SetLabel( FString( "Redo" ) );
+	MenuItems.Add( InitBottomRightItem );
 
-	BottomItem = InitBottomItem;
-	BottomItem->OnPressedDelegate.AddUObject( this, &UVREditorRadialMenu::OnSnapActorsToGroundClicked );
-	BottomItem->SetLabel( FString( "Snap To Ground" ) );
+	InitBottomItem->OnPressedDelegate.AddUObject( this, &UVREditorRadialMenu::OnSnapActorsToGroundClicked );
+	InitBottomItem->SetLabel( FString( "Snap To Ground" ) );
+	MenuItems.Add( InitBottomItem );
 
-	LeftBottomItem = InitLeftBottomItem;
-	LeftBottomItem->OnPressedDelegate.AddUObject(this, &UVREditorRadialMenu::OnUndoButtonClicked);
-	LeftBottomItem->SetLabel(FString("Undo"));
+	InitLeftBottomItem->OnPressedDelegate.AddUObject( this, &UVREditorRadialMenu::OnUndoButtonClicked );
+	InitLeftBottomItem->SetLabel( FString( "Undo" ) );
+	MenuItems.Add( InitLeftBottomItem );
 
-	LeftItem = InitLeftItem;
-	LeftItem->OnPressedDelegate.AddUObject(this, &UVREditorRadialMenu::OnDuplicateButtonClicked);
-	LeftItem->SetLabel(FString("Duplicate"));
+	InitLeftItem->OnPressedDelegate.AddUObject( this, &UVREditorRadialMenu::OnDuplicateButtonClicked );
+	InitLeftItem->SetLabel( FString( "Duplicate" ) );
+	MenuItems.Add( InitLeftItem );
 
-	TopLeftItem = InitTopLeftItem;
-	TopLeftItem->OnPressedDelegate.AddUObject(this, &UVREditorRadialMenu::OnDeleteButtonClicked);
-	TopLeftItem->SetLabel(FString("Delete"));
+	InitTopLeftItem->OnPressedDelegate.AddUObject( this, &UVREditorRadialMenu::OnDeleteButtonClicked );
+	InitTopLeftItem->SetLabel( FString( "Delete" ) );
+	MenuItems.Add( InitTopLeftItem );
 }
 
 void UVREditorRadialMenu::ResetItem()
@@ -157,7 +125,7 @@ void UVREditorRadialMenu::ResetItem()
 		PreviousItem->OnLeaveHover();
 	}
 
-	if ( CurrentItem )
+	if (CurrentItem)
 	{
 		CurrentItem->OnLeaveHover();
 	}
@@ -165,11 +133,11 @@ void UVREditorRadialMenu::ResetItem()
 
 void UVREditorRadialMenu::SelectCurrentItem( const FVirtualHand& Hand )
 {
-	if ( IsInMenuRadius() )
+	if (IsInMenuRadius())
 	{
 		GetOwner()->GetOwner().TogglePanelsVisibility();
 	}
-	else if(CurrentItem)
+	else if (CurrentItem)
 	{
 		CurrentItem->OnPressed();
 	}
@@ -185,7 +153,7 @@ void UVREditorRadialMenu::UpdateGizmoTypeLabel()
 {
 	FString Result;
 
-	switch ( GetOwner()->GetOwner().GetOwner().GetCurrentGizmoType() )
+	switch (GetOwner()->GetOwner().GetOwner().GetCurrentGizmoType())
 	{
 	case EGizmoHandleTypes::All:
 		Result = LOCTEXT( "AllGizmoType", "Universal Gizmo" ).ToString();
@@ -203,7 +171,7 @@ void UVREditorRadialMenu::UpdateGizmoTypeLabel()
 		break;
 	}
 
-	TopItem->SetLabel( Result );
+	MenuItems[0]->SetLabel( Result );
 }
 
 void UVREditorRadialMenu::OnGizmoCycle()
diff --git a/Engine/Source/Editor/VREditor/VREditorRadialMenu.h b/Engine/Source/Editor/VREditor/VREditorRadialMenu.h
--- a/Engine/Source/Editor/VREditor/VREditorRadialMenu.h
+++ b/Engine/Source/Editor/VREditor/VREditorRadialMenu.h
@@ -46,23 +46,8 @@ public:
 	float TrackpadAngle;
 
 private:
-	/** @todo vreditor : this should be a group/struct */
-	UPROPERTY()
-	UVREditorRadialMenuItem* TopItem;
-	UPROPERTY()
-	UVREditorRadialMenuItem* TopRightItem;
-	UPROPERTY()
-	UVREditorRadialMenuItem* RightItem;
-	UPROPERTY()
-	UVREditorRadialMenuItem* BottomRightItem;
-	UPROPERTY()
-	UVREditorRadialMenuItem* BottomItem;
-	UPROPERTY()
-	UVREditorRadialMenuItem* LeftBottomItem;
-	UPROPERTY()
-	UVREditorRadialMenuItem* LeftItem;
-	UPROPERTY()
-	UVREditorRadialMenuItem* TopLeftItem;
+	/** All the buttons in the radial menu */
+	TArray<UVREditorRadialMenuItem*> MenuItems;
 
 	/** The current widget currently selected */
 	UPROPERTY()
