VR Editor: Fixed some variables not being nulled out during clean up

diff --git a/Engine/Source/Editor/VREditor/VREditorMode.cpp b/Engine/Source/Editor/VREditor/VREditorMode.cpp
--- a/Engine/Source/Editor/VREditor/VREditorMode.cpp
+++ b/Engine/Source/Editor/VREditor/VREditorMode.cpp
@@ -2484,6 +2484,11 @@ void FVREditorMode::CleanUpActorsBeforeMapChangeOrSimulate()
 		SnapGridMID = nullptr;
 	}
 
+	PostProcessComponent = nullptr;
+	ScaleProgressMeshComponent = nullptr;
+	CurrentScaleProgressMeshComponent = nullptr;
+	UserScaleIndicatorText = nullptr;
+
 	for( int32 HandIndex = 0; HandIndex < VREditorConstants::NumVirtualHands; ++HandIndex )
 	{
 		FVirtualHand& Hand = VirtualHands[ HandIndex ];
