VREditor : Clean up duplicated gizmo handle mesh creation code

diff --git a/Engine/Source/Editor/VREditor/Gizmo/VREditorGizmoHandle.cpp b/Engine/Source/Editor/VREditor/Gizmo/VREditorGizmoHandle.cpp
--- a/Engine/Source/Editor/VREditor/Gizmo/VREditorGizmoHandle.cpp
+++ b/Engine/Source/Editor/VREditor/Gizmo/VREditorGizmoHandle.cpp
@@ -114,27 +114,9 @@ ETransformGizmoInteractionType UVREditorGizmoHandleGroup::GetInteractionType() c
 
 void UVREditorGizmoHandleGroup::UpdateGizmoHandleGroup( const FTransform& LocalToWorld, const FBox& LocalBounds, const FVector ViewLocation, bool bAllHandlesVisible, class UActorComponent* DraggingHandle, const TArray< UActorComponent* >& HoveringOverHandles, float AnimationAlpha, float GizmoScale, const float GizmoHoverScale, const float GizmoHoverAnimationDuration, bool& bOutIsHoveringOrDraggingThisHandleGroup )
 {
-	bOutIsHoveringOrDraggingThisHandleGroup = false;
-
-	// Update hover animation
-	for( FVREditorGizmoHandle& Handle : Handles )
-	{
-		const bool bIsHoveringOverHandle = HoveringOverHandles.Contains( Handle.HandleMesh ) || ( DraggingHandle != nullptr && DraggingHandle == Handle.HandleMesh );
-
-		if( bIsHoveringOverHandle )
-		{
-			Handle.HoverAlpha += GetWorld()->GetDeltaSeconds() / GizmoHoverAnimationDuration;
-			bOutIsHoveringOrDraggingThisHandleGroup = true;
-		}
-		else
-		{
-			Handle.HoverAlpha -= GetWorld()->GetDeltaSeconds() / GizmoHoverAnimationDuration;
-		}
-		Handle.HoverAlpha = FMath::Clamp( Handle.HoverAlpha, 0.0f, 1.0f );
-	}
+	UpdateHoverAnimation( DraggingHandle, HoveringOverHandles, GizmoHoverAnimationDuration, bOutIsHoveringOrDraggingThisHandleGroup );
 }
 
-
 int32 UVREditorGizmoHandleGroup::GetDraggedHandleIndex( class UStaticMeshComponent* DraggedMesh )
 {
 	for( int32 HandleIndex = 0; HandleIndex < Handles.Num(); ++HandleIndex )
@@ -224,3 +206,80 @@ void UVREditorGizmoHandleGroup::UpdateHandleColor( const int32 AxisIndex, FVREdi
 		}
 	}
 }
+
+class UStaticMeshComponent* UVREditorGizmoHandleGroup::CreateMeshHandle( class UStaticMesh* HandleMesh, const FString& ComponentName )
+{
+	const bool bAllowGizmoLighting = false;	// @todo vreditor: Not sure if we want this for gizmos or not yet.  Needs feedback.  Also they're translucent right now.
+
+	UStaticMeshComponent* HandleComponent = CreateDefaultSubobject<UStaticMeshComponent>( *ComponentName );
+	check( HandleComponent != nullptr );
+
+	HandleComponent->SetStaticMesh( HandleMesh );
+	HandleComponent->SetMobility( EComponentMobility::Movable );
+	HandleComponent->AttachTo( this );
+
+	// @todo vreditor: We added a new engine collision channel "ECC_EditorGizmo" so that we could trace only
+	// against gizmos and nothing else.  This allows you to "click thru" solid objects to hit a ghosted gizmo.
+	// However, we need to look out for backwards compatibility issues with adding a new engine collision
+	// channel.  When I tested this with an existing project, all of the engine collision channels had been
+	// duplicated into the game's DefaultEngine.ini and were not updated to filter out the ECC_EditorGizmo
+	// like the stock collision channels that I updated in BaseEngine.ini when adding this feature.
+	HandleComponent->SetCollisionEnabled( ECollisionEnabled::QueryOnly );
+	HandleComponent->SetCollisionResponseToAllChannels( ECR_Ignore );
+	HandleComponent->SetCollisionResponseToChannel( ECC_EditorGizmo, ECollisionResponse::ECR_Block );
+
+	HandleComponent->bGenerateOverlapEvents = false;
+	HandleComponent->SetCanEverAffectNavigation( false );
+	HandleComponent->bCastDynamicShadow = bAllowGizmoLighting;
+	HandleComponent->bCastStaticShadow = false;
+	HandleComponent->bAffectDistanceFieldLighting = bAllowGizmoLighting;
+	HandleComponent->bAffectDynamicIndirectLighting = bAllowGizmoLighting;
+
+	return HandleComponent;
+}
+
+UStaticMeshComponent* UVREditorGizmoHandleGroup::CreateAndAddMeshHandle( UStaticMesh* HandleMesh, const FString& ComponentName, const FTransformGizmoHandlePlacement& HandlePlacement )
+{
+	UStaticMeshComponent* HandleComponent = CreateMeshHandle( HandleMesh, ComponentName );
+	AddMeshToHandles( HandleComponent, HandlePlacement );
+	return HandleComponent;
+}
+
+void UVREditorGizmoHandleGroup::AddMeshToHandles( UStaticMeshComponent* HandleMeshComponent, const FTransformGizmoHandlePlacement& HandlePlacement )
+{
+	int32 HandleIndex = MakeHandleIndex( HandlePlacement );
+	if (Handles.Num() < (HandleIndex + 1))
+	{
+		Handles.SetNumZeroed( HandleIndex + 1 );
+	}
+	Handles[HandleIndex].HandleMesh = HandleMeshComponent;
+}
+
+FTransformGizmoHandlePlacement UVREditorGizmoHandleGroup::GetHandlePlacement( const int32 X, const int32 Y, const int32 Z ) const
+{
+	FTransformGizmoHandlePlacement HandlePlacement;
+	HandlePlacement.Axes[0] = (ETransformGizmoHandleDirection)X;
+	HandlePlacement.Axes[1] = (ETransformGizmoHandleDirection)Y;
+	HandlePlacement.Axes[2] = (ETransformGizmoHandleDirection)Z;
+
+	return HandlePlacement;
+}
+
+void UVREditorGizmoHandleGroup::UpdateHoverAnimation( UActorComponent* DraggingHandle, const TArray< UActorComponent* >& HoveringOverHandles, const float GizmoHoverAnimationDuration, bool& bOutIsHoveringOrDraggingThisHandleGroup )
+{
+	for (FVREditorGizmoHandle& Handle : Handles)
+	{
+		const bool bIsHoveringOverHandle = HoveringOverHandles.Contains( Handle.HandleMesh ) || (DraggingHandle != nullptr && DraggingHandle == Handle.HandleMesh);
+
+		if (bIsHoveringOverHandle)
+		{
+			Handle.HoverAlpha += GetWorld()->GetDeltaSeconds() / GizmoHoverAnimationDuration;
+			bOutIsHoveringOrDraggingThisHandleGroup = true;
+		}
+		else
+		{
+			Handle.HoverAlpha -= GetWorld()->GetDeltaSeconds() / GizmoHoverAnimationDuration;
+		}
+		Handle.HoverAlpha = FMath::Clamp( Handle.HoverAlpha, 0.0f, 1.0f );
+	}
+}
diff --git a/Engine/Source/Editor/VREditor/Gizmo/VREditorGizmoHandle.h b/Engine/Source/Editor/VREditor/Gizmo/VREditorGizmoHandle.h
--- a/Engine/Source/Editor/VREditor/Gizmo/VREditorGizmoHandle.h
+++ b/Engine/Source/Editor/VREditor/Gizmo/VREditorGizmoHandle.h
@@ -95,6 +95,18 @@ protected:
 	/** Updates the colors of the dynamic material instances for the handle passed using its axis index */	
 	void UpdateHandleColor( const int32 AxisIndex, FVREditorGizmoHandle& Handle, class UActorComponent* DraggingHandle, const TArray< UActorComponent* >& HoveringOverHandles );
 
+	/** Helper function to create gizmo handle meshes */
+	class UStaticMeshComponent* CreateMeshHandle( class UStaticMesh* HandleMesh, const FString& ComponentName );
+
+	/** Creates handle meshcomponent and adds it to the Handles list */
+	class UStaticMeshComponent* CreateAndAddMeshHandle( class UStaticMesh* HandleMesh, const FString& ComponentName, const FTransformGizmoHandlePlacement& HandlePlacement );
+
+	/** Adds the HandleMeshComponent to the Handles list */
+	void AddMeshToHandles( class UStaticMeshComponent* HandleMeshComponent, const FTransformGizmoHandlePlacement& HandlePlacement );
+
+	/** Gets the handleplacement axes */
+	FTransformGizmoHandlePlacement GetHandlePlacement( const int32 X, const int32 Y, const int32 Z ) const;
+
 	/** Gizmo material (opaque) */
 	UPROPERTY()
 	UMaterialInterface* GizmoMaterial;
@@ -106,4 +118,9 @@ protected:
 	/** All the StaticMeshes for this handle type */
 	UPROPERTY()
 	TArray< FVREditorGizmoHandle > Handles;
+
+private:
+
+	/** Updates the hover animation for the HoveringOverHandles */
+	void UpdateHoverAnimation( class UActorComponent* DraggingHandle, const TArray< UActorComponent* >& HoveringOverHandles, const float GizmoHoverAnimationDuration, bool& bOutIsHoveringOrDraggingThisHandleGroup );
 };
\ No newline at end of file
diff --git a/Engine/Source/Editor/VREditor/Gizmo/VREditorPlaneTranslationGizmoHandle.cpp b/Engine/Source/Editor/VREditor/Gizmo/VREditorPlaneTranslationGizmoHandle.cpp
--- a/Engine/Source/Editor/VREditor/Gizmo/VREditorPlaneTranslationGizmoHandle.cpp
+++ b/Engine/Source/Editor/VREditor/Gizmo/VREditorPlaneTranslationGizmoHandle.cpp
@@ -23,10 +23,7 @@ UVREditorPlaneTranslationGizmoHandleGroup::UVREditorPlaneTranslationGizmoHandleG
 		{
 			for (int32 Z = 0; Z < 3; ++Z)
 			{
-				FTransformGizmoHandlePlacement HandlePlacement;
-				HandlePlacement.Axes[0] = (ETransformGizmoHandleDirection)X;
-				HandlePlacement.Axes[1] = (ETransformGizmoHandleDirection)Y;
-				HandlePlacement.Axes[2] = (ETransformGizmoHandleDirection)Z;
+				FTransformGizmoHandlePlacement HandlePlacement = GetHandlePlacement( X, Y, Z );
 
 				int32 CenterHandleCount, FacingAxisIndex, CenterAxisIndex;
 				HandlePlacement.GetCenterHandleCountAndFacingAxisIndex( /* Out */ CenterHandleCount, /* Out */ FacingAxisIndex, /* Out */ CenterAxisIndex );
@@ -40,34 +37,8 @@ UVREditorPlaneTranslationGizmoHandleGroup::UVREditorPlaneTranslationGizmoHandleG
 					if (CenterHandleCount == 2 && HandlePlacement.Axes[FacingAxisIndex] == ETransformGizmoHandleDirection::Positive)
 					{
 						// Plane translation handle
-						{
-							FString ComponentName = HandleName + TEXT( "PlaneTranslationHandle" );
-
-							UStaticMeshComponent* PlaneTranslationHandle = CreateDefaultSubobject<UStaticMeshComponent>( *ComponentName );
-							check( PlaneTranslationHandle != nullptr );
-
-							PlaneTranslationHandle->SetStaticMesh( PlaneTranslationHandleMesh );
-							PlaneTranslationHandle->SetMobility( EComponentMobility::Movable );
-							PlaneTranslationHandle->AttachTo( this );
-
-							PlaneTranslationHandle->SetCollisionEnabled( ECollisionEnabled::QueryOnly );
-							PlaneTranslationHandle->SetCollisionResponseToAllChannels( ECR_Ignore );
-							PlaneTranslationHandle->SetCollisionResponseToChannel( ECC_EditorGizmo, ECollisionResponse::ECR_Block );
-
-							PlaneTranslationHandle->bGenerateOverlapEvents = false;
-							PlaneTranslationHandle->SetCanEverAffectNavigation( false );
-							PlaneTranslationHandle->bCastDynamicShadow = bAllowGizmoLighting;
-							PlaneTranslationHandle->bCastStaticShadow = false;
-							PlaneTranslationHandle->bAffectDistanceFieldLighting = bAllowGizmoLighting;
-							PlaneTranslationHandle->bAffectDynamicIndirectLighting = bAllowGizmoLighting;
-
-							int32 HandleIndex = MakeHandleIndex( HandlePlacement );
-							if (Handles.Num() < (HandleIndex + 1))
-							{
-								Handles.SetNumZeroed( HandleIndex + 1 );
-							}
-							Handles[HandleIndex].HandleMesh = PlaneTranslationHandle;
-						}
+						FString ComponentName = HandleName + TEXT( "PlaneTranslationHandle" );
+						CreateAndAddMeshHandle( PlaneTranslationHandleMesh, ComponentName, HandlePlacement );
 					}
 				}
 			}
diff --git a/Engine/Source/Editor/VREditor/Gizmo/VREditorRotationGizmoHandle.cpp b/Engine/Source/Editor/VREditor/Gizmo/VREditorRotationGizmoHandle.cpp
--- a/Engine/Source/Editor/VREditor/Gizmo/VREditorRotationGizmoHandle.cpp
+++ b/Engine/Source/Editor/VREditor/Gizmo/VREditorRotationGizmoHandle.cpp
@@ -23,10 +23,7 @@ UVREditorRotationGizmoHandleGroup::UVREditorRotationGizmoHandleGroup()
 		{
 			for (int32 Z = 0; Z < 3; ++Z)
 			{
-				FTransformGizmoHandlePlacement HandlePlacement;
-				HandlePlacement.Axes[0] = (ETransformGizmoHandleDirection)X;
-				HandlePlacement.Axes[1] = (ETransformGizmoHandleDirection)Y;
-				HandlePlacement.Axes[2] = (ETransformGizmoHandleDirection)Z;
+				FTransformGizmoHandlePlacement HandlePlacement = GetHandlePlacement( X, Y, Z );
 
 				int32 CenterHandleCount, FacingAxisIndex, CenterAxisIndex;
 				HandlePlacement.GetCenterHandleCountAndFacingAxisIndex( /* Out */ CenterHandleCount, /* Out */ FacingAxisIndex, /* Out */ CenterAxisIndex );
@@ -41,34 +38,7 @@ UVREditorRotationGizmoHandleGroup::UVREditorRotationGizmoHandleGroup()
 					if (CenterHandleCount == 1)	// Edge?
 					{
 						FString ComponentName = HandleName + TEXT( "RotationHandle" );
-
-						UStaticMeshComponent* RotationHandle = CreateDefaultSubobject<UStaticMeshComponent>( *ComponentName );
-						check( RotationHandle != nullptr );
-
-						RotationHandle->SetStaticMesh( RotationHandleMesh );
-
-						RotationHandle->SetMobility( EComponentMobility::Movable );
-						RotationHandle->AttachTo( this );
-
-						RotationHandle->SetCollisionEnabled( ECollisionEnabled::QueryOnly );
-						RotationHandle->SetCollisionResponseToAllChannels( ECR_Ignore );
-						RotationHandle->SetCollisionResponseToChannel( ECC_EditorGizmo, ECollisionResponse::ECR_Block );
-
-						RotationHandle->bGenerateOverlapEvents = false;
-						RotationHandle->SetCanEverAffectNavigation( false );
-						RotationHandle->bCastDynamicShadow = bAllowGizmoLighting;
-						RotationHandle->bCastStaticShadow = false;
-						RotationHandle->bAffectDistanceFieldLighting = bAllowGizmoLighting;
-						RotationHandle->bAffectDynamicIndirectLighting = bAllowGizmoLighting;
-
-						int32 HandleIndex = MakeHandleIndex( HandlePlacement );
-						if (Handles.Num() < (HandleIndex + 1))
-						{
-							Handles.SetNumZeroed( HandleIndex + 1 );
-						}
-						Handles[HandleIndex].HandleMesh = RotationHandle;
-
-						check( CenterAxisIndex != INDEX_NONE );
+						CreateAndAddMeshHandle( RotationHandleMesh, ComponentName, HandlePlacement );
 					}
 				}
 			}
diff --git a/Engine/Source/Editor/VREditor/Gizmo/VREditorStretchGizmoHandle.cpp b/Engine/Source/Editor/VREditor/Gizmo/VREditorStretchGizmoHandle.cpp
--- a/Engine/Source/Editor/VREditor/Gizmo/VREditorStretchGizmoHandle.cpp
+++ b/Engine/Source/Editor/VREditor/Gizmo/VREditorStretchGizmoHandle.cpp
@@ -37,10 +37,7 @@ UVREditorStretchGizmoHandleGroup::UVREditorStretchGizmoHandleGroup()
 		{
 			for (int32 Z = 0; Z < 3; ++Z)
 			{
-				FTransformGizmoHandlePlacement HandlePlacement;
-				HandlePlacement.Axes[0] = (ETransformGizmoHandleDirection)X;
-				HandlePlacement.Axes[1] = (ETransformGizmoHandleDirection)Y;
-				HandlePlacement.Axes[2] = (ETransformGizmoHandleDirection)Z;
+				FTransformGizmoHandlePlacement HandlePlacement = GetHandlePlacement( X, Y, Z );
 
 				int32 CenterHandleCount, FacingAxisIndex, CenterAxisIndex;
 				HandlePlacement.GetCenterHandleCountAndFacingAxisIndex( /* Out */ CenterHandleCount, /* Out */ FacingAxisIndex, /* Out */ CenterAxisIndex );
@@ -55,41 +52,22 @@ UVREditorStretchGizmoHandleGroup::UVREditorStretchGizmoHandleGroup()
 					{
 						FString ComponentName = HandleName + TEXT( "StretchingHandle" );
 
-						UStaticMeshComponent* StretchingHandle = CreateDefaultSubobject<UStaticMeshComponent>( *ComponentName );
-						check( StretchingHandle != nullptr );
+						UStaticMesh* Mesh = nullptr;
 
 						if (CenterHandleCount == 0)	// Corner?
 						{
-							StretchingHandle->SetStaticMesh( BoundingBoxCornerMesh );
+							Mesh = BoundingBoxCornerMesh;
 						}
 						else if (CenterHandleCount == 1)	// Edge?
 						{
-							StretchingHandle->SetStaticMesh( BoundingBoxEdgeMesh );
+							Mesh = BoundingBoxEdgeMesh;
 						}
 						else  // Face
 						{
-							StretchingHandle->SetStaticMesh( StretchingHandleMesh );
+							Mesh = StretchingHandleMesh;
 						}
-						StretchingHandle->SetMobility( EComponentMobility::Movable );
-						StretchingHandle->AttachTo( this );
-
-						StretchingHandle->SetCollisionEnabled( ECollisionEnabled::QueryOnly );
-						StretchingHandle->SetCollisionResponseToAllChannels( ECR_Ignore );
-						StretchingHandle->SetCollisionResponseToChannel( ECC_EditorGizmo, ECollisionResponse::ECR_Block );
-
-						StretchingHandle->bGenerateOverlapEvents = false;
-						StretchingHandle->SetCanEverAffectNavigation( false );
-						StretchingHandle->bCastDynamicShadow = bAllowGizmoLighting;
-						StretchingHandle->bCastStaticShadow = false;
-						StretchingHandle->bAffectDistanceFieldLighting = bAllowGizmoLighting;
-						StretchingHandle->bAffectDynamicIndirectLighting = bAllowGizmoLighting;
-
-						int32 HandleIndex = MakeHandleIndex( HandlePlacement );
-						if (Handles.Num() < (HandleIndex + 1))
-						{
-							Handles.SetNumZeroed( HandleIndex + 1 );
-						}
-						Handles[HandleIndex].HandleMesh = StretchingHandle;
+
+						CreateAndAddMeshHandle( Mesh, ComponentName, HandlePlacement );
 					}
 				}
 			}
diff --git a/Engine/Source/Editor/VREditor/Gizmo/VREditorTranslationGizmoHandle.cpp b/Engine/Source/Editor/VREditor/Gizmo/VREditorTranslationGizmoHandle.cpp
--- a/Engine/Source/Editor/VREditor/Gizmo/VREditorTranslationGizmoHandle.cpp
+++ b/Engine/Source/Editor/VREditor/Gizmo/VREditorTranslationGizmoHandle.cpp
@@ -23,10 +23,7 @@ UVREditorTranslationGizmoHandleGroup::UVREditorTranslationGizmoHandleGroup()
 		{
 			for (int32 Z = 0; Z < 3; ++Z)
 			{
-				FTransformGizmoHandlePlacement HandlePlacement;
-				HandlePlacement.Axes[0] = (ETransformGizmoHandleDirection)X;
-				HandlePlacement.Axes[1] = (ETransformGizmoHandleDirection)Y;
-				HandlePlacement.Axes[2] = (ETransformGizmoHandleDirection)Z;
+				FTransformGizmoHandlePlacement HandlePlacement = GetHandlePlacement( X, Y, Z );
 
 				int32 CenterHandleCount, FacingAxisIndex, CenterAxisIndex;
 				HandlePlacement.GetCenterHandleCountAndFacingAxisIndex( /* Out */ CenterHandleCount, /* Out */ FacingAxisIndex, /* Out */ CenterAxisIndex );
@@ -40,41 +37,8 @@ UVREditorTranslationGizmoHandleGroup::UVREditorTranslationGizmoHandleGroup()
 					if (CenterHandleCount == 2)
 					{
 						// Translation handle
-						{
-							FString ComponentName = HandleName + TEXT( "TranslationHandle" );
-
-							UStaticMeshComponent* TranslationHandle = CreateDefaultSubobject<UStaticMeshComponent>( *ComponentName );
-							check( TranslationHandle != nullptr );
-
-							TranslationHandle->SetStaticMesh( TranslationHandleMesh );
-							TranslationHandle->SetMobility( EComponentMobility::Movable );
-							TranslationHandle->AttachTo( this );
-
-							// @todo vreditor: We added a new engine collision channel "ECC_EditorGizmo" so that we could trace only
-							// against gizmos and nothing else.  This allows you to "click thru" solid objects to hit a ghosted gizmo.
-							// However, we need to look out for backwards compatibility issues with adding a new engine collision
-							// channel.  When I tested this with an existing project, all of the engine collision channels had been
-							// duplicated into the game's DefaultEngine.ini and were not updated to filter out the ECC_EditorGizmo
-							// like the stock collision channels that I updated in BaseEngine.ini when adding this feature.
-
-							TranslationHandle->SetCollisionEnabled( ECollisionEnabled::QueryOnly );
-							TranslationHandle->SetCollisionResponseToAllChannels( ECR_Ignore );
-							TranslationHandle->SetCollisionResponseToChannel( ECC_EditorGizmo, ECollisionResponse::ECR_Block );
-
-							TranslationHandle->bGenerateOverlapEvents = false;
-							TranslationHandle->SetCanEverAffectNavigation( false );
-							TranslationHandle->bCastDynamicShadow = bAllowGizmoLighting;
-							TranslationHandle->bCastStaticShadow = false;
-							TranslationHandle->bAffectDistanceFieldLighting = bAllowGizmoLighting;
-							TranslationHandle->bAffectDynamicIndirectLighting = bAllowGizmoLighting;
-
-							int32 HandleIndex = MakeHandleIndex( HandlePlacement );
-							if (Handles.Num() < (HandleIndex + 1))
-							{
-								Handles.SetNumZeroed( HandleIndex + 1 );
-							}
-							Handles[HandleIndex].HandleMesh = TranslationHandle;
-						}
+						FString ComponentName = HandleName + TEXT( "TranslationHandle" );
+						CreateAndAddMeshHandle( TranslationHandleMesh, ComponentName, HandlePlacement );
 					}
 				}
 			}
@@ -147,4 +111,4 @@ ETransformGizmoInteractionType UVREditorTranslationGizmoHandleGroup::GetInteract
 EGizmoHandleTypes UVREditorTranslationGizmoHandleGroup::GetHandleType() const
 {
 	return EGizmoHandleTypes::Translate;
-}
\ No newline at end of file
+}
