VR Editor: Removed some hacks in preparation for merge to main
- Actors prefixed with "Frozen" in their label are no longer unselectable in the VR Editor
- Removed various comments for VR Editor hacks that no longer applied

diff --git a/Engine/Plugins/Runtime/OculusRift/Source/OculusRift/Private/OculusRiftHMD.cpp b/Engine/Plugins/Runtime/OculusRift/Source/OculusRift/Private/OculusRiftHMD.cpp
--- a/Engine/Plugins/Runtime/OculusRift/Source/OculusRift/Private/OculusRiftHMD.cpp
+++ b/Engine/Plugins/Runtime/OculusRift/Source/OculusRift/Private/OculusRiftHMD.cpp
@@ -1310,7 +1310,7 @@ bool FOculusRiftHMD::DoEnableStereo(bool bStereo, bool bApplyToHmd)
 
 		if (SceneVP && SceneVP->GetViewportWidget().IsValid())
 		{
-			if (!IsFullscreenAllowed() && stereoToBeEnabled)	// @todo vreditor: What do we need to do to be able to support "full screen" on the HMD device in the editor?  Do we even care?  Perf benefits to that?
+			if (!IsFullscreenAllowed() && stereoToBeEnabled)
 			{
 				FOvrSessionShared::AutoSession OvrSession(Session);
 				if (OvrSession)
@@ -1344,7 +1344,7 @@ bool FOculusRiftHMD::DoEnableStereo(bool bStereo, bool bApplyToHmd)
 
 						if (stereoToBeEnabled)
 						{
-							EWindowMode::Type wm = (!GIsEditor) ? EWindowMode::Fullscreen : EWindowMode::WindowedFullscreen;	// @todo vreditor: Fullscreen gating based on editor here
+							EWindowMode::Type wm = (!GIsEditor) ? EWindowMode::Fullscreen : EWindowMode::WindowedFullscreen;
 							FVector2D size = Window->GetSizeInScreen();
 							SceneVP->ResizeFrame(size.X, size.Y, wm, 0, 0);
 						}
@@ -1352,7 +1352,7 @@ bool FOculusRiftHMD::DoEnableStereo(bool bStereo, bool bApplyToHmd)
 						{
 							// In Editor we cannot use ResizeFrame trick since it is called too late and App::IsGame
 							// returns false.
-							if (GIsEditor)	// @todo vreditor: More pre-existing editor hacks for VR PIE, might be fine for VR editor too
+							if (GIsEditor)
 							{
 								FSlateRect LocalPreFullScreenRect;
 								PopPreFullScreenRect(LocalPreFullScreenRect);
@@ -1707,7 +1707,7 @@ FOculusRiftHMD::~FOculusRiftHMD()
 void FOculusRiftHMD::Startup()
 {
 #if PLATFORM_MAC
-	if (GIsEditor)	// @todo vreditor: No Mac VR support yet?
+	if (GIsEditor)
 	{
 		// no editor support for Mac yet
 		return;
diff --git a/Engine/Source/Editor/EditorStyle/Private/SlateEditorStyle.cpp b/Engine/Source/Editor/EditorStyle/Private/SlateEditorStyle.cpp
--- a/Engine/Source/Editor/EditorStyle/Private/SlateEditorStyle.cpp
+++ b/Engine/Source/Editor/EditorStyle/Private/SlateEditorStyle.cpp
@@ -4565,7 +4565,6 @@ void FSlateEditorStyle::FStyle::SetupLevelEditorStyle()
 		Set( "LevelEditor.EditMatinee", new IMAGE_BRUSH( "Icons/icon_matinee_40x", Icon40x40 ) );
 		Set( "LevelEditor.EditMatinee.Small", new IMAGE_BRUSH( "Icons/icon_matinee_40x", Icon20x20 ) );
 
-		// @todo vreditor: Icon needs to be arted up
 		Set( "LevelEditor.ToggleVR", new IMAGE_BRUSH( "Icons/icon_ToggleVREditor_40x", Icon40x40 ) );
 		Set( "LevelEditor.ToggleVR.Small", new IMAGE_BRUSH( "Icons/icon_ToggleVREditor_16x", Icon40x40 ) );
 
diff --git a/Engine/Source/Editor/VREditor/Gizmo/VREditorGizmoHandle.cpp b/Engine/Source/Editor/VREditor/Gizmo/VREditorGizmoHandle.cpp
--- a/Engine/Source/Editor/VREditor/Gizmo/VREditorGizmoHandle.cpp
+++ b/Engine/Source/Editor/VREditor/Gizmo/VREditorGizmoHandle.cpp
@@ -218,12 +218,6 @@ class UStaticMeshComponent* UVREditorGizmoHandleGroup::CreateMeshHandle( class U
 	HandleComponent->SetMobility( EComponentMobility::Movable );
 	HandleComponent->AttachTo( this );
 
-	// @todo vreditor: We added a new engine collision channel "ECC_EditorGizmo" so that we could trace only
-	// against gizmos and nothing else.  This allows you to "click thru" solid objects to hit a ghosted gizmo.
-	// However, we need to look out for backwards compatibility issues with adding a new engine collision
-	// channel.  When I tested this with an existing project, all of the engine collision channels had been
-	// duplicated into the game's DefaultEngine.ini and were not updated to filter out the ECC_EditorGizmo
-	// like the stock collision channels that I updated in BaseEngine.ini when adding this feature.
 	HandleComponent->SetCollisionEnabled( ECollisionEnabled::QueryOnly );
 	HandleComponent->SetCollisionResponseToAllChannels( ECR_Ignore );
 	HandleComponent->SetCollisionResponseToChannel( ECC_EditorGizmo, ECollisionResponse::ECR_Block );
diff --git a/Engine/Source/Editor/VREditor/VREditorMode.cpp b/Engine/Source/Editor/VREditor/VREditorMode.cpp
--- a/Engine/Source/Editor/VREditor/VREditorMode.cpp
+++ b/Engine/Source/Editor/VREditor/VREditorMode.cpp
@@ -200,7 +200,6 @@ void FVREditorMode::Enter()
 		KeyToActionMap.Add( VREditorKeyNames::MotionController_Left_LightlyPressedTriggerAxis, FVRAction( EVRActionType::SelectAndMove_LightlyPressed, VREditorConstants::LeftHandIndex ) );
 		KeyToActionMap.Add( VREditorKeyNames::MotionController_Right_LightlyPressedTriggerAxis, FVRAction( EVRActionType::SelectAndMove_LightlyPressed, VREditorConstants::RightHandIndex ) );
 
-		// @todo vreditor UI: We moved most of these into the UI, and we're using the touchpad for Undo instead of summoning our Slate UI for now
 		KeyToActionMap.Add( EKeys::MotionController_Left_Thumbstick, FVRAction( EVRActionType::ConfirmRadialSelection, VREditorConstants::LeftHandIndex ) );
  		KeyToActionMap.Add( EKeys::MotionController_Right_Thumbstick, FVRAction( EVRActionType::ConfirmRadialSelection, VREditorConstants::RightHandIndex ) );
 
@@ -1942,7 +1941,7 @@ AActor* FVREditorMode::SpawnTransientSceneActor( TSubclassOf<AActor> ActorClass,
 	FActorSpawnParameters ActorSpawnParameters;
 	ActorSpawnParameters.Name = MakeUniqueObjectName( GetWorld(), ActorClass, *ActorName );	// @todo vreditor: Without this, SpawnActor() can return us an existing PendingKill actor of the same name!  WTF?
 	ActorSpawnParameters.SpawnCollisionHandlingOverride = ESpawnActorCollisionHandlingMethod::AlwaysSpawn;
-	ActorSpawnParameters.ObjectFlags = EObjectFlags::RF_Transient;	// @todo vreditor: Make sure transient actors aren't getting saved into maps, or picked up in world outliners, or other actor iterators where they should be hidden from
+	ActorSpawnParameters.ObjectFlags = EObjectFlags::RF_Transient;
 
 	check( ActorClass != nullptr );
 	AActor* NewActor = GetWorld()->SpawnActor< AActor >( ActorClass, ActorSpawnParameters );
@@ -2600,9 +2599,6 @@ void FVREditorMode::ApplyVelocityDamping( FVector& Velocity, const bool bVelocit
 		// Apply damping
 		if( bVelocitySensitive )
 		{
-			// @todo vreditor: Experimenting with speed-based dampening amounts here, so that you can drift further
-			// after you've been "thrown faster", but come to a stop rather quickly if you are not already moving fast.
-			//		--> If this is successful, consider generalizing this and doing it everywhere else we need inertial movement
 			const float DampenMultiplierAtLowSpeeds = 0.94f;	// @todo vreditor tweak
 			const float DampenMultiplierAtHighSpeeds = 0.99f;	// @todo vreditor tweak
 			const float SpeedForMinimalDamping = 2.5f * GetWorldScaleFactor();	// cm/frame	// @todo vreditor tweak
diff --git a/Engine/Source/Editor/VREditor/VREditorModule.cpp b/Engine/Source/Editor/VREditor/VREditorModule.cpp
--- a/Engine/Source/Editor/VREditor/VREditorModule.cpp
+++ b/Engine/Source/Editor/VREditor/VREditorModule.cpp
@@ -87,7 +87,7 @@ bool FVREditorModule::IsVREditorEnabled() const
 
 
 bool FVREditorModule::IsVREditorAvailable() const
-{	// @todo vreditor: For some reason, check state and disabled state is not reflected on the level editor toolbar!
+{
 	const bool bHasHMDDevice = GEngine->HMDDevice.IsValid() && GEngine->HMDDevice->IsHMDEnabled();
 	return bHasHMDDevice;
 }
diff --git a/Engine/Source/Editor/VREditor/VREditorWorldInteraction.cpp b/Engine/Source/Editor/VREditor/VREditorWorldInteraction.cpp
--- a/Engine/Source/Editor/VREditor/VREditorWorldInteraction.cpp
+++ b/Engine/Source/Editor/VREditor/VREditorWorldInteraction.cpp
@@ -1662,11 +1662,10 @@ void FVREditorWorldInteraction::RefreshTransformGizmo( const bool bNewObjectsSel
 bool FVREditorWorldInteraction::IsInteractableComponent( const UActorComponent* Component ) const
 {
 	// Don't interact with frozen actors
-	// @todo vreditor urgent: This is a super hacky way to determine whether something is "frozen".  This would need
-	// to become a first class concept in the editor for any interactable object.
 	if( Component != nullptr )
 	{
-		const bool bIsFrozen = Component->GetOwner()->GetActorLabel().StartsWith( TEXT( "Frozen_" ) );
+		static const bool bIsVREditorDemo = FParse::Param( FCommandLine::Get(), TEXT( "VREditorDemo" ) );	// @todo vreditor: Remove this when no longer needed (console variable, too!)
+		const bool bIsFrozen = bIsVREditorDemo && Component->GetOwner()->GetActorLabel().StartsWith( TEXT( "Frozen_" ) );
 		return
 			!bIsFrozen &&
 			( !Owner.GetUISystem().IsWidgetAnEditorUIWidget( Component ) ); // Don't allow user to move around our UI widgets
@@ -2531,7 +2530,7 @@ void FVREditorWorldInteraction::StartDraggingActors( FVRAction VRAction, UActorC
 			}
 
 			// Play a haptic effect when objects are picked up
-			const float Strength = VREd::DragHapticFeedbackStrength->GetFloat();	// @todo vreditor: Tweak
+			const float Strength = VREd::DragHapticFeedbackStrength->GetFloat();
 			Owner.PlayHapticEffect(
 				VRAction.HandIndex == VREditorConstants::LeftHandIndex ? Strength : 0.0f,
 				VRAction.HandIndex == VREditorConstants::RightHandIndex ? Strength : 0.0f );
@@ -2581,7 +2580,7 @@ void FVREditorWorldInteraction::StartDraggingMaterialOrTexture( FVRAction VRActi
 		LastDragGizmoStartTransform = FTransform::Identity;
 
 		// Play a haptic effect when objects are picked up
-		const float Strength = VREd::DragHapticFeedbackStrength->GetFloat();	// @todo vreditor: Tweak
+		const float Strength = VREd::DragHapticFeedbackStrength->GetFloat();
 		Owner.PlayHapticEffect(
 			VRAction.HandIndex == VREditorConstants::LeftHandIndex ? Strength : 0.0f,
 			VRAction.HandIndex == VREditorConstants::RightHandIndex ? Strength : 0.0f );
@@ -2743,7 +2742,7 @@ void FVREditorWorldInteraction::OnAssetDragStartedFromContentBrowser( const TArr
 	}
 
 	// We're always expecting a hand to be hovering at the time we receive this event
-	if( PlacingWithHandIndex != INDEX_NONE )	// @todo vreditor: This ensure fires when in "fake VR" mode when you drag an asset (Expected)
+	if( PlacingWithHandIndex != INDEX_NONE )
 	{
 		FVirtualHand& Hand = Owner.GetVirtualHand( PlacingWithHandIndex );
 
