VR Editor: Added console command to allow nonsensible world rotation (roll/pitch)
- Type: "VREd.AllowWorldRotationPitchAndRoll 1" to active this debug mode (NOT recommended!)
- When enabled, you'll not only be able to yaw, but also pitch and roll the world when rotating by gripping with two hands

diff --git a/Engine/Source/Editor/VREditor/VREditorWorldInteraction.cpp b/Engine/Source/Editor/VREditor/VREditorWorldInteraction.cpp
--- a/Engine/Source/Editor/VREditor/VREditorWorldInteraction.cpp
+++ b/Engine/Source/Editor/VREditor/VREditorWorldInteraction.cpp
@@ -54,6 +54,7 @@ namespace VREd
 	static FAutoConsoleVariable SnapGridLineWidth( TEXT( "VREd.SnapGridLineWidth" ), 3.0f, TEXT( "Width of the grid lines on the snap grid" ) );
 	static FAutoConsoleVariable HoverHapticFeedbackStrength( TEXT( "VREd.HoverHapticFeedbackStrength" ), 0.1f, TEXT( "Default strength for haptic feedback when hovering" ) );
 	static FAutoConsoleVariable HoverHapticFeedbackTime( TEXT( "VREd.HoverHapticFeedbackTime" ), 0.2f, TEXT( "The minimum time between haptic feedback for hovering" ) );
+	static FAutoConsoleVariable AllowWorldRotationPitchAndRoll( TEXT( "VREd.AllowWorldRotationPitchAndRoll" ), 0, TEXT( "When enabled, you'll not only be able to yaw, but also pitch and roll the world when rotating by gripping with two hands" ) );
 }
 
 
@@ -2336,17 +2337,22 @@ void FVREditorWorldInteraction::UpdateDragging(
 
 			// Apply rotation and translation
 			{
-				// Eliminate pitch and roll in rotation offset.  We don't want the user to get sick!
-				FTransform YawRotationOffsetTransform = FTransform::Identity;
-				FRotator YawRotationOffset = RotationOffset.Rotator();
-				YawRotationOffset.Pitch = YawRotationOffset.Roll = 0.0f;
-				YawRotationOffsetTransform.SetRotation( YawRotationOffset.Quaternion() );
+				FTransform RotationOffsetTransform = FTransform::Identity;
+				RotationOffsetTransform.SetRotation( RotationOffset );
+
+				if( VREd::AllowWorldRotationPitchAndRoll->GetInt() == 0 )
+				{
+					// Eliminate pitch and roll in rotation offset.  We don't want the user to get sick!
+					FRotator YawRotationOffset = RotationOffset.Rotator();
+					YawRotationOffset.Pitch = YawRotationOffset.Roll = 0.0f;
+					RotationOffsetTransform.SetRotation( YawRotationOffset.Quaternion() );
+				}
 
 				// Move the camera in the opposite direction, so it feels to the user as if they're dragging the entire world around
 				const FTransform TranslationOffsetTransform( FQuat::Identity, TranslationOffset );
 				const FTransform PivotToWorld = FTransform( FQuat::Identity, PivotLocation ) * RoomTransform;
 				const FTransform WorldToPivot = PivotToWorld.Inverse();
-				RoomTransform = TranslationOffsetTransform.Inverse() * RoomTransform * WorldToPivot * YawRotationOffsetTransform.Inverse() * PivotToWorld;
+				RoomTransform = TranslationOffsetTransform.Inverse() * RoomTransform * WorldToPivot * RotationOffsetTransform.Inverse() * PivotToWorld;
 			}
 
 			Owner.SetRoomTransform( RoomTransform );
