VR Editor: Sets the initial scale to 1 meter and initial rotation and location with the editor viewport

diff --git a/Engine/Source/Editor/VREditor/VREditorMode.cpp b/Engine/Source/Editor/VREditor/VREditorMode.cpp
--- a/Engine/Source/Editor/VREditor/VREditorMode.cpp
+++ b/Engine/Source/Editor/VREditor/VREditorMode.cpp
@@ -117,6 +117,7 @@ FVREditorMode::FVREditorMode()
 	  CurrentScaleProgressMeshComponent( nullptr ),
 	  UserScaleIndicatorText( nullptr ),
 	  CurrentGizmoType( EGizmoHandleTypes::All ),
+	  bFirstTick( true ),
 	  UISystem(),
 	  WorldInteraction()
 {
@@ -350,11 +351,6 @@ void FVREditorMode::Enter()
 			// @todo vreditor: Should save and restore camera position and any other settings we change (viewport type, pitch locking, etc.)
 			SavedEditorState.ViewLocation = VRViewportClient.GetViewLocation();
 			SavedEditorState.ViewRotation = VRViewportClient.GetViewRotation();
-			if( bActuallyUsingVR )
-			{
-				VRViewportClient.SetViewLocation( FVector::ZeroVector );
-				VRViewportClient.SetViewRotation( FRotator::ZeroRotator );
-			}
 
 			// Don't allow the tracking space to pitch up or down.  People hate that in VR.
 			// @todo vreditor: This doesn't seem to prevent people from pitching the camera with RMB drag
@@ -447,6 +443,12 @@ void FVREditorMode::Enter()
 	//	UISystem->ShowEditorUIPanel(FVREditorUISystem::EEditorUIPanel::ContentBrowser, HandIndex, bShouldShow);
 	//}
 
+	if ( bActuallyUsingVR )
+	{
+		 SetWorldToMetersScale( 100.0f );
+	}
+
+	bFirstTick = true;
 	bIsFullyInitialized = true;
 }
 
@@ -568,9 +570,7 @@ void FVREditorMode::Exit()
 	}
 
 	// Make sure we are no longer overriding WorldToMetersScale every frame
-	LAUNCH_API extern float GNewWorldToMetersScale;
-	GNewWorldToMetersScale = 0.0f;
-
+	SetWorldToMetersScale( 0.0f );
 
 	// Call parent implementation
 	FEdMode::Exit();
@@ -1341,10 +1341,34 @@ void FVREditorMode::Tick( FEditorViewportClient* ViewportClient, float DeltaTime
 		}
 	}
 
+	// Setting the initial position and rotation based on the editor viewport when going into VR mode
+	if ( bActuallyUsingVR && bFirstTick )
+	{
+		const FTransform RoomToWorld = GetRoomTransform();
+		const FTransform WorldToRoom = RoomToWorld.Inverse();
+		const FTransform ViewportToWorld = FTransform(SavedEditorState.ViewRotation, SavedEditorState.ViewLocation);
+		const FTransform ViewportToRoom = ViewportToWorld * WorldToRoom;
+
+		FTransform ViewportToRoomYaw = ViewportToRoom;
+		ViewportToRoomYaw.SetRotation(FQuat(FRotator(0.0f, ViewportToRoomYaw.GetRotation().Rotator().Yaw, 0.0f)));
+
+		FTransform HeadToRoomYaw = GetRoomSpaceHeadTransform();
+		HeadToRoomYaw.SetRotation(FQuat(FRotator(0.0f, HeadToRoomYaw.GetRotation().Rotator().Yaw, 0.0f)));
+
+		FTransform RoomToWorldYaw = RoomToWorld;
+		RoomToWorldYaw.SetRotation(FQuat(FRotator(0.0f, RoomToWorldYaw.GetRotation().Rotator().Yaw, 0.0f)));
+
+		FTransform ResultToWorld = ( HeadToRoomYaw.Inverse() * ViewportToRoomYaw ) * RoomToWorldYaw;
+		SetRoomTransform(ResultToWorld);
+	}
+
+
 	// Update floating help text
 	UpdateHelpLabels();
 
 	StopOldHapticEffects();
+
+	bFirstTick = false;
 }
 
 
@@ -1731,7 +1755,16 @@ float FVREditorMode::GetMinScale()
 	return VREd::ScaleMin->GetFloat();
 }
 
-bool FVREditorMode::GetLaserPointer( const int32 HandIndex, FVector& LaserPointerStart, FVector& LaserPointerEnd, const bool bEvenIfUIIsInFront, const float LaserLengthOverride ) const
+void FVREditorMode::SetWorldToMetersScale( const float NewWorldToMetersScale )
+{
+	// @todo vreditor: This is bad because we're clobbering the world settings which will be saved with the map.  Instead we need to 
+	// be able to apply an override before the scene view gets it
+
+	LAUNCH_API extern float GNewWorldToMetersScale;
+	GNewWorldToMetersScale = NewWorldToMetersScale;
+}
+
+bool FVREditorMode::GetLaserPointer(const int32 HandIndex, FVector& LaserPointerStart, FVector& LaserPointerEnd, const bool bEvenIfUIIsInFront, const float LaserLengthOverride) const
 {
 	const FVirtualHand& Hand = VirtualHands[ HandIndex ];
 
diff --git a/Engine/Source/Editor/VREditor/VREditorMode.h b/Engine/Source/Editor/VREditor/VREditorMode.h
--- a/Engine/Source/Editor/VREditor/VREditorMode.h
+++ b/Engine/Source/Editor/VREditor/VREditorMode.h
@@ -304,6 +304,9 @@ public:
 	/** Returns the minimum user scale */
 	float GetMinScale();
 
+	/** Sets GNewWorldToMetersScale */
+	void SetWorldToMetersScale( const float NewWorldToMetersScale );
+
 private:
 
 	/** Called when someone closes a standalone VR Editor window */
@@ -551,5 +554,8 @@ private:
 
 	// All the colors for this mode
 	TArray<FLinearColor> Colors;
+
+	/** If this is the first tick or before */
+	bool bFirstTick;
 };
 
diff --git a/Engine/Source/Editor/VREditor/VREditorWorldInteraction.cpp b/Engine/Source/Editor/VREditor/VREditorWorldInteraction.cpp
--- a/Engine/Source/Editor/VREditor/VREditorWorldInteraction.cpp
+++ b/Engine/Source/Editor/VREditor/VREditorWorldInteraction.cpp
@@ -2312,11 +2312,7 @@ void FVREditorWorldInteraction::UpdateDragging(
 				if( NewWorldToMetersScale != OldWorldToMetersScale &&
 					NewWorldToMetersScale >= Owner.GetMinScale() && NewWorldToMetersScale <= Owner.GetMaxScale() )
 				{
-					// @todo vreditor: This is bad because we're clobbering the world settings which will be saved with the map.  Instead we need to 
-					// be able to apply an override before the scene view gets it
-
-					LAUNCH_API extern float GNewWorldToMetersScale;
-					GNewWorldToMetersScale = NewWorldToMetersScale;
+					Owner.SetWorldToMetersScale( NewWorldToMetersScale );
 
 					// Because the tracking space size has changed, but our head position within that space relative to the origin
 					// of the room is the same (before scaling), we need to offset our location within the tracking space to compensate.
