VR Editor: Setting rotation and location of the viewport when exiting VR Editor mode to latest location and rotation pitch and yaw of the head

diff --git a/Engine/Source/Editor/VREditor/VREditorMode.cpp b/Engine/Source/Editor/VREditor/VREditorMode.cpp
--- a/Engine/Source/Editor/VREditor/VREditorMode.cpp
+++ b/Engine/Source/Editor/VREditor/VREditorMode.cpp
@@ -536,8 +536,12 @@ void FVREditorMode::Exit()
 				VRViewportClient.bAlwaysShowModeWidgetAfterSelectionChanges = SavedEditorState.bAlwaysShowModeWidgetAfterSelectionChanges;
 				VRViewportClient.EngineShowFlags = SavedEditorState.ShowFlags;
 				VRViewportClient.SetGameView( SavedEditorState.bGameView );
-				VRViewportClient.SetViewLocation( SavedEditorState.ViewLocation );
-				VRViewportClient.SetViewRotation( SavedEditorState.ViewRotation );
+				VRViewportClient.SetViewLocation( GetHeadTransform().GetLocation() ); // Use SavedEditorState.ViewLocation to go back to start location when entering VR Editor Mode
+				
+				FRotator HeadRotationNoRoll = GetHeadTransform().GetRotation().Rotator();
+				HeadRotationNoRoll.Roll = 0.0f;
+				VRViewportClient.SetViewRotation(HeadRotationNoRoll); // Use SavedEditorState.ViewRotation to go back to start rot
+				
 				VRViewportClient.SetRealtime( SavedEditorState.bRealTime );
 
 				GNearClippingPlane = SavedEditorState.NearClipPlane;
