Fixed Mac compile error with VR Editor

diff --git a/Engine/Source/Editor/VREditor/VirtualHand.h b/Engine/Source/Editor/VREditor/VirtualHand.h
--- a/Engine/Source/Editor/VREditor/VirtualHand.h
+++ b/Engine/Source/Editor/VREditor/VirtualHand.h
@@ -133,7 +133,7 @@ struct FVirtualHand
 	/** True if this hand is "captured" for each possible action type, meaning that only the active captor should 
 	    handle input events until it is no longer captured.  It's the captors responsibility to set this using 
 		OnVRAction(), or clear it when finished with capturing. */
-	bool IsInputCaptured[ EVRActionType::TotalActionTypes ];
+	bool IsInputCaptured[ (int32)EVRActionType::TotalActionTypes ];
 
 	/** True if we're currently holding the 'SelectAndMove' button down after clicking on an actor */
 	TWeakObjectPtr<UActorComponent> ClickingOnComponent;
