Fixed compiler warnings on Mac

diff --git a/Engine/Source/Editor/VREditor/VREditorFloatingUI.cpp b/Engine/Source/Editor/VREditor/VREditorFloatingUI.cpp
--- a/Engine/Source/Editor/VREditor/VREditorFloatingUI.cpp
+++ b/Engine/Source/Editor/VREditor/VREditorFloatingUI.cpp
@@ -19,24 +19,28 @@ namespace VREd
 
 
 AVREditorFloatingUI::AVREditorFloatingUI()
-	: Scale( 1.0f ),
-	  DockedTo( EDockedTo::Nothing ),
-	  PreviousDockedTo( EDockedTo::Nothing ),
+	: LocalRotation( FRotator( 90.0f, 180.0f, 0.0f ) ),
+	  RelativeOffset( FVector::ZeroVector ),
 	  SlateWidget( nullptr ),
-	  UserWidgetClass( nullptr ),
 	  UserWidget( nullptr ),
-	  Resolution( 0, 0 ),
 	  WidgetComponent( nullptr ),
+	  Scale( 1.0f ),
+	  Resolution( 0, 0 ),
+	  Owner( nullptr ),
+	  DockedTo( EDockedTo::Nothing ),
+	  PreviousDockedTo( EDockedTo::Nothing ),
+	  UserWidgetClass( nullptr ),
 	  bShouldBeVisible(),
 	  FadeAlpha( 1.0f ),
-	  LocalRotation( FRotator( 90.0f, 180.0f, 0.0f ) ),
 	  bCollisionOnShowUI( true ),
 	  FadeDelay( 0.0f ),
+	  InitialScale( 1.0f ),
 	  bIsMoving( false ),
 	  MoveToTransform( FTransform() ),
 	  StartMoveToTransform( FTransform() ),
 	  MoveToAlpha( 0.0f ),
-	  MoveToTime( 0.0f )
+	  MoveToTime( 0.0f ),
+	  MoveToResultDock( EDockedTo::Nothing )
 {
 	const bool bTransient = true;
 	USceneComponent* SceneComponent = CreateDefaultSubobject<USceneComponent>( TEXT( "SceneComponent" ), bTransient );
diff --git a/Engine/Source/Editor/VREditor/VREditorFloatingUI.h b/Engine/Source/Editor/VREditor/VREditorFloatingUI.h
--- a/Engine/Source/Editor/VREditor/VREditorFloatingUI.h
+++ b/Engine/Source/Editor/VREditor/VREditorFloatingUI.h
@@ -205,9 +205,6 @@ private:
 	/** Fade alpha, for visibility transitions */
 	float FadeAlpha;
 
-	/** To check if this UI should rotate towards the head */
-	bool bRotateToHead;
-
 	/** Collision on when showing UI */
 	bool bCollisionOnShowUI;
 
diff --git a/Engine/Source/Editor/VREditor/VREditorMode.cpp b/Engine/Source/Editor/VREditor/VREditorMode.cpp
--- a/Engine/Source/Editor/VREditor/VREditorMode.cpp
+++ b/Engine/Source/Editor/VREditor/VREditorMode.cpp
@@ -100,8 +100,6 @@ FVREditorMode::FVREditorMode()
 	: bWantsToExitMode( false ),
 	  bIsFullyInitialized( false ),
 	  AppTimeModeEntered( FTimespan::Zero() ),
-	  MotionControllerID( 0 ),	// @todo vreditor minor: We only support a single controller, and we assume the first controller are the motion controls
-	  LastFrameNumberInputWasPolled( 0 ),
 	  AvatarMeshActor( nullptr ),
 	  HeadMeshComponent( nullptr ),
 	  WorldMovementGridMeshComponent( nullptr ),
@@ -112,14 +110,16 @@ FVREditorMode::FVREditorMode()
 	  SnapGridActor( nullptr ),
 	  SnapGridMeshComponent( nullptr ),
 	  SnapGridMID( nullptr ),
-	  PostProcessComponent( nullptr ),
 	  ScaleProgressMeshComponent( nullptr ),
 	  CurrentScaleProgressMeshComponent( nullptr ),
 	  UserScaleIndicatorText( nullptr ),
-	  CurrentGizmoType( EGizmoHandleTypes::All ),
-	  bFirstTick( true ),
+	  PostProcessComponent( nullptr ),
+	  MotionControllerID( 0 ),	// @todo vreditor minor: We only support a single controller, and we assume the first controller are the motion controls
+	  LastFrameNumberInputWasPolled( 0 ),
 	  UISystem(),
-	  WorldInteraction()
+	  WorldInteraction(),
+	  CurrentGizmoType( EGizmoHandleTypes::All ),
+	  bFirstTick( true )
 {
 	FEditorDelegates::MapChange.AddRaw( this, &FVREditorMode::OnMapChange );
 	FEditorDelegates::BeginPIE.AddRaw( this, &FVREditorMode::OnBeginPIE );
diff --git a/Engine/Source/Editor/VREditor/VREditorUISystem.cpp b/Engine/Source/Editor/VREditor/VREditorUISystem.cpp
--- a/Engine/Source/Editor/VREditor/VREditorUISystem.cpp
+++ b/Engine/Source/Editor/VREditor/VREditorUISystem.cpp
@@ -80,13 +80,20 @@ FVREditorUISystem::FVREditorUISystem( FVREditorMode& InitOwner )
 	  FloatingUIs(),
 	  QuickMenuUI( nullptr ),
 	  QuickRadialMenu( nullptr ),
+	  RadialMenuHideDelayTime( 0.0f ),
 	  QuickMenuWidgetClass( nullptr ),
 	  QuickRadialWidgetClass( nullptr ),
-	  DraggingUI( nullptr ),
+	  TutorialWidgetClass( nullptr ),
 	  DraggingUIHandIndex( INDEX_NONE ),
 	  DraggingUIOffsetTransform( FTransform::Identity ),
-	  RadialMenuHideDelayTime( 0.0f ),
+	  DraggingUI( nullptr ),
 	  bRefocusViewport(false),
+	  StartDragUISound( nullptr ),
+	  StopDragUISound( nullptr ),
+	  HideUISound( nullptr ),
+	  ShowUISound( nullptr ),
+	  bDraggedDockFromHandPassedThreshold( false ),
+	  LastDraggingHoverLocation( FVector::ZeroVector ),
 	  bSetDefaultLayout( true )
 {
 	// Register to find out about VR events
diff --git a/Engine/Source/Editor/VREditor/VREditorWorldInteraction.cpp b/Engine/Source/Editor/VREditor/VREditorWorldInteraction.cpp
--- a/Engine/Source/Editor/VREditor/VREditorWorldInteraction.cpp
+++ b/Engine/Source/Editor/VREditor/VREditorWorldInteraction.cpp
@@ -109,36 +109,36 @@ struct DVector
 	double Y;
 	double Z;
 
-	DVector::DVector()
+	DVector()
 	{
 	}
 
-	DVector::DVector( const DVector& V )
+	DVector( const DVector& V )
 		: X( V.X ), Y( V.Y ), Z( V.Z )
 	{
 	}
 
-	DVector::DVector( const FVector& V )
+	DVector( const FVector& V )
 		: X( V.X ), Y( V.Y ), Z( V.Z )
 	{
 	}
 
-	DVector::DVector( double InX, double InY, double InZ )
+	DVector( double InX, double InY, double InZ )
 		: X( InX ), Y( InY ), Z( InZ )
 	{
 	}
 
-	DVector DVector::operator+( const DVector& V ) const
+	DVector operator+( const DVector& V ) const
 	{
 		return DVector( X + V.X, Y + V.Y, Z + V.Z );
 	}
 
-	DVector DVector::operator-( const DVector& V ) const
+	DVector operator-( const DVector& V ) const
 	{
 		return DVector( X - V.X, Y - V.Y, Z - V.Z );
 	}
 
-	DVector DVector::operator*( double Scale ) const
+	DVector operator*( double Scale ) const
 	{
 		return DVector( X * Scale, Y * Scale, Z * Scale );
 	}
