Fixed more Mac compiler warnings.  Few more to go.

diff --git a/Engine/Source/Editor/VREditor/VREditorMode.cpp b/Engine/Source/Editor/VREditor/VREditorMode.cpp
--- a/Engine/Source/Editor/VREditor/VREditorMode.cpp
+++ b/Engine/Source/Editor/VREditor/VREditorMode.cpp
@@ -589,8 +589,10 @@ void FVREditorMode::SpawnAvatarMeshActor()
 	// Setup our avatar
 	if( AvatarMeshActor == nullptr )
 	{
-		const bool bWithSceneComponent = true;
-		AvatarMeshActor = SpawnTransientSceneActor( AActor::StaticClass(), TEXT( "AvatarMesh" ), bWithSceneComponent );
+		{
+			const bool bWithSceneComponent = true;
+			AvatarMeshActor = SpawnTransientSceneActor( AActor::StaticClass(), TEXT( "AvatarMesh" ), bWithSceneComponent );
+		}
 
 		// Give us a head mesh
 		{
@@ -1325,16 +1327,16 @@ void FVREditorMode::Tick( FEditorViewportClient* ViewportClient, float DeltaTime
 		const bool bShouldDrawWorldMovementPostProcess = WorldMovementGridOpacity > KINDA_SMALL_NUMBER;
 		if( bShouldDrawWorldMovementPostProcess != bIsDrawingWorldMovementPostProcess )
 		{
-			UPostProcessComponent* PostProcessComponent = GetPostProcessComponent();
+			UPostProcessComponent* PostProcess = GetPostProcessComponent();
 			if( bShouldDrawWorldMovementPostProcess )
 			{
-				PostProcessComponent->AddOrUpdateBlendable( WorldMovementPostProcessMaterial );
+				PostProcess->AddOrUpdateBlendable( WorldMovementPostProcessMaterial );
 				bIsDrawingWorldMovementPostProcess = true;
 			}
 			else
 			{
 				bIsDrawingWorldMovementPostProcess = false;
-				PostProcessComponent->Settings.RemoveBlendable( WorldMovementPostProcessMaterial );
+				PostProcess->Settings.RemoveBlendable( WorldMovementPostProcessMaterial );
 			}
 		}
 
diff --git a/Engine/Source/Editor/VREditor/VREditorUISystem.cpp b/Engine/Source/Editor/VREditor/VREditorUISystem.cpp
--- a/Engine/Source/Editor/VREditor/VREditorUISystem.cpp
+++ b/Engine/Source/Editor/VREditor/VREditorUISystem.cpp
@@ -215,8 +215,6 @@ void FVREditorUISystem::OnVRAction( FEditorViewportClient& ViewportClient, const
 								( Event == IE_Pressed && Hand.bIsModifierPressed ) ||
 								( Event == IE_Released && Hand.bIsRightClickingOnUI );
 
-							UWidgetComponent* WidgetComponent = Cast<UWidgetComponent>( HitResult.GetComponent() );
-							if( WidgetComponent )
 							{
 								FVector2D LastLocalHitLocation = WidgetComponent->GetLastLocalHitLocation();
 
@@ -395,7 +393,6 @@ void FVREditorUISystem::OnVRHoverUpdate( FEditorViewportClient& ViewportClient,
 					FWidgetPath WidgetPathUnderFinger = FWidgetPath( WidgetComponent->GetHitWidgetPath( HitResult.ImpactPoint, /*bIgnoreEnabledStatus*/ false ) );
 					if ( WidgetPathUnderFinger.IsValid() )
 					{
-						FVirtualHand& Hand = Owner.GetVirtualHand(HandIndex);
 						Hand.bIsHovering = true;
 						Hand.HoverLocation = HitResult.ImpactPoint;
 
diff --git a/Engine/Source/Editor/VREditor/VREditorWorldInteraction.cpp b/Engine/Source/Editor/VREditor/VREditorWorldInteraction.cpp
--- a/Engine/Source/Editor/VREditor/VREditorWorldInteraction.cpp
+++ b/Engine/Source/Editor/VREditor/VREditorWorldInteraction.cpp
@@ -267,16 +267,21 @@ FVREditorWorldInteraction::FVREditorWorldInteraction( FVREditorMode& InitOwner )
 	  TransformablesInterpolationStartTime( FTimespan::Zero() ),
 	  TransformablesInterpolationDuration( 1.0f ),
 	  PlacingMaterialOrTextureAsset( nullptr ),
+	  LastWorldToMetersScale( 100.0f ),
 	  TransformGizmoActor( nullptr ),
 	  GizmoLocalBounds( FBox( 0 ) ),
 	  bDraggedSinceLastSelection( false ),
 	  LastDragGizmoStartTransform( FTransform::Identity ),
 	  bIsTeleporting( false ),
 	  TeleportLerpAlpha( 0 ),
-	  DockSelectDistance ( 100 ),
+	  TeleportStartLocation( FVector::ZeroVector ),
+	  TeleportGoalLocation( FVector::ZeroVector ),
 	  bJustTeleported( false ),
+	  TeleportSound( nullptr ),
+	  DropMaterialOrMaterialSound( nullptr ),
+	  DockSelectDistance( 100 ),
 	  FloatingUIAssetDraggedFrom( nullptr )
-{
+	{
 	// Register to find out about VR input events
 	Owner.OnVRAction().AddRaw( this, &FVREditorWorldInteraction::OnVRAction );
 	Owner.OnVRHoverUpdate().AddRaw( this, &FVREditorWorldInteraction::OnVRHoverUpdate );
