Fixed more Mac compiler warnings (think I got them all!)

diff --git a/Engine/Source/Editor/VREditor/VREditorWorldInteraction.cpp b/Engine/Source/Editor/VREditor/VREditorWorldInteraction.cpp
--- a/Engine/Source/Editor/VREditor/VREditorWorldInteraction.cpp
+++ b/Engine/Source/Editor/VREditor/VREditorWorldInteraction.cpp
@@ -499,7 +499,6 @@ void FVREditorWorldInteraction::OnVRAction( FEditorViewportClient& ViewportClien
 
 								// If the other hand is holding down the button after selecting an object, we'll allow this hand to toggle selection
 								// of additional objects (multi select)
-								FVirtualHand& OtherHand = Owner.GetOtherHand( VRAction.HandIndex );
 								{
 									if( OtherHand.ClickingOnComponent.IsValid() &&		// Other hand is clicking on something
 										Actor != TransformGizmoActor &&					// Not clicking on a gizmo
@@ -990,11 +989,10 @@ void FVREditorWorldInteraction::Tick( FEditorViewportClient* ViewportClient, con
 								// it's position over a very short timespan
 								if( !bIsInterpolatingTransformablesFromSnapshotTransform )	// Let the last animation finish first
 								{
-									const float DragDelta = ( HitLocation - Hand.LastDragToLocation ).Size() * WorldScaleFactor;
-									if( DragDelta >= VREd::DragAtLaserImpactInterpolationThreshold->GetFloat() )
+									const float ScaledDragDistance = DragDelta.Size() * WorldScaleFactor;
+									if( ScaledDragDistance >= VREd::DragAtLaserImpactInterpolationThreshold->GetFloat() )
 									{
 										bIsInterpolatingTransformablesFromSnapshotTransform = true;
-										const FTimespan CurrentTime = FTimespan::FromSeconds( FPlatformTime::Seconds() );
 										TransformablesInterpolationStartTime = CurrentTime;
 										TransformablesInterpolationDuration = VREd::DragAtLaserImpactInterpolationDuration->GetFloat();
 
