VREditor : Pivot Transform Gizmo first version
- Switch to old gizmo with command: VREd.PivotPointTransformGizmo 0

diff --git a/Engine/Source/Editor/VREditor/Gizmo/VREditorBaseTransformGizmo.cpp b/Engine/Source/Editor/VREditor/Gizmo/VREditorBaseTransformGizmo.cpp
--- a/Engine/Source/Editor/VREditor/Gizmo/VREditorBaseTransformGizmo.cpp
+++ b/Engine/Source/Editor/VREditor/Gizmo/VREditorBaseTransformGizmo.cpp
@@ -8,9 +8,9 @@
 
 namespace VREd
 {
-	// @todo vreditor tweak: Tweak out console variables
 	static FAutoConsoleVariable GizmoSelectionAnimationDuration( TEXT( "VREd.GizmoSelectionAnimationDuration" ), 0.15f, TEXT( "How long to animate the gizmo after objects are selected" ) );
 	static FAutoConsoleVariable GizmoSelectionAnimationCurvePower( TEXT( "VREd.GizmoSelectionAnimationCurvePower" ), 2.0f, TEXT( "Controls the animation curve for the gizmo after objects are selected" ) );
+	static FAutoConsoleVariable GizmoShowMeasurementText( TEXT( "VREd.GizmoShowMeasurementText" ), 0, TEXT( "When enabled, gizmo measurements will always be visible.  Otherwise, only when hovering over a scale/stretch gizmo handle" ) );
 }
 
 ABaseTransformGizmo::ABaseTransformGizmo( ) :
@@ -44,6 +44,70 @@
 	SelectedAtTime = FTimespan::FromSeconds( FApp::GetCurrentTime() );
 }
 
+ETransformGizmoInteractionType ABaseTransformGizmo::GetInteractionType( UActorComponent* DraggedComponent, TOptional<FTransformGizmoHandlePlacement>& OutHandlePlacement )
+{
+	OutHandlePlacement.Reset();
+	if ( DraggedComponent != nullptr )
+	{
+		UStaticMeshComponent* DraggedMesh = Cast<UStaticMeshComponent>( DraggedComponent );
+		if ( DraggedMesh != nullptr )
+		{
+			ETransformGizmoInteractionType ResultInteractionType;
+
+			for ( UVREditorGizmoHandleGroup* HandleGroup : AllHandleGroups )
+			{
+				if ( HandleGroup != nullptr )
+				{
+					int32 HandIndex = HandleGroup->GetDraggedHandleIndex( DraggedMesh );
+					if ( HandIndex != INDEX_NONE )
+					{
+						HandleGroup->GetHandleIndexInteractionType( HandIndex, ResultInteractionType, OutHandlePlacement );
+						return ResultInteractionType;
+					}
+				}
+			}
+		}
+	}
+
+	OutHandlePlacement.Reset();
+	return ETransformGizmoInteractionType::Translate;
+}
+
+float ABaseTransformGizmo::GetAnimationAlpha()
+{
+	// Update animation
+	float AnimationAlpha = 0.0f;
+	{
+		const FTimespan CurrentTime = FTimespan::FromSeconds( FApp::GetCurrentTime() );
+		const float TimeSinceSelectionChange = (CurrentTime - SelectedAtTime).GetTotalSeconds();
+		const float AnimLength = VREd::GizmoSelectionAnimationDuration->GetFloat();
+		if ( TimeSinceSelectionChange < AnimLength )
+		{
+			AnimationAlpha = FMath::Max( 0.0f, TimeSinceSelectionChange / AnimLength );
+		}
+		else
+		{
+			AnimationAlpha = 1.0f;
+		}
+
+		// Apply a bit of a curve to the animation
+		AnimationAlpha = FMath::Pow( AnimationAlpha, VREd::GizmoSelectionAnimationCurvePower->GetFloat() );
+	}
+
+	return AnimationAlpha;
+}
+
+void ABaseTransformGizmo::SetOwnerMode( FVREditorMode* InOwner )
+{
+	Owner = InOwner;
+}
+
+FVREditorMode* ABaseTransformGizmo::GetOwnerMode() const
+{
+	return Owner;
+}
+
+
 void ABaseTransformGizmo::GetBoundingBoxEdge( const FBox& Box, const int32 AxisIndex, const int32 EdgeIndex, FVector& OutVertex0, FVector& OutVertex1 )
 {
 	check( AxisIndex >= 0 && AxisIndex < 3 );
@@ -145,39 +209,7 @@
 }
 
 
-float ABaseTransformGizmo::GetAnimationAlpha()
-{
-	// Update animation
-	float AnimationAlpha = 0.0f;
-	{
-		const FTimespan CurrentTime = FTimespan::FromSeconds( FApp::GetCurrentTime() );
-		const float TimeSinceSelectionChange = (CurrentTime - SelectedAtTime).GetTotalSeconds();
-		const float AnimLength = VREd::GizmoSelectionAnimationDuration->GetFloat();
-		if ( TimeSinceSelectionChange < AnimLength )
-		{
-			AnimationAlpha = FMath::Max( 0.0f, TimeSinceSelectionChange / AnimLength );
-		}
-		else
-		{
-			AnimationAlpha = 1.0f;
-		}
 
-		// Apply a bit of a curve to the animation
-		AnimationAlpha = FMath::Pow( AnimationAlpha, VREd::GizmoSelectionAnimationCurvePower->GetFloat() );
-	}
 
-	return AnimationAlpha;
-}
-
-void ABaseTransformGizmo::SetOwnerMode( FVREditorMode* InOwner )
-{
-	Owner = InOwner;
-}
-
-FVREditorMode* ABaseTransformGizmo::GetOwnerMode() const
-{
-	return Owner;
-}
-
 void ABaseTransformGizmo::UpdateHandleVisibility( const EGizmoHandleTypes GizmoType, const ECoordSystem GizmoCoordinateSpace, const bool bAllHandlesVisible, UActorComponent* DraggingHandle )
 {
@@ -185,14 +217,14 @@
 	{
 		if ( HandleGroup != nullptr )
 		{
-			const bool bIsSameType = ( GizmoType == EGizmoHandleTypes::All ) || HandleGroup->GetHandleType() == GizmoType;
+			const bool bIsTypeSupported = ( GizmoType == EGizmoHandleTypes::All && HandleGroup->GetShowOnUniversalGizmo() ) || HandleGroup->GetHandleType() == GizmoType;
 			const bool bSupportsCurrentCoordinateSpace = HandleGroup->SupportsWorldCoordinateSpace() || GizmoCoordinateSpace != COORD_World;
 
 			for ( FVREditorGizmoHandle& Handle : HandleGroup->GetHandles() )
 			{
 				if( Handle.HandleMesh != nullptr )
 				{
-					const bool bShowIt = ( bIsSameType && bSupportsCurrentCoordinateSpace && bAllHandlesVisible ) || ( DraggingHandle != nullptr && DraggingHandle == Handle.HandleMesh );
+					const bool bShowIt = ( bIsTypeSupported && bSupportsCurrentCoordinateSpace && bAllHandlesVisible ) || ( DraggingHandle != nullptr && DraggingHandle == Handle.HandleMesh );
 
 					Handle.HandleMesh->SetVisibility( bShowIt );
 
@@ -204,33 +236,9 @@
 	}
 }
 
-ETransformGizmoInteractionType ABaseTransformGizmo::GetInteractionType( UActorComponent* DraggedComponent, TOptional<FTransformGizmoHandlePlacement>& OutHandlePlacement )
+bool ABaseTransformGizmo::GetShowMeasurementText() const
 {
-	OutHandlePlacement.Reset();
-	if (DraggedComponent != nullptr)
-	{
-		UStaticMeshComponent* DraggedMesh = Cast<UStaticMeshComponent>( DraggedComponent );
-		if (DraggedMesh != nullptr)
-		{
-			ETransformGizmoInteractionType ResultInteractionType;
-
-			for ( UVREditorGizmoHandleGroup* HandleGroup : AllHandleGroups)
-			{
-				if (HandleGroup != nullptr)
-				{
-					int32 HandIndex = HandleGroup->GetDraggedHandleIndex( DraggedMesh );
-					if (HandIndex != INDEX_NONE)
-					{
-						HandleGroup->GetHandleIndexInteractionType( HandIndex, ResultInteractionType, OutHandlePlacement );
-						return ResultInteractionType;
-					}
-				}
-			}
-		}
-	}
-
-	OutHandlePlacement.Reset();
-	return ETransformGizmoInteractionType::Translate;
+	return VREd::GizmoShowMeasurementText->GetInt() != 0;
 }
 
 void FTransformGizmoHandlePlacement::GetCenterHandleCountAndFacingAxisIndex( int32& OutCenterHandleCount, int32& OutFacingAxisIndex, int32& OutCenterAxisIndex ) const
diff --git a/Engine/Source/Editor/VREditor/Gizmo/VREditorBaseTransformGizmo.h b/Engine/Source/Editor/VREditor/Gizmo/VREditorBaseTransformGizmo.h
--- a/Engine/Source/Editor/VREditor/Gizmo/VREditorBaseTransformGizmo.h
+++ b/Engine/Source/Editor/VREditor/Gizmo/VREditorBaseTransformGizmo.h
@@ -47,10 +47,13 @@ public:
 	/** Gets the owner */
 	class FVREditorMode* GetOwnerMode() const;
 
-protected:
+	/** Called by the world interaction system after we've been spawned into the world, to allow
+	    us to create components and set everything up nicely for the selected objects that we'll be
+		used to manipulate */
+	virtual void UpdateGizmo( const EGizmoHandleTypes GizmoType, const ECoordSystem GizmoCoordinateSpace, const FTransform& LocalToWorld, const FBox& LocalBounds, const FVector ViewLocation, bool bAllHandlesVisible, 
+		class UActorComponent* DraggingHandle, const TArray< UActorComponent* >& HoveringOverHandles, const float GizmoHoverScale, const float GizmoHoverAnimationDuration ) {}
 
-	/** Makes up a name string for a handle */
-	FString MakeHandleName( const FTransformGizmoHandlePlacement HandlePlacement ) const;
+protected:
 
 	/** Static: Given a bounding box and information about which edge to query, returns the vertex positions of that edge */
 	static void GetBoundingBoxEdge( const FBox& Box, const int32 AxisIndex, const int32 EdgeIndex, FVector& OutVertex0, FVector& OutVertex1 );
@@ -58,6 +61,9 @@ protected:
 	/** Updates the visibility of all the handles */
 	void UpdateHandleVisibility( const EGizmoHandleTypes GizmoType, const ECoordSystem GizmoCoordinateSpace, bool bAllHandlesVisible, UActorComponent* DraggingHandle );
 
+	/** Gets if the gizmo shows measurement texts */
+	bool GetShowMeasurementText() const;
+
 protected:
 	
 	/** Real time that the gizmo was last attached to a selected set of objects.  This is used for animation transitions */
diff --git a/Engine/Source/Editor/VREditor/Gizmo/VREditorGizmoHandle.cpp b/Engine/Source/Editor/VREditor/Gizmo/VREditorGizmoHandle.cpp
--- a/Engine/Source/Editor/VREditor/Gizmo/VREditorGizmoHandle.cpp
+++ b/Engine/Source/Editor/VREditor/Gizmo/VREditorGizmoHandle.cpp
@@ -6,7 +6,8 @@
 #include "VREditorMode.h"
 
 UVREditorGizmoHandleGroup::UVREditorGizmoHandleGroup()
-	: Super()
+	: Super(),
+	bShowOnUniversalGizmo( true )
 {
 
 }
@@ -112,12 +113,13 @@ ETransformGizmoInteractionType UVREditorGizmoHandleGroup::GetInteractionType() c
 	return ETransformGizmoInteractionType::Translate;
 }
 
-void UVREditorGizmoHandleGroup::UpdateGizmoHandleGroup( const FTransform& LocalToWorld, const FBox& LocalBounds, const FVector ViewLocation, bool bAllHandlesVisible, class UActorComponent* DraggingHandle, const TArray< UActorComponent* >& HoveringOverHandles, float AnimationAlpha, float GizmoScale, const float GizmoHoverScale, const float GizmoHoverAnimationDuration, bool& bOutIsHoveringOrDraggingThisHandleGroup )
+void UVREditorGizmoHandleGroup::UpdateGizmoHandleGroup(const FTransform& LocalToWorld, const FBox& LocalBounds, const FVector ViewLocation, bool bAllHandlesVisible, class UActorComponent* DraggingHandle, 
+	const TArray< UActorComponent* >& HoveringOverHandles, float AnimationAlpha, float GizmoScale, const float GizmoHoverScale, const float GizmoHoverAnimationDuration, bool& bOutIsHoveringOrDraggingThisHandleGroup )
 {
-	UpdateHoverAnimation( DraggingHandle, HoveringOverHandles, GizmoHoverAnimationDuration, bOutIsHoveringOrDraggingThisHandleGroup );
+	UpdateHoverAnimation(DraggingHandle, HoveringOverHandles, GizmoHoverAnimationDuration, bOutIsHoveringOrDraggingThisHandleGroup);
 }
 
-int32 UVREditorGizmoHandleGroup::GetDraggedHandleIndex( class UStaticMeshComponent* DraggedMesh )
+int32 UVREditorGizmoHandleGroup::GetDraggedHandleIndex(class UStaticMeshComponent* DraggedMesh)
 {
 	for( int32 HandleIndex = 0; HandleIndex < Handles.Num(); ++HandleIndex )
 	{
@@ -149,6 +151,16 @@ EGizmoHandleTypes UVREditorGizmoHandleGroup::GetHandleType() const
 	return EGizmoHandleTypes::All;
 }
 
+void UVREditorGizmoHandleGroup::SetShowOnUniversalGizmo( const bool bInShowOnUniversal )
+{
+	bShowOnUniversalGizmo = bInShowOnUniversal;
+}
+
+bool UVREditorGizmoHandleGroup::GetShowOnUniversalGizmo() const
+{
+	return bShowOnUniversalGizmo;
+}
+
 void UVREditorGizmoHandleGroup::UpdateHandleColor( const int32 AxisIndex, FVREditorGizmoHandle& Handle, class UActorComponent* DraggingHandle, const TArray< UActorComponent* >& HoveringOverHandles )
 {
 	UStaticMeshComponent* HandleMesh = Handle.HandleMesh;
@@ -278,3 +290,98 @@ void UVREditorGizmoHandleGroup::UpdateHoverAnimation( UActorComponent* DraggingH
 		Handle.HoverAlpha = FMath::Clamp( Handle.HoverAlpha, 0.0f, 1.0f );
 	}
 }
+
+void UVREditorAxisGizmoHandleGroup::CreateHandles( UStaticMesh* HandleMesh, const FString& HandleComponentName )
+{
+	for ( int32 X = 0; X < 3; ++X )
+	{
+		for ( int32 Y = 0; Y < 3; ++Y )
+		{
+			for ( int32 Z = 0; Z < 3; ++Z )
+			{
+				FTransformGizmoHandlePlacement HandlePlacement = GetHandlePlacement(X, Y, Z);
+
+				int32 CenterHandleCount, FacingAxisIndex, CenterAxisIndex;
+				HandlePlacement.GetCenterHandleCountAndFacingAxisIndex( /* Out */ CenterHandleCount, /* Out */ FacingAxisIndex, /* Out */ CenterAxisIndex);
+
+				// Don't allow translation/stretching/rotation from the origin
+				if ( CenterHandleCount < 3 )
+				{
+					const FString HandleName = MakeHandleName(HandlePlacement);
+
+					// Only happens on the center of an axis. And we only bother drawing one for the "positive" direction.
+					if ( CenterHandleCount == 2 && HandlePlacement.Axes[FacingAxisIndex] == ETransformGizmoHandleDirection::Positive )
+					{
+						// Plane translation handle
+						FString ComponentName = HandleName + HandleName;
+						CreateAndAddMeshHandle( HandleMesh, ComponentName, HandlePlacement );
+					}
+				}
+			}
+		}
+	}
+}
+
+void UVREditorAxisGizmoHandleGroup::UpdateHandlesRelativeTransformOnAxis( const FTransform& HandleToCenter, const float AnimationAlpha, const float GizmoScale, const float GizmoHoverScale,
+	const FVector& ViewLocation, class UActorComponent* DraggingHandle, const TArray< UActorComponent* >& HoveringOverHandles )
+{
+	const float WorldScaleFactor = GetWorld()->GetWorldSettings()->WorldToMeters / 100.0f;
+
+	for (int32 HandleIndex = 0; HandleIndex < Handles.Num(); ++HandleIndex)
+	{
+		FVREditorGizmoHandle& Handle = Handles[HandleIndex];
+		UStaticMeshComponent* GizmoHandleMeshComponent = Handle.HandleMesh;
+		if (GizmoHandleMeshComponent != nullptr)	// Can be null if no handle for this specific placement
+		{
+			const FTransformGizmoHandlePlacement HandlePlacement = MakeHandlePlacementForIndex( HandleIndex );
+
+			int32 CenterHandleCount, FacingAxisIndex, CenterAxisIndex;
+			HandlePlacement.GetCenterHandleCountAndFacingAxisIndex( /* Out */ CenterHandleCount, /* Out */ FacingAxisIndex, /* Out */ CenterAxisIndex);
+			
+			FVector GizmoSpaceFacingAxisVector = GetAxisVector( FacingAxisIndex, HandlePlacement.Axes[FacingAxisIndex] );
+			
+			// Check on which side we are relative to the gizmo
+			const FVector GizmoSpaceViewLocation = GetOwner()->GetTransform().InverseTransformPosition( ViewLocation );
+			if ( GizmoSpaceViewLocation[ FacingAxisIndex ] < 0 )
+			{
+				GizmoSpaceFacingAxisVector[ FacingAxisIndex ] *= -1.0f;
+			}
+
+			const FTransform GizmoOriginToFacingAxisRotation( GizmoSpaceFacingAxisVector.ToOrientationQuat() );
+	
+			FTransform HandleToGizmoOrigin = HandleToCenter * GizmoOriginToFacingAxisRotation;
+				
+			// Check on the axis if the offset is on the other side of the object compared to the view. Switch the offset to the side of the view if it does
+			FVector GizmoSpaceFacingAxisOffset = HandleToGizmoOrigin.GetLocation();
+			for ( int32 AxisIndex = 0; AxisIndex < 3; ++AxisIndex )
+			{		
+				if ( AxisIndex != FacingAxisIndex &&
+					( GizmoSpaceFacingAxisOffset[ AxisIndex ] > 0 && GizmoSpaceViewLocation[ AxisIndex ] < 0 ) ||
+					( GizmoSpaceFacingAxisOffset[ AxisIndex ] < 0 && GizmoSpaceViewLocation[ AxisIndex ] > 0 ) )
+				{
+					GizmoSpaceFacingAxisOffset[ AxisIndex ] *= -1.0f;
+				}	
+			}
+
+			GizmoSpaceFacingAxisOffset *= AnimationAlpha;
+
+			HandleToGizmoOrigin.SetLocation( GizmoSpaceFacingAxisOffset * GizmoScale );
+
+			// Set the final transform
+			GizmoHandleMeshComponent->SetRelativeTransform( HandleToGizmoOrigin );
+			
+			float GizmoHandleScale = GizmoScale;
+			if ( GizmoScale != GizmoHoverScale )
+			{
+				// Make the handle bigger while hovered (but don't affect the offset -- we want it to scale about it's origin)
+				GizmoHandleScale *= FMath::Lerp( 1.0f, GizmoHoverScale, Handle.HoverAlpha );
+			}
+			
+			GizmoHandleScale *= AnimationAlpha;
+			GizmoHandleMeshComponent->SetRelativeScale3D( FVector( GizmoHandleScale ) );
+
+			// Update material
+			UpdateHandleColor( FacingAxisIndex, Handle, DraggingHandle, HoveringOverHandles );
+		}
+	}
+}
diff --git a/Engine/Source/Editor/VREditor/Gizmo/VREditorGizmoHandle.h b/Engine/Source/Editor/VREditor/Gizmo/VREditorGizmoHandle.h
--- a/Engine/Source/Editor/VREditor/Gizmo/VREditorGizmoHandle.h
+++ b/Engine/Source/Editor/VREditor/Gizmo/VREditorGizmoHandle.h
@@ -6,21 +6,6 @@
 #include "VREditorWorldInteractionTypes.h"
 #include "VREditorGizmoHandle.generated.h"
 
-namespace VREd
-{
-	namespace GizmoColor
-	{
-		const FLinearColor RedGizmoColor( 0.4f, 0.05f, 0.05f, 1.0f );
-		const FLinearColor GreenGizmoColor( 0.05f, 0.4f, 0.05f, 1.0f );
-		const FLinearColor BlueGizmoColor( 0.05f, 0.05f, 0.4f, 1.0f );
-		const FLinearColor WhiteGizmoColor( 0.7f, 0.7f, 0.7f, 1.0f );
-
-		// @todo vreditor: Use configured editor selection color instead
-		const FLinearColor HoverGizmoColor( FLinearColor::Yellow );
-		const FLinearColor DraggingGizmoColor( FLinearColor::White );
-	}
-}
-
 enum class EGizmoHandleTypes : uint8;
 
 USTRUCT()
@@ -62,7 +47,8 @@ public:
 	static FVector GetAxisVector( const int32 AxisIndex, const ETransformGizmoHandleDirection HandleDirection );
 
 	/** Updates the Gizmo handles, needs to be implemented by derived classes */
-	virtual void UpdateGizmoHandleGroup( const FTransform& LocalToWorld, const FBox& LocalBounds, const FVector ViewLocation, bool bAllHandlesVisible, class UActorComponent* DraggingHandle, const TArray< UActorComponent* >& HoveringOverHandles, float AnimationAlpha, float GizmoScale, const float GizmoHoverScale, const float GizmoHoverAnimationDuration, bool& bOutIsHoveringOrDraggingThisHandleGroup );
+	virtual void UpdateGizmoHandleGroup(  const FTransform& LocalToWorld, const FBox& LocalBounds, const FVector ViewLocation, bool bAllHandlesVisible, class UActorComponent* DraggingHandle, const TArray< UActorComponent* >& HoveringOverHandles, 
+		float AnimationAlpha, float GizmoScale, const float GizmoHoverScale, const float GizmoHoverAnimationDuration, bool& bOutIsHoveringOrDraggingThisHandleGroup );
 
 	/** Gets the InteractionType and the HandlePlacement for this Gizmo handle */
 	virtual void GetHandleIndexInteractionType( const int32 HandleIndex, ETransformGizmoInteractionType& OutInteractionType, TOptional<FTransformGizmoHandlePlacement>& OutHandlePlacement );
@@ -91,6 +77,12 @@ public:
 		return true;
 	}
 
+	/** Sets if this handlegroup will be visible with the universal gizmo */
+	void SetShowOnUniversalGizmo( const bool bShowHandleUniversal );
+	
+	/** Gets if this handlegroup will be visible with the universal gizmo */
+	bool GetShowOnUniversalGizmo() const;
+
 protected:
 	/** Updates the colors of the dynamic material instances for the handle passed using its axis index */	
 	void UpdateHandleColor( const int32 AxisIndex, FVREditorGizmoHandle& Handle, class UActorComponent* DraggingHandle, const TArray< UActorComponent* >& HoveringOverHandles );
@@ -123,4 +115,29 @@ private:
 
 	/** Updates the hover animation for the HoveringOverHandles */
 	void UpdateHoverAnimation( class UActorComponent* DraggingHandle, const TArray< UActorComponent* >& HoveringOverHandles, const float GizmoHoverAnimationDuration, bool& bOutIsHoveringOrDraggingThisHandleGroup );
+
+	/** If this handlegroup will be visible with the universal gizmo */
+	bool bShowOnUniversalGizmo;
+};
+
+/**
+ * Base class for gizmo handles on axis
+ */
+UCLASS(ABSTRACT)
+class UVREditorAxisGizmoHandleGroup : public UVREditorGizmoHandleGroup
+{
+	GENERATED_BODY()
+
+protected:
+
+	/** Creates mesh for every axis */
+	void CreateHandles(UStaticMesh* HandleMesh, const FString& HandleComponentName);
+	
+
+	/** Calculates the transforms of meshes of this handlegroup 
+	 * @param FTransform - The offset from the center of the gizmo in roomspace
+	 */
+	void UpdateHandlesRelativeTransformOnAxis( const FTransform& HandleToCenter, const float AnimationAlpha, const float GizmoScale, const float GizmoHoverScale, 
+		const FVector& ViewLocation, class UActorComponent* DraggingHandle, const TArray< UActorComponent* >& HoveringOverHandles );
+
 };
\ No newline at end of file
diff --git a/Engine/Source/Editor/VREditor/Gizmo/VREditorPivotTransformGizmo.cpp b/Engine/Source/Editor/VREditor/Gizmo/VREditorPivotTransformGizmo.cpp
new file mode 100644
--- /dev/null
+++ b/Engine/Source/Editor/VREditor/Gizmo/VREditorPivotTransformGizmo.cpp
@@ -0,0 +1,254 @@
+// Copyright 1998-2015 Epic Games, Inc. All Rights Reserved.
+
+#include "VREditorModule.h"
+#include "VREditorPivotTransformGizmo.h"
+#include "VREditorTranslationGizmoHandle.h"
+#include "VREditorRotationGizmoHandle.h"
+#include "VREditorPlaneTranslationGizmoHandle.h"
+#include "VREditorStretchGizmoHandle.h"
+#include "VREditorUniformScaleGizmoHandle.h"
+#include "VREditorMode.h"
+
+namespace VREd //@todo VREditor: Duplicates of TransformGizmo
+{
+	// @todo vreditor tweak: Tweak out console variables
+	static FAutoConsoleVariable PivotGizmoDistanceScaleFactor( TEXT( "VREd.PivotGizmoDistanceScaleFactor" ), 0.003f, TEXT( "How much the gizmo handles should increase in size with distance from the camera, to make it easier to select"));
+	static FAutoConsoleVariable PivotGizmoTranslationPivotOffsetX( TEXT("VREd.PivotGizmoTranslationPivotOffsetX" ), 20.0f, TEXT( "How much the translation cylinder is offsetted from the pivot" ) );
+	static FAutoConsoleVariable PivotGizmoScalePivotOffsetX( TEXT( "VREd.PivotGizmoScalePivotOffsetX" ), 120.0f, TEXT( "How much the non-uniform scale is offsetted from the pivot" ) );
+	static FAutoConsoleVariable PivotGizmoPlaneTranslationPivotOffsetYZ(TEXT("VREd.PivotGizmoPlaneTranslationPivotOffsetYZ" ), 40.0f, TEXT( "How much the plane translation is offsetted from the pivot" ) );
+	static FAutoConsoleVariable PivotGizmoTranslationScaleMultiply( TEXT( "VREd.PivotGizmoTranslationScaleMultiply" ), 2.0f, TEXT( "Multiplies translation handles scale" ) );
+	static FAutoConsoleVariable PivotGizmoTranslationHoverScaleMultiply( TEXT( "VREd.PivotGizmoTranslationHoverScaleMultiply" ), 2.5f, TEXT( "Multiplies translation handles hover scale" ) );
+}
+
+APivotTransformGizmo::APivotTransformGizmo() :
+	Super()
+{
+	UniformScaleGizmoHandleGroup = CreateDefaultSubobject<UVREditorUniformScaleGizmoHandleGroup>( TEXT( "UniformScaleHandles" ), true );
+	UniformScaleGizmoHandleGroup->SetTranslucentGizmoMaterial( TranslucentGizmoMaterial );
+	UniformScaleGizmoHandleGroup->SetGizmoMaterial( GizmoMaterial );
+	UniformScaleGizmoHandleGroup->AttachParent = SceneComponent;
+	AllHandleGroups.Add( UniformScaleGizmoHandleGroup );
+
+	TranslationGizmoHandleGroup = CreateDefaultSubobject<UVREditorPivotTranslationGizmoHandleGroup>(TEXT("TranslationHandles"), true);
+	TranslationGizmoHandleGroup->SetTranslucentGizmoMaterial(TranslucentGizmoMaterial);
+	TranslationGizmoHandleGroup->SetGizmoMaterial(GizmoMaterial);
+	TranslationGizmoHandleGroup->AttachParent = SceneComponent;
+	AllHandleGroups.Add(TranslationGizmoHandleGroup);
+
+	ScaleGizmoHandleGroup = CreateDefaultSubobject<UVREditorPivotScaleGizmoHandleGroup>( TEXT( "ScaleHandles" ), true );
+	ScaleGizmoHandleGroup->SetTranslucentGizmoMaterial( TranslucentGizmoMaterial );
+	ScaleGizmoHandleGroup->SetGizmoMaterial( GizmoMaterial );
+	ScaleGizmoHandleGroup->AttachParent = SceneComponent;
+	AllHandleGroups.Add( ScaleGizmoHandleGroup );
+
+	PlaneTranslationGizmoHandleGroup = CreateDefaultSubobject<UVREditorPivotPlaneTranslationGizmoHandleGroup>( TEXT( "PlaneTranslationHandles" ), true );
+	PlaneTranslationGizmoHandleGroup->SetTranslucentGizmoMaterial( TranslucentGizmoMaterial );
+	PlaneTranslationGizmoHandleGroup->SetGizmoMaterial( GizmoMaterial );
+	PlaneTranslationGizmoHandleGroup->AttachParent = SceneComponent;
+	AllHandleGroups.Add( PlaneTranslationGizmoHandleGroup );
+
+	RotationGizmoHandleGroup = CreateDefaultSubobject<UVREditorPivotRotationGizmoHandleGroup>( TEXT( "RotationHandles" ), true );
+	RotationGizmoHandleGroup->SetTranslucentGizmoMaterial( TranslucentGizmoMaterial );
+	RotationGizmoHandleGroup->SetGizmoMaterial( GizmoMaterial );
+	RotationGizmoHandleGroup->AttachParent = SceneComponent;
+	AllHandleGroups.Add( RotationGizmoHandleGroup );
+
+	StretchGizmoHandleGroup = CreateDefaultSubobject<UVREditorStretchGizmoHandleGroup>( TEXT( "StretchHandles" ), true );
+	StretchGizmoHandleGroup->SetTranslucentGizmoMaterial( TranslucentGizmoMaterial );
+	StretchGizmoHandleGroup->SetGizmoMaterial( GizmoMaterial );
+	StretchGizmoHandleGroup->SetShowOnUniversalGizmo( false );
+	StretchGizmoHandleGroup->AttachParent = SceneComponent;
+	AllHandleGroups.Add( StretchGizmoHandleGroup );
+
+	// There may already be some objects selected as we switch into VR mode, so we'll pretend that just happened so
+	// that we can make sure all transitions complete properly
+	OnNewObjectsSelected();
+}
+
+void APivotTransformGizmo::UpdateGizmo( const EGizmoHandleTypes GizmoType, const ECoordSystem GizmoCoordinateSpace, const FTransform& LocalToWorld, const FBox& LocalBounds, const FVector ViewLocation, bool bAllHandlesVisible, 
+	UActorComponent* DraggingHandle, const TArray< UActorComponent* >& HoveringOverHandles, const float GizmoHoverScale, const float GizmoHoverAnimationDuration )
+{
+	const float WorldScaleFactor = GetWorld()->GetWorldSettings()->WorldToMeters / 100.0f;
+
+	// Position the gizmo at the location of the first selected actor
+	const bool bSweep = false;
+	this->SetActorTransform( LocalToWorld, bSweep );
+
+	// Increase scale with distance, to make gizmo handles easier to click on
+	// @todo vreditor: Should probably be a curve, not linear
+	// @todo vreditor: Should take FOV into account (especially in non-stereo/HMD mode)
+	const float WorldSpaceDistanceToToPivot = FMath::Sqrt( FVector::DistSquared( GetActorLocation(), ViewLocation ) );
+	const float GizmoScale( ( ( WorldSpaceDistanceToToPivot / WorldScaleFactor ) * VREd::PivotGizmoDistanceScaleFactor->GetFloat() ) * WorldScaleFactor );
+
+	// Update animation
+	float AnimationAlpha = GetAnimationAlpha();
+
+	// Update all the handles
+	bool bIsHoveringOrDraggingScaleGizmo = false;
+	for ( UVREditorGizmoHandleGroup* HandleGroup : AllHandleGroups )
+	{
+		if ( HandleGroup != nullptr && 
+		   ( DraggingHandle == nullptr || HandleGroup == StretchGizmoHandleGroup ) )
+		{
+			bool bIsHoveringOrDraggingThisHandleGroup = false;
+			HandleGroup->UpdateGizmoHandleGroup( LocalToWorld, LocalBounds, ViewLocation, bAllHandlesVisible, DraggingHandle,
+				HoveringOverHandles, AnimationAlpha, GizmoScale, GizmoHoverScale, GizmoHoverAnimationDuration, /* Out */ bIsHoveringOrDraggingThisHandleGroup );
+		}
+	}
+
+	UpdateHandleVisibility( GizmoType, GizmoCoordinateSpace, bAllHandlesVisible, DraggingHandle );
+}
+
+/************************************************************************/
+/* Translation                                                          */
+/************************************************************************/
+UVREditorPivotTranslationGizmoHandleGroup::UVREditorPivotTranslationGizmoHandleGroup() :
+	Super()
+{
+	UStaticMesh* TranslationHandleMesh = nullptr;
+	{
+		static ConstructorHelpers::FObjectFinder<UStaticMesh> ObjectFinder( TEXT( "/Engine/VREditor/TransformGizmo/TranslateHandleLong" ) );
+		TranslationHandleMesh = ObjectFinder.Object;
+		check( TranslationHandleMesh != nullptr );
+	}
+
+	CreateHandles( TranslationHandleMesh, FString( "PivotTranslationHandle" ) );
+}
+
+
+void UVREditorPivotTranslationGizmoHandleGroup::UpdateGizmoHandleGroup( const FTransform& LocalToWorld, const FBox& LocalBounds, const FVector ViewLocation, bool bAllHandlesVisible, class UActorComponent* DraggingHandle, 
+	const TArray< UActorComponent* >& HoveringOverHandles, float AnimationAlpha, float GizmoScale, const float GizmoHoverScale, const float GizmoHoverAnimationDuration, bool& bOutIsHoveringOrDraggingThisHandleGroup )
+{
+	// Call parent implementation (updates hover animation)
+	Super::UpdateGizmoHandleGroup( LocalToWorld, LocalBounds, ViewLocation, bAllHandlesVisible, DraggingHandle, HoveringOverHandles, AnimationAlpha,
+		GizmoScale, GizmoHoverScale, GizmoHoverAnimationDuration, bOutIsHoveringOrDraggingThisHandleGroup );
+	
+	const float MultipliedGizmoScale = GizmoScale * VREd::PivotGizmoTranslationScaleMultiply->GetFloat();
+	const float MultipliedGizmoHoverScale = GizmoHoverScale * VREd::PivotGizmoTranslationHoverScaleMultiply->GetFloat();
+
+	UpdateHandlesRelativeTransformOnAxis( FTransform( FVector( VREd::PivotGizmoTranslationPivotOffsetX->GetFloat(), 0, 0 ) ), AnimationAlpha, MultipliedGizmoScale, MultipliedGizmoHoverScale, ViewLocation, DraggingHandle, HoveringOverHandles );
+}
+
+ETransformGizmoInteractionType UVREditorPivotTranslationGizmoHandleGroup::GetInteractionType() const
+{
+	return ETransformGizmoInteractionType::Translate;
+}
+
+EGizmoHandleTypes UVREditorPivotTranslationGizmoHandleGroup::GetHandleType() const
+{
+	return EGizmoHandleTypes::Translate;
+}
+
+/************************************************************************/
+/* Scale	                                                            */
+/************************************************************************/
+UVREditorPivotScaleGizmoHandleGroup::UVREditorPivotScaleGizmoHandleGroup() :
+	Super()
+{
+	UStaticMesh* TranslationHandleMesh = nullptr;
+	{
+		static ConstructorHelpers::FObjectFinder<UStaticMesh> ObjectFinder( TEXT( "/Engine/VREditor/TransformGizmo/UniformScaleHandle" ) );
+		TranslationHandleMesh = ObjectFinder.Object;
+		check( TranslationHandleMesh != nullptr );
+	}
+
+	CreateHandles( TranslationHandleMesh, FString( "PivotScaleHandle" ) );	
+}
+
+void UVREditorPivotScaleGizmoHandleGroup::UpdateGizmoHandleGroup( const FTransform& LocalToWorld, const FBox& LocalBounds, const FVector ViewLocation, bool bAllHandlesVisible, class UActorComponent* DraggingHandle, 
+	const TArray< UActorComponent* >& HoveringOverHandles, float AnimationAlpha, float GizmoScale, const float GizmoHoverScale, const float GizmoHoverAnimationDuration, bool& bOutIsHoveringOrDraggingThisHandleGroup )
+{
+	// Call parent implementation (updates hover animation)
+	Super::UpdateGizmoHandleGroup( LocalToWorld, LocalBounds, ViewLocation, bAllHandlesVisible, DraggingHandle, HoveringOverHandles, AnimationAlpha, 
+		GizmoScale, GizmoHoverScale, GizmoHoverAnimationDuration, bOutIsHoveringOrDraggingThisHandleGroup );
+
+	UpdateHandlesRelativeTransformOnAxis( FTransform( FVector( VREd::PivotGizmoScalePivotOffsetX->GetFloat(), 0, 0 ) ), AnimationAlpha, GizmoScale, GizmoHoverScale, ViewLocation, DraggingHandle, HoveringOverHandles );
+}
+
+ETransformGizmoInteractionType UVREditorPivotScaleGizmoHandleGroup::GetInteractionType() const
+{
+	return ETransformGizmoInteractionType::Scale;
+}
+
+EGizmoHandleTypes UVREditorPivotScaleGizmoHandleGroup::GetHandleType() const
+{
+	return EGizmoHandleTypes::Scale;
+}
+
+bool UVREditorPivotScaleGizmoHandleGroup::SupportsWorldCoordinateSpace() const
+{
+	return false;
+}
+
+/************************************************************************/
+/* Plane Translation	                                                */
+/************************************************************************/
+UVREditorPivotPlaneTranslationGizmoHandleGroup::UVREditorPivotPlaneTranslationGizmoHandleGroup() :
+	Super()
+{
+	UStaticMesh* TranslationHandleMesh = nullptr;
+	{
+		static ConstructorHelpers::FObjectFinder<UStaticMesh> ObjectFinder( TEXT( "/Engine/VREditor/TransformGizmo/PlaneTranslationHandle" ) );
+		TranslationHandleMesh = ObjectFinder.Object;
+		check( TranslationHandleMesh != nullptr );
+	}
+
+	CreateHandles( TranslationHandleMesh, FString( "PlaneTranslationHandle" ) );
+}
+
+void UVREditorPivotPlaneTranslationGizmoHandleGroup::UpdateGizmoHandleGroup( const FTransform& LocalToWorld, const FBox& LocalBounds, const FVector ViewLocation, bool bAllHandlesVisible, class UActorComponent* DraggingHandle,
+	const TArray< UActorComponent* >& HoveringOverHandles, float AnimationAlpha, float GizmoScale, const float GizmoHoverScale, const float GizmoHoverAnimationDuration, bool& bOutIsHoveringOrDraggingThisHandleGroup )
+{
+	// Call parent implementation (updates hover animation)
+	Super::UpdateGizmoHandleGroup( LocalToWorld, LocalBounds, ViewLocation, bAllHandlesVisible, DraggingHandle, HoveringOverHandles, AnimationAlpha,
+		GizmoScale, GizmoHoverScale, GizmoHoverAnimationDuration, bOutIsHoveringOrDraggingThisHandleGroup );
+
+	UpdateHandlesRelativeTransformOnAxis( FTransform( FVector( 0, VREd::PivotGizmoPlaneTranslationPivotOffsetYZ->GetFloat(), VREd::PivotGizmoPlaneTranslationPivotOffsetYZ->GetFloat() ) ), 
+		AnimationAlpha, GizmoScale, GizmoHoverScale, ViewLocation, DraggingHandle, HoveringOverHandles );
+}
+
+ETransformGizmoInteractionType UVREditorPivotPlaneTranslationGizmoHandleGroup::GetInteractionType() const
+{
+	return ETransformGizmoInteractionType::TranslateOnPlane;
+}
+
+EGizmoHandleTypes UVREditorPivotPlaneTranslationGizmoHandleGroup::GetHandleType() const
+{
+	return EGizmoHandleTypes::Translate;
+}
+
+/************************************************************************/
+/* Rotation																*/
+/************************************************************************/
+UVREditorPivotRotationGizmoHandleGroup::UVREditorPivotRotationGizmoHandleGroup() :
+	Super()
+{
+	UStaticMesh* TranslationHandleMesh = nullptr;
+	{
+		static ConstructorHelpers::FObjectFinder<UStaticMesh> ObjectFinder(TEXT("/Engine/VREditor/TransformGizmo/RotationHandleFull"));
+		TranslationHandleMesh = ObjectFinder.Object;
+		check(TranslationHandleMesh != nullptr);
+	}
+
+	CreateHandles(TranslationHandleMesh, FString("RotationHandle"));
+}
+
+void UVREditorPivotRotationGizmoHandleGroup::UpdateGizmoHandleGroup( const FTransform& LocalToWorld, const FBox& LocalBounds, const FVector ViewLocation, bool bAllHandlesVisible, class UActorComponent* DraggingHandle,
+	const TArray< UActorComponent* >& HoveringOverHandles, float AnimationAlpha, float GizmoScale, const float GizmoHoverScale, const float GizmoHoverAnimationDuration, bool& bOutIsHoveringOrDraggingThisHandleGroup )
+{
+	// Call parent implementation (updates hover animation)
+	Super::UpdateGizmoHandleGroup(LocalToWorld, LocalBounds, ViewLocation, bAllHandlesVisible, DraggingHandle, HoveringOverHandles, AnimationAlpha,
+		GizmoScale, GizmoHoverScale, GizmoHoverAnimationDuration, bOutIsHoveringOrDraggingThisHandleGroup );
+	
+	UpdateHandlesRelativeTransformOnAxis( FTransform( FVector( 0, 0, 0 ) ), AnimationAlpha, GizmoScale, GizmoScale, ViewLocation, DraggingHandle, HoveringOverHandles );
+}
+
+ETransformGizmoInteractionType UVREditorPivotRotationGizmoHandleGroup::GetInteractionType() const
+{
+	return ETransformGizmoInteractionType::RotateOnAngle;
+}
+
+EGizmoHandleTypes UVREditorPivotRotationGizmoHandleGroup::GetHandleType() const
+{
+	return EGizmoHandleTypes::Rotate;
+}
\ No newline at end of file
diff --git a/Engine/Source/Editor/VREditor/Gizmo/VREditorPivotTransformGizmo.h b/Engine/Source/Editor/VREditor/Gizmo/VREditorPivotTransformGizmo.h
new file mode 100644
--- /dev/null
+++ b/Engine/Source/Editor/VREditor/Gizmo/VREditorPivotTransformGizmo.h
@@ -0,0 +1,154 @@
+// Copyright 1998-2015 Epic Games, Inc. All Rights Reserved.
+
+#pragma once
+
+#include "VREditorBaseTransformGizmo.h"
+#include "VREditorGizmoHandle.h"
+#include "VREditorPivotTransformGizmo.generated.h"
+
+/**
+ * A transform gizmo on the pivot that allows you to interact with selected objects by moving, scaling and rotating.
+ */
+UCLASS()
+class APivotTransformGizmo : public ABaseTransformGizmo
+{
+	GENERATED_BODY()
+
+public:
+
+	/** Default constructor that sets up CDO properties */
+	APivotTransformGizmo();
+
+	/** Called by the world interaction system after we've been spawned into the world, to allow
+	    us to create components and set everything up nicely for the selected objects that we'll be
+		used to manipulate */
+	virtual void UpdateGizmo( const EGizmoHandleTypes GizmoType, const ECoordSystem GizmoCoordinateSpace, const FTransform& LocalToWorld, const FBox& LocalBounds, const FVector ViewLocation, 
+		bool bAllHandlesVisible, class UActorComponent* DraggingHandle, const TArray< UActorComponent* >& HoveringOverHandles, const float GizmoHoverScale, const float GizmoHoverAnimationDuration ) override;
+
+private:
+	/** Uniform scale handle group component */
+	UPROPERTY()
+	class UVREditorUniformScaleGizmoHandleGroup* UniformScaleGizmoHandleGroup;
+
+	/** Translation handle group component */
+	UPROPERTY()
+	class UVREditorPivotTranslationGizmoHandleGroup* TranslationGizmoHandleGroup;
+	
+	/** Scale handle group component */
+	UPROPERTY()
+	class UVREditorPivotScaleGizmoHandleGroup* ScaleGizmoHandleGroup;
+	
+	/** Plane translation handle group component */
+	UPROPERTY()
+	class UVREditorPivotPlaneTranslationGizmoHandleGroup* PlaneTranslationGizmoHandleGroup;
+
+	/** Rotation handle group component */
+	UPROPERTY()
+	class UVREditorPivotRotationGizmoHandleGroup* RotationGizmoHandleGroup;
+
+	/** Stretch handle group component */
+	UPROPERTY()
+	class UVREditorStretchGizmoHandleGroup* StretchGizmoHandleGroup;
+};
+
+/**
+ * Axis Gizmo handle for translating
+ */
+UCLASS()
+class UVREditorPivotTranslationGizmoHandleGroup : public UVREditorAxisGizmoHandleGroup
+{
+	GENERATED_BODY()
+
+public:
+
+	/** Default constructor that sets up CDO properties */
+	UVREditorPivotTranslationGizmoHandleGroup();
+
+	/** Updates the gizmo handles */
+	virtual void UpdateGizmoHandleGroup( const FTransform& LocalToWorld, const FBox& LocalBounds, const FVector ViewLocation, bool bAllHandlesVisible, class UActorComponent* DraggingHandle, const TArray< UActorComponent* >& HoveringOverHandles, 
+		float AnimationAlpha, float GizmoScale, const float GizmoHoverScale, const float GizmoHoverAnimationDuration, bool& bOutIsHoveringOrDraggingThisHandleGroup ) override;
+
+	/** Gets the InteractionType for this Gizmo handle */
+	virtual ETransformGizmoInteractionType GetInteractionType() const override;
+
+	/** Gets the GizmoType for this Gizmo handle */
+	virtual EGizmoHandleTypes GetHandleType() const override;
+};
+
+/**
+ * Axis Gizmo handle for scaling
+ */
+UCLASS()
+class UVREditorPivotScaleGizmoHandleGroup : public UVREditorAxisGizmoHandleGroup
+{
+	GENERATED_BODY()
+
+public:
+
+	/** Default constructor that sets up CDO properties */
+	UVREditorPivotScaleGizmoHandleGroup();
+
+
+	/** Updates the gizmo handles */
+	virtual void UpdateGizmoHandleGroup( const FTransform& LocalToWorld, const FBox& LocalBounds,  const FVector ViewLocation, bool bAllHandlesVisible, class UActorComponent* DraggingHandle, const TArray< UActorComponent* >& HoveringOverHandles, 
+		float AnimationAlpha, float GizmoScale, const float GizmoHoverScale, const float GizmoHoverAnimationDuration, bool& bOutIsHoveringOrDraggingThisHandleGroup ) override;
+
+	/** Gets the InteractionType for this Gizmo handle */
+	virtual ETransformGizmoInteractionType GetInteractionType() const override;
+
+	/** Gets the GizmoType for this Gizmo handle */
+	virtual EGizmoHandleTypes GetHandleType() const override;
+
+	/** Returns true if this type of handle is allowed with world space gizmos */
+	virtual bool SupportsWorldCoordinateSpace() const override;
+};
+
+/**
+ * Axis Gizmo handle for plane translation
+ */
+UCLASS()
+class UVREditorPivotPlaneTranslationGizmoHandleGroup : public UVREditorAxisGizmoHandleGroup
+{
+	GENERATED_BODY()
+
+public:
+
+	/** Default constructor that sets up CDO properties */
+	UVREditorPivotPlaneTranslationGizmoHandleGroup();
+
+
+	/** Updates the gizmo handles */
+	virtual void UpdateGizmoHandleGroup( const FTransform& LocalToWorld, const FBox& LocalBounds, const FVector ViewLocation, bool bAllHandlesVisible, class UActorComponent* DraggingHandle, const TArray< UActorComponent* >& HoveringOverHandles, 
+		float AnimationAlpha, float GizmoScale, const float GizmoHoverScale, const float GizmoHoverAnimationDuration, bool& bOutIsHoveringOrDraggingThisHandleGroup ) override;
+
+	/** Gets the InteractionType for this Gizmo handle */
+	virtual ETransformGizmoInteractionType GetInteractionType() const override;
+
+	/** Gets the GizmoType for this Gizmo handle */
+	virtual EGizmoHandleTypes GetHandleType() const override;
+};
+
+
+/**
+ * Axis Gizmo handle for rotation
+ */
+UCLASS()
+class UVREditorPivotRotationGizmoHandleGroup : public UVREditorAxisGizmoHandleGroup
+{
+	GENERATED_BODY()
+
+public:
+
+	/** Default constructor that sets up CDO properties */
+	UVREditorPivotRotationGizmoHandleGroup();
+
+	/** Updates the gizmo handles */
+	virtual void UpdateGizmoHandleGroup(const FTransform& LocalToWorld, const FBox& LocalBounds, const FVector ViewLocation, bool bAllHandlesVisible, class UActorComponent* DraggingHandle, const TArray< UActorComponent* >& HoveringOverHandles, 
+		float AnimationAlpha, float GizmoScale, const float GizmoHoverScale, const float GizmoHoverAnimationDuration, bool& bOutIsHoveringOrDraggingThisHandleGroup ) override;
+
+	/** Gets the InteractionType for this Gizmo handle */
+	virtual ETransformGizmoInteractionType GetInteractionType() const override;
+
+	/** Gets the GizmoType for this Gizmo handle */
+	virtual EGizmoHandleTypes GetHandleType() const override;
+};
\ No newline at end of file
diff --git a/Engine/Source/Editor/VREditor/Gizmo/VREditorPlaneTranslationGizmoHandle.h b/Engine/Source/Editor/VREditor/Gizmo/VREditorPlaneTranslationGizmoHandle.h
--- a/Engine/Source/Editor/VREditor/Gizmo/VREditorPlaneTranslationGizmoHandle.h
+++ b/Engine/Source/Editor/VREditor/Gizmo/VREditorPlaneTranslationGizmoHandle.h
@@ -19,7 +19,8 @@ public:
 	UVREditorPlaneTranslationGizmoHandleGroup();
 	
 	/** Updates the gizmo handles */
-	virtual void UpdateGizmoHandleGroup( const FTransform& LocalToWorld, const FBox& LocalBounds, const FVector ViewLocation, bool bAllHandlesVisible, class UActorComponent* DraggingHandle, const TArray< UActorComponent* >& HoveringOverHandles, float AnimationAlpha, float GizmoScale, const float GizmoHoverScale, const float GizmoHoverAnimationDuration, bool& bOutIsHoveringOrDraggingThisHandleGroup ) override;
+	void UpdateGizmoHandleGroup(const FTransform& LocalToWorld, const FBox& LocalBounds, const FVector ViewLocation, bool bAllHandlesVisible, class UActorComponent* DraggingHandle,
+		const TArray< UActorComponent* >& HoveringOverHandles, float AnimationAlpha, float GizmoScale, const float GizmoHoverScale, const float GizmoHoverAnimationDuration, bool& bOutIsHoveringOrDraggingThisHandleGroup ) override;
 
 	/** Gets the InteractionType for this Gizmo handle */
 	virtual ETransformGizmoInteractionType GetInteractionType() const override;
diff --git a/Engine/Source/Editor/VREditor/Gizmo/VREditorRotationGizmoHandle.h b/Engine/Source/Editor/VREditor/Gizmo/VREditorRotationGizmoHandle.h
--- a/Engine/Source/Editor/VREditor/Gizmo/VREditorRotationGizmoHandle.h
+++ b/Engine/Source/Editor/VREditor/Gizmo/VREditorRotationGizmoHandle.h
@@ -17,10 +17,11 @@ public:
 
 	/** Default constructor that sets up CDO properties */
 	UVREditorRotationGizmoHandleGroup();
-	
-	/** Updates the gizmo handles */
-	virtual void UpdateGizmoHandleGroup( const FTransform& LocalToWorld, const FBox& LocalBounds, const FVector ViewLocation, bool bAllHandlesVisible, class UActorComponent* DraggingHandle, const TArray< UActorComponent* >& HoveringOverHandles, float AnimationAlpha, float GizmoScale, const float GizmoHoverScale, const float GizmoHoverAnimationDuration, bool& bOutIsHoveringOrDraggingThisHandleGroup ) override;
-	
+
+	/** Updates the Gizmo handles */
+	virtual void UpdateGizmoHandleGroup( const FTransform& LocalToWorld, const FBox& LocalBounds, const FVector ViewLocation, bool bAllHandlesVisible, class UActorComponent* DraggingHandle, const TArray< UActorComponent* >& HoveringOverHandles, 
+		float AnimationAlpha, float GizmoScale, const float GizmoHoverScale, const float GizmoHoverAnimationDuration, bool& bOutIsHoveringOrDraggingThisHandleGroup ) override;
+
 	/** Gets the InteractionType for this Gizmo handle */
 	virtual ETransformGizmoInteractionType GetInteractionType() const override;
 
diff --git a/Engine/Source/Editor/VREditor/Gizmo/VREditorStretchGizmoHandle.h b/Engine/Source/Editor/VREditor/Gizmo/VREditorStretchGizmoHandle.h
--- a/Engine/Source/Editor/VREditor/Gizmo/VREditorStretchGizmoHandle.h
+++ b/Engine/Source/Editor/VREditor/Gizmo/VREditorStretchGizmoHandle.h
@@ -18,8 +18,9 @@ public:
 	/** Default constructor that sets up CDO properties */
 	UVREditorStretchGizmoHandleGroup();
 	
-	/** Updates the gizmo handles */
-	virtual void UpdateGizmoHandleGroup( const FTransform& LocalToWorld, const FBox& LocalBounds, const FVector ViewLocation, bool bAllHandlesVisible, class UActorComponent* DraggingHandle, const TArray< UActorComponent* >& HoveringOverHandles, float AnimationAlpha, float GizmoScale, const float GizmoHoverScale, const float GizmoHoverAnimationDuration, bool& bOutIsHoveringOrDraggingThisHandleGroup ) override;
+	/** Updates the Gizmo handles */
+	virtual void UpdateGizmoHandleGroup( const FTransform& LocalToWorld, const FBox& LocalBounds, const FVector ViewLocation, bool bAllHandlesVisible, class UActorComponent* DraggingHandle, const TArray< UActorComponent* >& HoveringOverHandles, 
+		float AnimationAlpha, float GizmoScale, const float GizmoHoverScale, const float GizmoHoverAnimationDuration, bool& bOutIsHoveringOrDraggingThisHandleGroup ) override;
 
 	/** Gets the InteractionType for this Gizmo handle */
 	virtual ETransformGizmoInteractionType GetInteractionType() const override;
diff --git a/Engine/Source/Editor/VREditor/Gizmo/VREditorTranslationGizmoHandle.h b/Engine/Source/Editor/VREditor/Gizmo/VREditorTranslationGizmoHandle.h
--- a/Engine/Source/Editor/VREditor/Gizmo/VREditorTranslationGizmoHandle.h
+++ b/Engine/Source/Editor/VREditor/Gizmo/VREditorTranslationGizmoHandle.h
@@ -18,12 +18,13 @@ public:
 	/** Default constructor that sets up CDO properties */
 	UVREditorTranslationGizmoHandleGroup();
 	
-	/** Updates the gizmo handles */
-	virtual void UpdateGizmoHandleGroup( const FTransform& LocalToWorld, const FBox& LocalBounds, const FVector ViewLocation, bool bAllHandlesVisible, class UActorComponent* DraggingHandle, const TArray< UActorComponent* >& HoveringOverHandles, float AnimationAlpha, float GizmoScale, const float GizmoHoverScale, const float GizmoHoverAnimationDuration, bool& bOutIsHoveringOrDraggingThisHandleGroup ) override;
-	
+	/** Updates the Gizmo handles */
+	virtual void UpdateGizmoHandleGroup( const FTransform& LocalToWorld, const FBox& LocalBounds, const FVector ViewLocation, bool bAllHandlesVisible, class UActorComponent* DraggingHandle,
+		const TArray< UActorComponent* >& HoveringOverHandles, float AnimationAlpha, float GizmoScale, const float GizmoHoverScale, const float GizmoHoverAnimationDuration, bool& bOutIsHoveringOrDraggingThisHandleGroup ) override;
+
 	/** Gets the InteractionType for this Gizmo handle */
 	virtual ETransformGizmoInteractionType GetInteractionType() const override;
 
 	/** Gets the GizmoType for this Gizmo handle */
 	virtual EGizmoHandleTypes GetHandleType() const override;
-};
+};
\ No newline at end of file
diff --git a/Engine/Source/Editor/VREditor/Gizmo/VREditorUniformScaleGizmoHandle.cpp b/Engine/Source/Editor/VREditor/Gizmo/VREditorUniformScaleGizmoHandle.cpp
--- a/Engine/Source/Editor/VREditor/Gizmo/VREditorUniformScaleGizmoHandle.cpp
+++ b/Engine/Source/Editor/VREditor/Gizmo/VREditorUniformScaleGizmoHandle.cpp
@@ -7,7 +7,8 @@
 #include "UnitConversion.h"
 
 UVREditorUniformScaleGizmoHandleGroup::UVREditorUniformScaleGizmoHandleGroup()
-	: Super()
+	: Super(),
+	bUsePivotAsLocation( true )
 {
 	// Setup uniform scaling
 	UStaticMesh* UniformScaleMesh = nullptr;
@@ -17,42 +18,29 @@ UVREditorUniformScaleGizmoHandleGroup::UVREditorUniformScaleGizmoHandleGroup()
 		check( UniformScaleMesh != nullptr );
 	}
 
-	UStaticMeshComponent* UniformScaleHandle = CreateDefaultSubobject<UStaticMeshComponent>( TEXT( "UniformScaleHandle" ) );
+	UStaticMeshComponent* UniformScaleHandle = CreateMeshHandle( UniformScaleMesh, FString( "UniformScaleHandle" ) );
 	check( UniformScaleHandle != nullptr );
 
-	const bool bAllowGizmoLighting = false;	// @todo vreditor: Not sure if we want this for gizmos or not yet.  Needs feedback.  Also they're translucent right now.
-
-	UniformScaleHandle->SetStaticMesh( UniformScaleMesh );
-	UniformScaleHandle->SetMobility( EComponentMobility::Movable );
-	UniformScaleHandle->AttachTo( this );
-
-	UniformScaleHandle->SetCollisionEnabled( ECollisionEnabled::QueryOnly );
-	UniformScaleHandle->SetCollisionResponseToAllChannels( ECR_Ignore );
-	UniformScaleHandle->SetCollisionResponseToChannel( ECC_EditorGizmo, ECollisionResponse::ECR_Block );
-	UniformScaleHandle->SetCollisionObjectType( ECC_EditorGizmo );
-
-	UniformScaleHandle->bGenerateOverlapEvents = false;
-	UniformScaleHandle->SetCanEverAffectNavigation( false );
-	UniformScaleHandle->bCastDynamicShadow = bAllowGizmoLighting;
-	UniformScaleHandle->bCastStaticShadow = false;
-	UniformScaleHandle->bAffectDistanceFieldLighting = bAllowGizmoLighting;
-	UniformScaleHandle->bAffectDynamicIndirectLighting = bAllowGizmoLighting;
-
 	FVREditorGizmoHandle& NewHandle = *new( Handles ) FVREditorGizmoHandle();
 	NewHandle.HandleMesh = UniformScaleHandle;
 }
 
-void UVREditorUniformScaleGizmoHandleGroup::UpdateGizmoHandleGroup( const FTransform& LocalToWorld, const FBox& LocalBounds, const FVector ViewLocation, bool bAllHandlesVisible, class UActorComponent* DraggingHandle, const TArray< UActorComponent* >& HoveringOverHandles, float AnimationAlpha, float GizmoScale, const float GizmoHoverScale, const float GizmoHoverAnimationDuration, bool& bOutIsHoveringOrDraggingThisHandleGroup )
+void UVREditorUniformScaleGizmoHandleGroup::UpdateGizmoHandleGroup( const FTransform& LocalToWorld, const FBox& LocalBounds, const FVector ViewLocation, bool bAllHandlesVisible, class UActorComponent* DraggingHandle, const TArray< UActorComponent* >& HoveringOverHandles, 
+	float AnimationAlpha, float GizmoScale, const float GizmoHoverScale, const float GizmoHoverAnimationDuration, bool& bOutIsHoveringOrDraggingThisHandleGroup )
 {
 	// Call parent implementation (updates hover animation)
-	Super::UpdateGizmoHandleGroup( LocalToWorld, LocalBounds, ViewLocation, bAllHandlesVisible, DraggingHandle, HoveringOverHandles, AnimationAlpha, GizmoScale, GizmoHoverScale, GizmoHoverAnimationDuration, bOutIsHoveringOrDraggingThisHandleGroup );
+	Super::UpdateGizmoHandleGroup( LocalToWorld, LocalBounds, ViewLocation, bAllHandlesVisible, DraggingHandle, HoveringOverHandles, 
+		AnimationAlpha, GizmoScale, GizmoHoverScale, GizmoHoverAnimationDuration, bOutIsHoveringOrDraggingThisHandleGroup );
 
 	FVREditorGizmoHandle& Handle = Handles[ 0 ];
 	UStaticMeshComponent* UniformScaleHandle = Handle.HandleMesh;
 	if (UniformScaleHandle != nullptr)	// Can be null if no handle for this specific placement
 	{
-		FVector HandleRelativeLocation = LocalBounds.GetCenter();
-		UniformScaleHandle->SetRelativeLocation( HandleRelativeLocation );
+		FVector HandleRelativeLocation;
+		if ( !bUsePivotAsLocation )
+		{
+			UniformScaleHandle->SetRelativeLocation( HandleRelativeLocation );
+		}
 
 		float GizmoHandleScale = GizmoScale;
 
@@ -87,14 +75,9 @@ void UVREditorUniformScaleGizmoHandleGroup::UpdateGizmoHandleGroup( const FTrans
 					{
 						HandleColor = Mode->GetColor( FVREditorMode::EColors::DraggingGizmoColor );
 					}
-					else
+					else if (HoveringOverHandles.Contains( UniformScaleHandle ))
 					{
-						HandleColor = VREd::GizmoColor::WhiteGizmoColor;
-
-						if (HoveringOverHandles.Contains( UniformScaleHandle ))
-						{
-							HandleColor = FLinearColor::LerpUsingHSV( HandleColor, Mode->GetColor( FVREditorMode::EColors::HoverGizmoColor ), Handle.HoverAlpha );
-						}
+						HandleColor = FLinearColor::LerpUsingHSV( HandleColor, Mode->GetColor( FVREditorMode::EColors::HoverGizmoColor ), Handle.HoverAlpha );
 					}
 
 					MID0->SetVectorParameterValue( "Color", HandleColor );
@@ -113,4 +96,9 @@ ETransformGizmoInteractionType UVREditorUniformScaleGizmoHandleGroup::GetInterac
 EGizmoHandleTypes UVREditorUniformScaleGizmoHandleGroup::GetHandleType() const
 {
 	return EGizmoHandleTypes::Scale;
-}
\ No newline at end of file
+}
+
+void UVREditorUniformScaleGizmoHandleGroup::SetUsePivotPointAsLocation( bool bUsePivotAsLocation )
+{
+	this->bUsePivotAsLocation = bUsePivotAsLocation;
+}
diff --git a/Engine/Source/Editor/VREditor/Gizmo/VREditorUniformScaleGizmoHandle.h b/Engine/Source/Editor/VREditor/Gizmo/VREditorUniformScaleGizmoHandle.h
--- a/Engine/Source/Editor/VREditor/Gizmo/VREditorUniformScaleGizmoHandle.h
+++ b/Engine/Source/Editor/VREditor/Gizmo/VREditorUniformScaleGizmoHandle.h
@@ -18,12 +18,21 @@ public:
 	/** Default constructor that sets up CDO properties */
 	UVREditorUniformScaleGizmoHandleGroup();
 	
-	/** Updates the gizmo handles */
-	virtual void UpdateGizmoHandleGroup( const FTransform& LocalToWorld, const FBox& LocalBounds, const FVector ViewLocation, bool bAllHandlesVisible, class UActorComponent* DraggingHandle, const TArray< UActorComponent* >& HoveringOverHandles, float AnimationAlpha, float GizmoScale, const float GizmoHoverScale, const float GizmoHoverAnimationDuration, bool& bOutIsHoveringOrDraggingThisHandleGroup ) override;
+	/** Updates the Gizmo handles */
+	virtual void UpdateGizmoHandleGroup( const FTransform& LocalToWorld, const FBox& LocalBounds, const FVector ViewLocation, bool bAllHandlesVisible, class UActorComponent* DraggingHandle,
+		const TArray< UActorComponent* >& HoveringOverHandles, float AnimationAlpha, float GizmoScale, const float GizmoHoverScale, const float GizmoHoverAnimationDuration, bool& bOutIsHoveringOrDraggingThisHandleGroup ) override;
 
 	/** Gets the InteractionType for this Gizmo handle */
 	virtual ETransformGizmoInteractionType GetInteractionType() const override;
 
 	/** Gets the GizmoType for this Gizmo handle */
 	virtual EGizmoHandleTypes GetHandleType() const override;
+
+	/** Sets if the pivot point is used as location for the handle */
+	void SetUsePivotPointAsLocation( bool bUsePivotAsLocation );
+
+private:
+
+	/** If the pivot point is used for the uniform scaling handle */
+	bool bUsePivotAsLocation;
 };
\ No newline at end of file
diff --git a/Engine/Source/Editor/VREditor/VREditorMode.cpp b/Engine/Source/Editor/VREditor/VREditorMode.cpp
--- a/Engine/Source/Editor/VREditor/VREditorMode.cpp
+++ b/Engine/Source/Editor/VREditor/VREditorMode.cpp
@@ -119,7 +119,8 @@ FVREditorMode::FVREditorMode()
 	  UISystem(),
 	  WorldInteraction(),
 	  CurrentGizmoType( EGizmoHandleTypes::All ),
-	  bFirstTick( true )
+	  bFirstTick( true ),
+	  bWasInWorldSpaceBeforeScaleMode( false )
 {
 	FEditorDelegates::MapChange.AddRaw( this, &FVREditorMode::OnMapChange );
 	FEditorDelegates::BeginPIE.AddRaw( this, &FVREditorMode::OnBeginPIE );
@@ -1352,20 +1353,20 @@ void FVREditorMode::Tick( FEditorViewportClient* ViewportClient, float DeltaTime
 	{
 		const FTransform RoomToWorld = GetRoomTransform();
 		const FTransform WorldToRoom = RoomToWorld.Inverse();
-		const FTransform ViewportToWorld = FTransform(SavedEditorState.ViewRotation, SavedEditorState.ViewLocation);
+		const FTransform ViewportToWorld = FTransform( SavedEditorState.ViewRotation, SavedEditorState.ViewLocation );
 		const FTransform ViewportToRoom = ViewportToWorld * WorldToRoom;
 
 		FTransform ViewportToRoomYaw = ViewportToRoom;
-		ViewportToRoomYaw.SetRotation(FQuat(FRotator(0.0f, ViewportToRoomYaw.GetRotation().Rotator().Yaw, 0.0f)));
+		ViewportToRoomYaw.SetRotation( FQuat( FRotator( 0.0f, ViewportToRoomYaw.GetRotation().Rotator().Yaw, 0.0f ) ) );
 
 		FTransform HeadToRoomYaw = GetRoomSpaceHeadTransform();
-		HeadToRoomYaw.SetRotation(FQuat(FRotator(0.0f, HeadToRoomYaw.GetRotation().Rotator().Yaw, 0.0f)));
+		HeadToRoomYaw.SetRotation( FQuat( FRotator( 0.0f, HeadToRoomYaw.GetRotation().Rotator().Yaw, 0.0f ) ) );
 
 		FTransform RoomToWorldYaw = RoomToWorld;
-		RoomToWorldYaw.SetRotation(FQuat(FRotator(0.0f, RoomToWorldYaw.GetRotation().Rotator().Yaw, 0.0f)));
+		RoomToWorldYaw.SetRotation( FQuat( FRotator( 0.0f, RoomToWorldYaw.GetRotation().Rotator().Yaw, 0.0f ) ) );
 
 		FTransform ResultToWorld = ( HeadToRoomYaw.Inverse() * ViewportToRoomYaw ) * RoomToWorldYaw;
-		SetRoomTransform(ResultToWorld);
+		SetRoomTransform( ResultToWorld );
 	}
 
 
@@ -2678,78 +2679,81 @@ void FVREditorMode::ApplyButtonPressColors( FVRAction VRAction, EInputEvent Even
 {
 	SpawnAvatarMeshActor();
 	
-	FVirtualHand Hand = GetVirtualHand( VRAction.HandIndex );
+	FVirtualHand& Hand = GetVirtualHand( VRAction.HandIndex );
 
 	const float PressStrength = 10.0f;
 
-	//Trigger
+	// Trigger
 	if ( VRAction.ActionType == EVRActionType::SelectAndMove || VRAction.ActionType == EVRActionType::SelectAndMove_LightlyPressed )
 	{
 		static FName StaticTriggerParameter( "B1" );
-
-		if( Event == IE_Pressed)
-		{
-			Hand.HandMeshMID->SetScalarParameterValue( StaticTriggerParameter, PressStrength );
-		}
-		else if ( Event == IE_Released )
-		{
-			Hand.HandMeshMID->SetScalarParameterValue( StaticTriggerParameter, 0.0f );
-		}
+		SetMotionControllerButtonPressedVisuals( Hand, Event, StaticTriggerParameter, PressStrength );
 	}
 
-	//Shoulder button
+	// Shoulder button
 	if ( VRAction.ActionType == EVRActionType::WorldMovement )
 	{
 		static FName StaticShoulderParameter( "B2" );
-
-		if (Event == IE_Pressed)
-		{
-			Hand.HandMeshMID->SetScalarParameterValue( StaticShoulderParameter, PressStrength );
-		}
-		else if( Event == IE_Released )
-		{
-			Hand.HandMeshMID->SetScalarParameterValue( StaticShoulderParameter, 0.0f );
-		}
+		SetMotionControllerButtonPressedVisuals( Hand, Event, StaticShoulderParameter, PressStrength );
 	}
 
-	//Trackpad
+	// Trackpad
 	if( VRAction.ActionType == EVRActionType::ConfirmRadialSelection )
 	{
 		static FName StaticTrackpadParameter( "B3" );
-
-		if( Event == IE_Pressed )
-		{
-			Hand.HandMeshMID->SetScalarParameterValue( StaticTrackpadParameter, PressStrength );
-		}
-		else if( Event == IE_Released )
-		{
-			Hand.HandMeshMID->SetScalarParameterValue( StaticTrackpadParameter, 0.0f );
-		}
+		SetMotionControllerButtonPressedVisuals( Hand, Event, StaticTrackpadParameter, PressStrength );
 	}
 
+	// Modifier
 	if( VRAction.ActionType == EVRActionType::Modifier )
 	{
 		static FName StaticModifierParameter( "B4" );
+		SetMotionControllerButtonPressedVisuals( Hand, Event, StaticModifierParameter, PressStrength );
+	}
+}
 
-		if( Event == IE_Pressed )
-		{
-			Hand.HandMeshMID->SetScalarParameterValue( StaticModifierParameter, PressStrength );
-		}
-		else if( Event == IE_Released )
-		{
-			Hand.HandMeshMID->SetScalarParameterValue( StaticModifierParameter, 0.0f );
-		}
+void FVREditorMode::SetMotionControllerButtonPressedVisuals( FVirtualHand& Hand, EInputEvent Event, const FName& ParameterName, const float PressStrength )
+{
+	if ( Event == IE_Pressed )
+	{
+		Hand.HandMeshMID->SetScalarParameterValue( ParameterName, PressStrength);
+	}
+	else if ( Event == IE_Released )
+	{
+		Hand.HandMeshMID->SetScalarParameterValue( ParameterName, 0.0f);
 	}
 }
 
 void FVREditorMode::CycleTransformGizmoHandleType()
 {
 	EGizmoHandleTypes NewGizmoType = (EGizmoHandleTypes)( (uint8)CurrentGizmoType + 1 );
+	
 	if( NewGizmoType > EGizmoHandleTypes::Scale )
 	{
 		NewGizmoType = EGizmoHandleTypes::All;
 	}
 
+	// Set coordinate system to local if the next gizmo will be for non-uniform scaling 
+	if ( NewGizmoType == EGizmoHandleTypes::Scale )
+	{
+		const ECoordSystem CurrentCoordSystem = GetTransformGizmoCoordinateSpace();
+		if ( CurrentCoordSystem == COORD_World )
+		{
+			GLevelEditorModeTools().SetCoordSystem( COORD_Local );
+			// Remember if coordinate system was in world space before scaling
+			bWasInWorldSpaceBeforeScaleMode = true;
+		}
+		else if ( CurrentCoordSystem == COORD_Local )
+		{
+			bWasInWorldSpaceBeforeScaleMode = false;
+		}
+	} 
+	else if ( CurrentGizmoType == EGizmoHandleTypes::Scale && bWasInWorldSpaceBeforeScaleMode )
+	{
+		// Set the coordinate system to world space if the coordinate system was world before scaling
+		SetTransformGizmoCoordinateSpace( COORD_World );
+	}
+	
 	CurrentGizmoType = NewGizmoType;
 }
 
@@ -2759,17 +2763,22 @@ EGizmoHandleTypes FVREditorMode::GetCurrentGizmoType() const
 	return CurrentGizmoType;
 }
 
-
 void FVREditorMode::CycleTransformGizmoCoordinateSpace()
 {
-	const ECoordSystem CurrentCoordSystem = GLevelEditorModeTools().GetCoordSystem();
-	GLevelEditorModeTools().SetCoordSystem( CurrentCoordSystem == COORD_World ? COORD_Local : COORD_World );
+	const bool bGetRawValue = true;
+	const ECoordSystem CurrentCoordSystem = GLevelEditorModeTools().GetCoordSystem( bGetRawValue );
+	SetTransformGizmoCoordinateSpace( CurrentCoordSystem == COORD_World ? COORD_Local : COORD_World );
 }
 
+void FVREditorMode::SetTransformGizmoCoordinateSpace( const ECoordSystem NewCoordSystem )
+{
+	GLevelEditorModeTools().SetCoordSystem( NewCoordSystem );
+}
 
 ECoordSystem FVREditorMode::GetTransformGizmoCoordinateSpace() const
 {
-	const ECoordSystem CurrentCoordSystem = GLevelEditorModeTools().GetCoordSystem();
+	const bool bGetRawValue = true;
+	const ECoordSystem CurrentCoordSystem = GLevelEditorModeTools().GetCoordSystem( bGetRawValue );
 	return CurrentCoordSystem;
 }
 
diff --git a/Engine/Source/Editor/VREditor/VREditorMode.h b/Engine/Source/Editor/VREditor/VREditorMode.h
--- a/Engine/Source/Editor/VREditor/VREditorMode.h
+++ b/Engine/Source/Editor/VREditor/VREditorMode.h
@@ -275,9 +275,12 @@ public:
 	/** Gets the current Gizmo handle type */
 	EGizmoHandleTypes GetCurrentGizmoType() const;
 
-	/** Switch which trasform gizmo coordinate space we're using. */
+	/** Switch which transform gizmo coordinate space we're using. */
 	void CycleTransformGizmoCoordinateSpace();
 
+	/** Sets which transform gizmo coordinate space is used */
+	void SetTransformGizmoCoordinateSpace( const ECoordSystem NewCoordSystem );
+
 	/** Returns which transform gizmo coordinate space we're using, world or local */
 	ECoordSystem GetTransformGizmoCoordinateSpace() const;
 		
@@ -347,6 +350,8 @@ private:
 	/** Changes the color of the buttons on the handmesh */
 	void ApplyButtonPressColors( FVRAction VRAction, EInputEvent Event );
 
+	/** Set the visuals for a button on the motion controller */
+	void SetMotionControllerButtonPressedVisuals( FVirtualHand& Hand, EInputEvent Event, const FName& ParameterName, const float PressStrength );
 
 protected:
 
@@ -557,5 +562,8 @@ private:
 
 	/** If this is the first tick or before */
 	bool bFirstTick;
+
+	/** If the coordinate system was in world space before entering (local) scale mode */
+	bool bWasInWorldSpaceBeforeScaleMode;
 };
 
diff --git a/Engine/Source/Editor/VREditor/VREditorTransformGizmo.cpp b/Engine/Source/Editor/VREditor/VREditorTransformGizmo.cpp
--- a/Engine/Source/Editor/VREditor/VREditorTransformGizmo.cpp
+++ b/Engine/Source/Editor/VREditor/VREditorTransformGizmo.cpp
@@ -16,7 +16,6 @@ namespace VREd
 	static FAutoConsoleVariable GizmoScale( TEXT( "VREd.GizmoScale" ), 0.80f, TEXT( "How big the gizmo handles should be" ) );
 	static FAutoConsoleVariable GizmoDistanceScaleFactor( TEXT( "VREd.GizmoDistanceScaleFactor" ), 0.002f, TEXT( "How much the gizmo handles should increase in size with distance from the camera, to make it easier to select" ) );
 	static FAutoConsoleVariable MinActorSizeForTransformGizmo( TEXT( "VREd.MinActorSizeForTransformGizmo" ), 50.0f, TEXT( "How big an object must be in scaled world units before we'll start to shrink the gizmo" ) );
-	static FAutoConsoleVariable AlwaysShowGizmoMeasurementText( TEXT( "VREd.AlwaysShowGizmoMeasurementText" ), 0, TEXT( "When enabled, gizmo measurements will always be visible.  Otherwise, only when hovering over a scale/stretch gizmo handle" ) );
 }
 
 
@@ -49,6 +48,7 @@ ATransformGizmo::ATransformGizmo()
 	UniformScaleGizmoHandleGroup = CreateDefaultSubobject<UVREditorUniformScaleGizmoHandleGroup>( TEXT( "UniformScaleHandles" ), true );
 	UniformScaleGizmoHandleGroup->SetTranslucentGizmoMaterial( TranslucentGizmoMaterial );
 	UniformScaleGizmoHandleGroup->SetGizmoMaterial( GizmoMaterial );
+	UniformScaleGizmoHandleGroup->SetUsePivotPointAsLocation( false );
 	UniformScaleGizmoHandleGroup->AttachParent = SceneComponent;
 	AllHandleGroups.Add( UniformScaleGizmoHandleGroup );
 
@@ -128,8 +128,12 @@ ATransformGizmo::ATransformGizmo()
 }
 
 
-void ATransformGizmo::UpdateGizmo( const EGizmoHandleTypes GizmoType, const ECoordSystem GizmoCoordinateSpace, const FTransform& LocalToWorld, const FBox& LocalBounds, const FVector ViewLocation, bool bAllHandlesVisible, UActorComponent* DraggingHandle, const TArray< UActorComponent* >& HoveringOverHandles, const float GizmoHoverScale, const float GizmoHoverAnimationDuration )
+void ATransformGizmo::UpdateGizmo( const EGizmoHandleTypes GizmoType, const ECoordSystem GizmoCoordinateSpace, const FTransform& LocalToWorld, const FBox& LocalBounds, const FVector ViewLocation, bool bAllHandlesVisible, 
+	UActorComponent* DraggingHandle, const TArray< UActorComponent* >& HoveringOverHandles, const float GizmoHoverScale, const float GizmoHoverAnimationDuration )
 {
+	Super::UpdateGizmo( GizmoType, GizmoCoordinateSpace, LocalToWorld, LocalBounds, ViewLocation, bAllHandlesVisible,
+		DraggingHandle, HoveringOverHandles, GizmoHoverScale, GizmoHoverAnimationDuration );
+
 	const float WorldScaleFactor = GetWorld()->GetWorldSettings()->WorldToMeters / 100.0f;
 
 	// Position the gizmo at the location of the first selected actor
@@ -164,7 +168,8 @@ void ATransformGizmo::UpdateGizmo( const EGizmoHandleTypes GizmoType, const ECoo
 		if (HandleGroup != nullptr)
 		{
 			bool bIsHoveringOrDraggingThisHandleGroup = false;
-			HandleGroup->UpdateGizmoHandleGroup( LocalToWorld, LocalBounds, ViewLocation, bAllHandlesVisible, DraggingHandle, HoveringOverHandles, AnimationAlpha, GizmoScale, GizmoHoverScale, GizmoHoverAnimationDuration, /* Out */ bIsHoveringOrDraggingThisHandleGroup );
+			HandleGroup->UpdateGizmoHandleGroup( LocalToWorld, LocalBounds, ViewLocation, bAllHandlesVisible, DraggingHandle, 
+				HoveringOverHandles, AnimationAlpha, GizmoScale, GizmoHoverScale, GizmoHoverAnimationDuration, /* Out */ bIsHoveringOrDraggingThisHandleGroup );
 			
 			if( HandleGroup->GetHandleType() == EGizmoHandleTypes::Scale && bIsHoveringOrDraggingThisHandleGroup )
 			{
@@ -263,10 +268,8 @@ void ATransformGizmo::UpdateGizmo( const EGizmoHandleTypes GizmoType, const ECoo
 
 	// Update visibility
 	{
-		UpdateHandleVisibility( GizmoType, GizmoCoordinateSpace, bAllHandlesVisible, DraggingHandle );
-
 		const bool bShouldShowMeasurements =
-			VREd::AlwaysShowGizmoMeasurementText->GetInt() != 0 ||
+			GetShowMeasurementText() ||
 			bIsHoveringOrDraggingScaleGizmo;
 
 		for( int32 AxisIndex = 0; AxisIndex < 3; ++AxisIndex )
diff --git a/Engine/Source/Editor/VREditor/VREditorTransformGizmo.h b/Engine/Source/Editor/VREditor/VREditorTransformGizmo.h
--- a/Engine/Source/Editor/VREditor/VREditorTransformGizmo.h
+++ b/Engine/Source/Editor/VREditor/VREditorTransformGizmo.h
@@ -35,7 +35,8 @@ public:
 	/** Called by the world interaction system after we've been spawned into the world, to allow
 	    us to create components and set everything up nicely for the selected objects that we'll be
 		used to manipulate */
-	void UpdateGizmo( const EGizmoHandleTypes GizmoType, const ECoordSystem GizmoCoordinateSpace, const FTransform& LocalToWorld, const FBox& LocalBounds, const FVector ViewLocation, bool bAllHandlesVisible, class UActorComponent* DraggingHandle, const TArray< UActorComponent* >& HoveringOverHandles, const float GizmoHoverScale, const float GizmoHoverAnimationDuration );
+	virtual void UpdateGizmo( const EGizmoHandleTypes GizmoType, const ECoordSystem GizmoCoordinateSpace, const FTransform& LocalToWorld, const FBox& LocalBounds, const FVector ViewLocation, bool bAllHandlesVisible, 
+		class UActorComponent* DraggingHandle, const TArray< UActorComponent* >& HoveringOverHandles, const float GizmoHoverScale, const float GizmoHoverAnimationDuration ) override;
 
 private:
 
diff --git a/Engine/Source/Editor/VREditor/VREditorWorldInteraction.cpp b/Engine/Source/Editor/VREditor/VREditorWorldInteraction.cpp
--- a/Engine/Source/Editor/VREditor/VREditorWorldInteraction.cpp
+++ b/Engine/Source/Editor/VREditor/VREditorWorldInteraction.cpp
@@ -4,6 +4,7 @@
 #include "VREditorWorldInteraction.h"
 #include "VREditorMode.h"
 #include "VREditorTransformGizmo.h"
+#include "Gizmo/VREditorPivotTransformGizmo.h"
 #include "VREditorUISystem.h"
 #include "VREditorFloatingUI.h"
 #include "VREditorDockableWindow.h"
@@ -41,7 +42,7 @@ namespace VREd
 	static FAutoConsoleVariable DragAtLaserImpactInterpolationDuration( TEXT( "VREd.DragAtLaserImpactInterpolationDuration" ), 0.1f, TEXT( "How long we should interpolate objects between positions when dragging under the laser's impact point" ) );
 	static FAutoConsoleVariable DragAtLaserImpactInterpolationThreshold( TEXT( "VREd.DragAtLaserImpactInterpolationThreshold" ), 5.0f, TEXT( "Minimum distance jumped between frames before we'll force interpolation mode to activated" ) );
 	static FAutoConsoleVariable TeleportLaserPointerLength( TEXT( "VREd.TeleportLaserPointerLength" ), 500000.0f, TEXT( "Distance of the LaserPointer for teleporting" ) );
-	static FAutoConsoleVariable UniformScaleSensitivity( TEXT( "VREd.UniformScaleSensitivity" ), 0.002f, TEXT( "Sensitivity for uniform scaling" ) );
+	static FAutoConsoleVariable ScaleSensitivity( TEXT( "VREd.ScaleSensitivity" ), 0.01f, TEXT( "Sensitivity for scaling" ) );
 	static FAutoConsoleVariable GridHapticFeedbackStrength( TEXT( "VREd.GridHapticFeedbackStrength" ), 0.4f, TEXT( "Default strength for haptic feedback when moving across grid points" ) );
 	static FAutoConsoleVariable SelectionHapticFeedbackStrength( TEXT( "VREd.SelectionHapticFeedbackStrength" ), 0.5f, TEXT( "Default strength for haptic feedback when selecting objects" ) );
 	static FAutoConsoleVariable DragHapticFeedbackStrength( TEXT( "VREd.DragHapticFeedbackStrength" ), 1.0f, TEXT( "Default strength for haptic feedback when starting to drag objects" ) );
@@ -55,6 +56,7 @@ namespace VREd
 	static FAutoConsoleVariable HoverHapticFeedbackStrength( TEXT( "VREd.HoverHapticFeedbackStrength" ), 0.1f, TEXT( "Default strength for haptic feedback when hovering" ) );
 	static FAutoConsoleVariable HoverHapticFeedbackTime( TEXT( "VREd.HoverHapticFeedbackTime" ), 0.2f, TEXT( "The minimum time between haptic feedback for hovering" ) );
 	static FAutoConsoleVariable AllowWorldRotationPitchAndRoll( TEXT( "VREd.AllowWorldRotationPitchAndRoll" ), 0, TEXT( "When enabled, you'll not only be able to yaw, but also pitch and roll the world when rotating by gripping with two hands" ) );
+	static FAutoConsoleVariable PivotPointTransformGizmo( TEXT( "VREd.PivotPointTransformGizmo" ), 1, TEXT( "If the pivot point transform gizmo is used instead of the bounding box gizmo" ) );
 }
 
 
@@ -272,16 +274,17 @@ FVREditorWorldInteraction::FVREditorWorldInteraction( FVREditorMode& InitOwner )
 	  GizmoLocalBounds( FBox( 0 ) ),
 	  bDraggedSinceLastSelection( false ),
 	  LastDragGizmoStartTransform( FTransform::Identity ),
+	  StartDragAngleOnRotation( ),
 	  bIsTeleporting( false ),
 	  TeleportLerpAlpha( 0 ),
 	  TeleportStartLocation( FVector::ZeroVector ),
 	  TeleportGoalLocation( FVector::ZeroVector ),
 	  bJustTeleported( false ),
+	  FloatingUIAssetDraggedFrom( nullptr ),
 	  TeleportSound( nullptr ),
 	  DropMaterialOrMaterialSound( nullptr ),
-	  DockSelectDistance( 100 ),
-	  FloatingUIAssetDraggedFrom( nullptr )
-	{
+	  DockSelectDistance( 100 )
+{
 	// Register to find out about VR input events
 	Owner.OnVRAction().AddRaw( this, &FVREditorWorldInteraction::OnVRAction );
 	Owner.OnVRHoverUpdate().AddRaw( this, &FVREditorWorldInteraction::OnVRHoverUpdate );
@@ -1121,7 +1124,6 @@ void FVREditorWorldInteraction::Tick( FEditorViewportClient* ViewportClient, con
 				}
 
 				const FVector OldViewLocation = ViewportClient->GetViewLocation();
-
 				// Dragging transform gizmo handle
 				const bool bWithTwoHands = ( OtherHand != nullptr || OtherHandThatWasAssistingDrag != nullptr );
 				const bool bIsLaserPointerValid = true;
@@ -1144,7 +1146,8 @@ void FVREditorWorldInteraction::Tick( FEditorViewportClient* ViewportClient, con
 					LaserPointerDirection, 
 					bIsLaserPointerValid, 
 					Hand.GizmoStartTransform, 
-					Hand.GizmoStartLocalBounds, 
+					Hand.GizmoStartLocalBounds,
+					Hand.DraggingTransformGizmoComponent.Get(),
 					/* In/Out */ Hand.GizmoSpaceFirstDragUpdateOffsetAlongAxis,
 					/* In/Out */ Hand.GizmoSpaceDragDeltaFromStartOffset,
 					Hand.bIsDrivingVelocityOfSimulatedTransformables,
@@ -1293,6 +1296,7 @@ void FVREditorWorldInteraction::Tick( FEditorViewportClient* ViewportClient, con
 						bIsLaserPointerValid, 
 						Hand.GizmoStartTransform, 
 						Hand.GizmoStartLocalBounds, 
+						Hand.DraggingTransformGizmoComponent.Get(),
 						/* In/Out */ Hand.GizmoSpaceFirstDragUpdateOffsetAlongAxis,
 						/* In/Out */ Hand.GizmoSpaceDragDeltaFromStartOffset,
 						Hand.bIsDrivingVelocityOfSimulatedTransformables,
@@ -1463,10 +1467,34 @@ void FVREditorWorldInteraction::RefreshTransformGizmo( const bool bNewObjectsSel
 {
 	if( GEditor->GetSelectedActorCount() > 0 )
 	{
-		if( TransformGizmoActor == nullptr )
+		// Check if there no gizmo or if the wrong gizmo is being used
+		bool bSpawnNewGizmo = false;
+		if( TransformGizmoActor == nullptr || 
+			( ( TransformGizmoActor->IsA( APivotTransformGizmo::StaticClass() ) && VREd::PivotPointTransformGizmo->GetInt() == 0 ) || 
+			  ( TransformGizmoActor->IsA( ATransformGizmo::StaticClass() ) && VREd::PivotPointTransformGizmo->GetInt() == 1 ) ) )
+		{
+			bSpawnNewGizmo = true;
+		}
+
+		if ( bSpawnNewGizmo )
 		{
+			// Destroy the previous gizmo
+			if ( TransformGizmoActor != nullptr )
+			{
+				TransformGizmoActor->Destroy();
+			}
+
+			// Create the correct gizmo
 			const bool bWithSceneComponent = false;	// We already have our own scene component
-			TransformGizmoActor = Owner.SpawnTransientSceneActor<ATransformGizmo>( TEXT( "TransformGizmo" ), bWithSceneComponent );
+			if( VREd::PivotPointTransformGizmo->GetInt() == 1 )
+			{
+				TransformGizmoActor = Owner.SpawnTransientSceneActor<APivotTransformGizmo>( TEXT( "PivotTransformGizmo" ), bWithSceneComponent );
+			}
+			else
+			{
+				TransformGizmoActor = Owner.SpawnTransientSceneActor<ATransformGizmo>( TEXT( "TransformGizmo" ), bWithSceneComponent );
+			}
+
 			check( TransformGizmoActor != nullptr );
 			TransformGizmoActor->SetOwnerMode( &Owner ) ;
 
@@ -1525,7 +1553,6 @@ void FVREditorWorldInteraction::RefreshTransformGizmo( const bool bNewObjectsSel
 			GizmoSpaceSelectedActorsBounds = GizmoSpaceSelectedActorsBounds.ShiftBy( -GizmoSpaceSelectedActorsBounds.GetCenter() );
 		}
 
-
 		if( bNewObjectsSelected )
 		{
 			TransformGizmoActor->OnNewObjectsSelected();
@@ -1535,8 +1562,9 @@ void FVREditorWorldInteraction::RefreshTransformGizmo( const bool bNewObjectsSel
 		const ECoordSystem GizmoCoordinateSpace = Owner.GetTransformGizmoCoordinateSpace();
 
 		GizmoLocalBounds = GizmoSpaceSelectedActorsBounds;
-		TransformGizmoActor->UpdateGizmo( GizmoType, GizmoCoordinateSpace, GizmoToWorld, GizmoSpaceSelectedActorsBounds, Owner.GetHeadTransform().GetLocation(), bAllHandlesVisible, SingleVisibleHandle, HoveringOverHandles, VREd::GizmoHandleHoverScale->GetFloat(), VREd::GizmoHandleHoverAnimationDuration->GetFloat() );
 
+		TransformGizmoActor->UpdateGizmo( GizmoType, GizmoCoordinateSpace, GizmoToWorld, GizmoSpaceSelectedActorsBounds, Owner.GetHeadTransform().GetLocation(), bAllHandlesVisible, SingleVisibleHandle, 
+			HoveringOverHandles, VREd::GizmoHandleHoverScale->GetFloat(), VREd::GizmoHandleHoverAnimationDuration->GetFloat() );
 
 		// Only draw if snapping is turned on
 		AActor* SnapGridActor = Owner.GetSnapGridActor();
@@ -1683,7 +1711,6 @@ FVector FVREditorWorldInteraction::ComputeConstrainedDragDeltaFromStart( const b
 {
 	FVector ConstrainedGizmoSpaceDeltaFromStart;
 
-
 	bool bConstrainDirectlyToLineOrPlane = false;
 	if( OptionalHandlePlacement.IsSet() )
 	{
@@ -1813,25 +1840,26 @@ FVector FVREditorWorldInteraction::ComputeConstrainedDragDeltaFromStart( const b
 }
 
 
-void FVREditorWorldInteraction::UpdateDragging( 
+void FVREditorWorldInteraction::UpdateDragging(
 	const float DeltaTime,
-	bool& bIsFirstDragUpdate, 
-	const EVREditorDraggingMode DraggingMode, 
-	const ETransformGizmoInteractionType InteractionType, 
-	const bool bWithTwoHands, 
+	bool& bIsFirstDragUpdate,
+	const EVREditorDraggingMode DraggingMode,
+	const ETransformGizmoInteractionType InteractionType,
+	const bool bWithTwoHands,
 	FEditorViewportClient& ViewportClient,
-	const TOptional<FTransformGizmoHandlePlacement> OptionalHandlePlacement, 
-	const FVector& DragDelta, 
-	const FVector& OtherHandDragDelta, 
+	const TOptional<FTransformGizmoHandlePlacement> OptionalHandlePlacement,
+	const FVector& DragDelta,
+	const FVector& OtherHandDragDelta,
 	const FVector& DraggedTo,
 	const FVector& OtherHandDraggedTo,
-	const FVector& DragDeltaFromStart, 
-	const FVector& OtherHandDragDeltaFromStart, 
-	const FVector& LaserPointerStart, 
-	const FVector& LaserPointerDirection, 
+	const FVector& DragDeltaFromStart,
+	const FVector& OtherHandDragDeltaFromStart,
+	const FVector& LaserPointerStart,
+	const FVector& LaserPointerDirection,
 	const bool bIsLaserPointerValid,
-	const FTransform& GizmoStartTransform, 
-	const FBox& GizmoStartLocalBounds, 
+	const FTransform& GizmoStartTransform,
+	const FBox& GizmoStartLocalBounds,
+	const USceneComponent* const DraggingTransformGizmoComponent,
 	FVector& GizmoSpaceFirstDragUpdateOffsetAlongAxis,
 	FVector& DragDeltaFromStartOffset,
 	bool& bIsDrivingVelocityOfSimulatedTransformables,
@@ -1997,7 +2025,61 @@ void FVREditorWorldInteraction::UpdateDragging(
 					bShouldApplyVelocitiesFromDrag = false;
 				}
 			}
+			else if ( InteractionType == ETransformGizmoInteractionType::Scale || InteractionType == ETransformGizmoInteractionType::UniformScale )
+			{
+				const FTransformGizmoHandlePlacement& HandlePlacement = OptionalHandlePlacement.GetValue();
+
+				int32 CenterHandleCount, FacingAxisIndex, CenterAxisIndex;
+				HandlePlacement.GetCenterHandleCountAndFacingAxisIndex( /* Out */ CenterHandleCount, /* Out */ FacingAxisIndex, /* Out */ CenterAxisIndex);
+
+				const FVector PassGizmoSpaceDraggedTo = GizmoStartTransform.InverseTransformPositionNoScale( PassDraggedTo );
+
+				float AddedScaleOnAxis = 0.0f;
+				if ( InteractionType == ETransformGizmoInteractionType::Scale )
+				{
+					AddedScaleOnAxis = PassGizmoSpaceDraggedTo[ FacingAxisIndex ] * VREd::ScaleSensitivity->GetFloat();
+
+					// Invert if we we are scaling on the negative side of the gizmo
+					if ( DraggingTransformGizmoComponent && DraggingTransformGizmoComponent->GetRelativeTransform().GetLocation()[ FacingAxisIndex ] < 0 )
+					{
+						AddedScaleOnAxis *= -1;
+					}
+				}
+				else if ( InteractionType == ETransformGizmoInteractionType::UniformScale )
+				{
+					// Always use Z for uniform scale
+					const FVector RelativeDraggedTo = PassDraggedTo - GizmoStartTransform.GetLocation();
+					AddedScaleOnAxis = RelativeDraggedTo[ 2 ] * VREd::ScaleSensitivity->GetFloat();
+				}
+
+				// Scale the objects!
+				for ( TUniquePtr<FVREditorTransformable>& TransformablePtr : Transformables )
+				{
+					FVREditorTransformable& Transformable = *TransformablePtr;
+					FTransform& PassTargetTransform = bIsUpdatingUnsnappedTarget ? Transformable.UnsnappedTargetTransform : Transformable.TargetTransform;
+					FVector NewTotalScale = Transformable.StartTransform.GetScale3D();
+
+					if ( InteractionType == ETransformGizmoInteractionType::Scale )
+					{
+						NewTotalScale[ FacingAxisIndex ] += AddedScaleOnAxis;
+					}
+					else if ( InteractionType == ETransformGizmoInteractionType::UniformScale )
+					{
+						NewTotalScale += FVector( AddedScaleOnAxis );
+					}
+
+					// Scale snap!
+					if ( !bIsUpdatingUnsnappedTarget && FSnappingUtils::IsScaleSnapEnabled() )
+					{
+						FSnappingUtils::SnapScale( NewTotalScale, FVector::ZeroVector );
+					}
+
+					PassTargetTransform.SetScale3D( NewTotalScale );
+					bMovedAnyTransformables = true;
+					bShouldApplyVelocitiesFromDrag = true;
 
+				}
+			}
 			else if( InteractionType == ETransformGizmoInteractionType::Rotate )
 			{
 				// We only support rotating by a handle currently
@@ -2082,45 +2164,71 @@ void FVREditorWorldInteraction::UpdateDragging(
 					}
 				}
 			}
-			else if( InteractionType == ETransformGizmoInteractionType::UniformScale )
-			{
-				FBox NewGizmoLocalBounds = GizmoStartLocalBounds;
+			else if( InteractionType == ETransformGizmoInteractionType::RotateOnAngle )
+			{			
+				const FTransformGizmoHandlePlacement& HandlePlacement = OptionalHandlePlacement.GetValue();
 
-				float AddedScale = DragDeltaFromStart.Z * VREd::UniformScaleSensitivity->GetFloat();
-				NewGizmoLocalBounds.Min += GizmoStartLocalBounds.Min * AddedScale;
-				NewGizmoLocalBounds.Max += GizmoStartLocalBounds.Max * AddedScale;
+				int32 CenterHandleCount, FacingAxisIndex, CenterAxisIndex;
+				HandlePlacement.GetCenterHandleCountAndFacingAxisIndex( /* Out */ CenterHandleCount, /* Out */ FacingAxisIndex, /* Out */ CenterAxisIndex);
+		
+				FVector GizmoSpaceFacingAxisVector = UVREditorGizmoHandleGroup::GetAxisVector( FacingAxisIndex, HandlePlacement.Axes[ FacingAxisIndex ] );
+				
+				if ( DraggingTransformGizmoComponent )
+				{
+					const FVector RotationHandlePlaneLocation = TransformGizmoActor->GetActorLocation();		
+					FVector RotationHandlePlaneNormal = DraggingTransformGizmoComponent->GetComponentTransform().GetRotation().Vector();
+					RotationHandlePlaneNormal.Normalize();
 
-				const FVector NewGizmoLocalSize = NewGizmoLocalBounds.GetSize();
+					// Get the laserpointer intersection on the plane of the handle
+					FPlane Plane = FPlane( RotationHandlePlaneLocation, RotationHandlePlaneNormal );
+					const FVector IntersectWorldLocationOnPlane = FMath::LinePlaneIntersection( LaserPointerStart, LaserPointerStart + LaserPointerDirection, RotationHandlePlaneLocation, RotationHandlePlaneNormal );
+					
+					const FTransform WorldToGizmo = GizmoStartTransform.Inverse();
+					
+					DrawDebugLine( Owner.GetWorld(), TransformGizmoActor->GetActorLocation(), IntersectWorldLocationOnPlane, FColor::Red, false, -1.f, 0, 3.0f * Owner.GetWorldScaleFactor() );
 
-				FVector NewGizmoLocalScaleFromStart = NewGizmoLocalSize;
-				if (!FMath::IsNearlyZero( NewGizmoLocalSize.GetAbsMin() ))
-				{
-					NewGizmoLocalScaleFromStart = NewGizmoLocalSize / GizmoStartLocalBounds.GetSize();
-				}
-				else
-				{
-					// Zero scale.  This is allowed in Unreal, for better or worse.
-					NewGizmoLocalScaleFromStart = FVector( 0 );
-				}
+					// Rotate the objects!
+					for ( TUniquePtr<FVREditorTransformable>& TransformablePtr : Transformables )
+					{
+						FVREditorTransformable& Transformable = *TransformablePtr;
+						FTransform& PassTargetTransform = bIsUpdatingUnsnappedTarget ? Transformable.UnsnappedTargetTransform : Transformable.TargetTransform;
 
-				// Scale the objects!
-				for (TUniquePtr<FVREditorTransformable>& TransformablePtr : Transformables)
-				{
-					FVREditorTransformable& Transformable = *TransformablePtr;
-					FTransform& PassTargetTransform = bIsUpdatingUnsnappedTarget ? Transformable.UnsnappedTargetTransform : Transformable.TargetTransform;
+						// The starting transformation of this transformable in gizmo space
+						const FTransform GizmoSpaceStartTransform = Transformable.StartTransform * WorldToGizmo;
+						const FVector TransformableSpaceFacingAxisVector = Transformable.StartTransform.InverseTransformVectorNoScale( GizmoStartTransform.InverseTransformVectorNoScale( GizmoSpaceFacingAxisVector ) );
 
-					FVector NewTotalScale = Transformable.StartTransform.GetScale3D() * NewGizmoLocalScaleFromStart;
+						// Relative offset of the intersection on the plane
+						const FVector IntersectLocalLocationOnPlane = Transformable.StartTransform.InverseTransformPositionNoScale( IntersectWorldLocationOnPlane );
+						FVector RotatedIntersectLocationOnPlane = TransformableSpaceFacingAxisVector.Rotation().UnrotateVector( IntersectLocalLocationOnPlane );
+						float AngleToIntersectedLocation = FMath::Atan2( RotatedIntersectLocationOnPlane.Y, RotatedIntersectLocationOnPlane.Z );
 
-					// Scale snap!
-					if( !bIsUpdatingUnsnappedTarget && FSnappingUtils::IsScaleSnapEnabled() )
-					{
-						FSnappingUtils::SnapScale( NewTotalScale, FVector::ZeroVector );
-					}
+						if ( !StartDragAngleOnRotation.IsSet() )
+						{
+							StartDragAngleOnRotation = AngleToIntersectedLocation;
+						}
 
-					PassTargetTransform.SetScale3D( NewTotalScale );
+						// Delta rotation in gizmo space between the starting and the intersection rotation
+						const float AngleDeltaRotationFromStart = FMath::FindDeltaAngle( AngleToIntersectedLocation, StartDragAngleOnRotation.GetValue() );
+						const FQuat GizmoSpaceDeltaRotation = FQuat( GizmoSpaceFacingAxisVector, AngleDeltaRotationFromStart );
 
-					bMovedAnyTransformables = true;
-					bShouldApplyVelocitiesFromDrag = true;
+						// Add the delta rotation to the starting transformation
+						FTransform GizmoSpaceRotatedTransform = GizmoSpaceStartTransform * GizmoSpaceDeltaRotation;
+
+						// Snap rotation in gizmo space
+						if ( !bIsUpdatingUnsnappedTarget )
+						{
+							FRotator SnappedRotation = GizmoSpaceRotatedTransform.GetRotation().Rotator();
+							FSnappingUtils::SnapRotatorToGrid( SnappedRotation );
+							GizmoSpaceRotatedTransform.SetRotation( SnappedRotation.Quaternion() );
+						}
+
+						// Rotation back to world space 
+						const FTransform WorldSpaceRotatedTransform = GizmoSpaceRotatedTransform * GizmoStartTransform;
+
+						PassTargetTransform = WorldSpaceRotatedTransform;
+						bMovedAnyTransformables = true;
+						bShouldApplyVelocitiesFromDrag = false;
+					}
 				}
 			}
 		}
@@ -2434,8 +2542,7 @@ void FVREditorWorldInteraction::UpdateDragging(
 	bIsFirstDragUpdate = false;
 }
 
-
-void FVREditorWorldInteraction::StartDraggingActors( FVRAction VRAction, UActorComponent* ClickedComponent, const FVector HitLocation, const bool bIsPlacingActors )
+void FVREditorWorldInteraction::StartDraggingActors(FVRAction VRAction, UActorComponent* ClickedComponent, const FVector HitLocation, const bool bIsPlacingActors)
 {
 	if( IsInteractableComponent( ClickedComponent ) )
 	{
@@ -2496,7 +2603,7 @@ void FVREditorWorldInteraction::StartDraggingActors( FVRAction VRAction, UActorC
 			// the gizmo to transform it.
 			if( bUsingGizmo )
 			{
-				Hand.DraggingTransformGizmoComponent = ClickedComponent;
+				Hand.DraggingTransformGizmoComponent = Cast<USceneComponent>( ClickedComponent );
 			}
 			else
 			{
@@ -2622,6 +2729,12 @@ void FVREditorWorldInteraction::StopDragging( const int32 HandIndex )
 				Hand.DraggingMode == EVREditorDraggingMode::ActorsFreely ||
 				Hand.DraggingMode == EVREditorDraggingMode::ActorsAtLaserImpact )
 			{
+				if ( Hand.DraggingMode == EVREditorDraggingMode::ActorsWithGizmo &&
+					 Hand.TransformGizmoInteractionType == ETransformGizmoInteractionType::RotateOnAngle )
+				{
+					StartDragAngleOnRotation.Reset();
+				}
+
 				// Finalize undo
 				{
 					// @todo vreditor undo: This doesn't actually encapsulate any inertial movement that happens after the drag is released! 
diff --git a/Engine/Source/Editor/VREditor/VREditorWorldInteraction.h b/Engine/Source/Editor/VREditor/VREditorWorldInteraction.h
--- a/Engine/Source/Editor/VREditor/VREditorWorldInteraction.h
+++ b/Engine/Source/Editor/VREditor/VREditorWorldInteraction.h
@@ -28,7 +28,7 @@ public:
 	void Render( const class FSceneView* SceneView, class FViewport* Viewport, class FPrimitiveDrawInterface* PDI );
 
 	/** Gets the transform gizmo actor, or returns null if we currently don't have one */
-	class ATransformGizmo* GetTransformGizmoActor() const
+	class ABaseTransformGizmo* GetTransformGizmoActor() const
 	{
 		return TransformGizmoActor;
 	}
@@ -72,7 +72,7 @@ protected:
 	/** Called by the world interaction system when one of our components is dragged by the user.  If null is
 	    passed in then we'll treat it as dragging the whole object (rather than a specific axis/handle) */
 	void UpdateDragging( 
-		const float DeltaTime, 
+		const float DeltaTime,
 		bool& bIsFirstDragUpdate, 
 		const EVREditorDraggingMode DraggingMode, 
 		const ETransformGizmoInteractionType InteractionType, 
@@ -90,6 +90,7 @@ protected:
 		const bool bIsLaserPointerValid, 
 		const FTransform& GizmoStartTransform, 
 		const FBox& GizmoStartLocalBounds, 
+		const USceneComponent* const DraggingTransformGizmoComponent,
 		FVector& GizmoSpaceFirstDragUpdateOffsetAlongAxis,
 		FVector& DragDeltaFromStartOffset,
 		bool& bIsDrivingVelocityOfSimulatedTransformables,
@@ -178,7 +179,7 @@ protected:
 	//
 
 	/** Transform gizmo actor, for manipulating selected actor(s) */
-	class ATransformGizmo* TransformGizmoActor;
+	class ABaseTransformGizmo* TransformGizmoActor;
 
 	/** Current gizmo bounds in local space.  Updated every frame.  This is not the same as the gizmo actor's bounds. */
 	FBox GizmoLocalBounds;
@@ -189,7 +190,8 @@ protected:
 	/** Gizmo start transform of the last drag we did with either hand.  Only valid if bDraggedSinceLastSelection is true */
 	FTransform LastDragGizmoStartTransform;
 
-
+	/** Starting angle when rotating an object with  ETransformGizmoInteractionType::RotateOnAngle */
+	TOptional<float> StartDragAngleOnRotation;
 
 	//
 	// Hover state
diff --git a/Engine/Source/Editor/VREditor/VREditorWorldInteractionTypes.h b/Engine/Source/Editor/VREditor/VREditorWorldInteractionTypes.h
--- a/Engine/Source/Editor/VREditor/VREditorWorldInteractionTypes.h
+++ b/Engine/Source/Editor/VREditor/VREditorWorldInteractionTypes.h
@@ -53,9 +53,15 @@ enum class ETransformGizmoInteractionType : uint8
 	/** Stretch the object non-uniformly while simultaneously repositioning it so account for it's new size */
 	StretchAndReposition,
 
+	/** Scaling the object non-uniformly while	 */
+	Scale,
+
 	/** Rotate the object(s) around a specific axis relative to the gizmo's pivot, translating as needed */
 	Rotate,
 
+	/** Rotate the object(s) around the pivot based on start angle */
+	RotateOnAngle,
+
 	/** Uniform scale the object(s) */
 	UniformScale,
 };
diff --git a/Engine/Source/Editor/VREditor/VirtualHand.h b/Engine/Source/Editor/VREditor/VirtualHand.h
--- a/Engine/Source/Editor/VREditor/VirtualHand.h
+++ b/Engine/Source/Editor/VREditor/VirtualHand.h
@@ -230,13 +230,13 @@ struct FVirtualHand
 	TOptional<FTransformGizmoHandlePlacement> OptionalHandlePlacement;
 
 	/** The gizmo component we're dragging right now */
-	TWeakObjectPtr<class UActorComponent> DraggingTransformGizmoComponent;
+	TWeakObjectPtr<class USceneComponent> DraggingTransformGizmoComponent;
 
 	/** Gizmo component that we're hovering over, or nullptr if not hovering over any */
-	TWeakObjectPtr<class UActorComponent> HoveringOverTransformGizmoComponent;
+	TWeakObjectPtr<class USceneComponent> HoveringOverTransformGizmoComponent;
 
 	/** Gizmo handle that we hovered over last (used only for avoiding spamming of hover haptics!) */
-	TWeakObjectPtr<class UActorComponent> HoverHapticCheckLastHoveredGizmoComponent;
+	TWeakObjectPtr<class USceneComponent> HoverHapticCheckLastHoveredGizmoComponent;
 
 	/** Last real time that we played hover haptic feedback for a gizmo handle.  This is just to avoid spamming haptic effects every frame due to motion jitter */
 	FTimespan LastGizmoHandleHapticTime;
