Adding flashlight functionality to the VR Editor quick menu. The "Light" button on the quick menu will add a flashlight to the hand without the menu attached.

diff --git a/Engine/Source/Editor/VREditor/VREditorMode.cpp b/Engine/Source/Editor/VREditor/VREditorMode.cpp
--- a/Engine/Source/Editor/VREditor/VREditorMode.cpp
+++ b/Engine/Source/Editor/VREditor/VREditorMode.cpp
@@ -102,6 +102,8 @@ FVREditorMode::FVREditorMode()
 	  AppTimeModeEntered( FTimespan::Zero() ),
 	  AvatarMeshActor( nullptr ),
 	  HeadMeshComponent( nullptr ),
+      FlashlightComponent( nullptr ),
+	  bIsFlashlightOn( false ),
 	  WorldMovementGridMeshComponent( nullptr ),
 	  WorldMovementGridMID( nullptr ),
 	  WorldMovementGridOpacity( 0.0f ),	// NOTE: Intentionally not quite zero so that we update the MIDs on the first frame
@@ -476,6 +478,7 @@ void FVREditorMode::Exit()
 		DestroyTransientActor( AvatarMeshActor );
 		AvatarMeshActor = nullptr;
 		HeadMeshComponent = nullptr;
+		FlashlightComponent = nullptr;
 		WorldMovementGridMeshComponent = nullptr;
 		WorldMovementGridMID = nullptr;
 		DestroyTransientActor( SnapGridActor );
@@ -486,6 +489,7 @@ void FVREditorMode::Exit()
 		ScaleProgressMeshComponent = nullptr;
 		CurrentScaleProgressMeshComponent = nullptr;
 		UserScaleIndicatorText = nullptr;
+		
 
 		for( int32 HandIndex = 0; HandIndex < VREditorConstants::NumVirtualHands; ++HandIndex )
 		{
@@ -1240,6 +1244,17 @@ void FVREditorMode::Tick( FEditorViewportClient* ViewportClient, float DeltaTime
 		}
 	}
 
+	// Updating the scale and intensity of the flashlight according to the world scale
+	if (FlashlightComponent)
+	{
+		float CurrentFalloffExponent = FlashlightComponent->LightFalloffExponent;
+		//@todo vreditor tweak
+		float UpdatedFalloffExponent = FMath::Clamp(CurrentFalloffExponent / GetWorldScaleFactor(), 2.0f, 16.0f);
+		FlashlightComponent->SetLightFalloffExponent(UpdatedFalloffExponent);
+	}
+	
+
+
 	// Updating the opacity and visibility of the grid according to the controllers
 	{
 		if( VREd::ShowMovementGrid->GetInt() == 1)
@@ -1400,6 +1415,7 @@ bool FVREditorMode::HandleInputKey(FEditorViewportClient* ViewportClient, FViewp
 		FVRAction VRAction( KnownAction->ActionType, GetHandIndexFromKey( Key ) );
 		if( VRAction.HandIndex != INDEX_NONE )
 		{
+
 			FVirtualHand& Hand = VirtualHands[ VRAction.HandIndex ];
 
 			bool bShouldAbsorbEvent = false;
@@ -1658,6 +1674,7 @@ void FVREditorMode::AddReferencedObjects( FReferenceCollector& Collector )
 {
 	Collector.AddReferencedObject( AvatarMeshActor );
 	Collector.AddReferencedObject( HeadMeshComponent );
+	Collector.AddReferencedObject( FlashlightComponent );
 	Collector.AddReferencedObject( WorldMovementGridMeshComponent );
 	Collector.AddReferencedObject( SnapGridActor );
 	Collector.AddReferencedObject( SnapGridMeshComponent );
@@ -2517,6 +2534,8 @@ void FVREditorMode::CleanUpActorsBeforeMapChangeOrSimulate()
 	DestroyTransientActor( AvatarMeshActor );
 	AvatarMeshActor = nullptr;
 	HeadMeshComponent = nullptr;
+	FlashlightComponent = nullptr;
+
 	WorldMovementGridMeshComponent = nullptr;
 
 	if( WorldMovementGridMID != nullptr )
@@ -2625,6 +2644,33 @@ void FVREditorMode::SnapSelectedActorsToGround()
 	WorldInteraction->SetupTransformablesForSelectedActors();
 }
 
+void FVREditorMode::ToggleFlashlight( const int32 HandIndex )
+{
+
+	if (FlashlightComponent == nullptr)
+	{
+		FlashlightComponent = NewObject<USpotLightComponent>(AvatarMeshActor);
+		AvatarMeshActor->AddOwnedComponent(FlashlightComponent);
+		FlashlightComponent->RegisterComponent();
+		FlashlightComponent->SetMobility(EComponentMobility::Movable);
+		FlashlightComponent->SetCastShadows(false);
+		FlashlightComponent->bUseInverseSquaredFalloff = false;
+		//@todo vreditor tweak
+		FlashlightComponent->SetLightFalloffExponent(8.0f);
+		FlashlightComponent->SetIntensity(20.0f);
+		FlashlightComponent->SetOuterConeAngle(25.0f);
+		FlashlightComponent->SetInnerConeAngle(25.0f);
+		
+	}
+
+	FVirtualHand Hand = GetVirtualHand(HandIndex);
+	FlashlightComponent->AttachTo(Hand.MotionControllerComponent);
+	bIsFlashlightOn = !bIsFlashlightOn;
+	FlashlightComponent->SetVisibility(bIsFlashlightOn);
+	
+
+}
+
 void FVREditorMode::SetLaserVisuals( const int32 HandIndex, const FLinearColor& NewColor, const float CrawlFade, const float CrawlSpeed )
 {
 	FVirtualHand Hand = GetVirtualHand( HandIndex );
diff --git a/Engine/Source/Editor/VREditor/VREditorMode.h b/Engine/Source/Editor/VREditor/VREditorMode.h
--- a/Engine/Source/Editor/VREditor/VREditorMode.h
+++ b/Engine/Source/Editor/VREditor/VREditorMode.h
@@ -90,6 +90,12 @@ public:
 		return MotionControllerID;
 	}
 
+	/** Returns whether or not the flashlight is visible */
+	bool IsFlashlightOn() const
+	{
+		return bIsFlashlightOn;
+	}
+
 	/** Returns the time since the VR Editor mode was last entered */
 	inline FTimespan GetTimeSinceModeEntered() const
 	{
@@ -263,6 +269,9 @@ public:
 	/** Snaps the selected objects to the ground */
 	void SnapSelectedActorsToGround();
 
+	/** Spawns a flashlight on the specified hand */
+	void ToggleFlashlight( const int32 HandIndex );
+
 	/** Sets the visuals of the LaserPointer */
 	void SetLaserVisuals( const int32 HandIndex, const FLinearColor& NewColor, const float CrawlFade = 0.0f, const float CrawlSpeed = 0.0f );
 
@@ -426,6 +435,12 @@ protected:
 	/** Our avatar's head mesh */
 	class UStaticMeshComponent* HeadMeshComponent;
 
+	/** Spotlight for the flashlight */
+	class USpotLightComponent* FlashlightComponent;
+
+	/** If there is currently a flashlight in the scene */
+	bool bIsFlashlightOn;
+
 
 	//
 	// World movement grid & FX
@@ -527,6 +542,8 @@ protected:
 	// Slate Input Processor
 	TSharedPtr<FVREditorInputProcessor> InputProcessor;
 
+
+
 	//
 	// Colors
 	//
diff --git a/Engine/Source/Editor/VREditor/VREditorQuickMenu.cpp b/Engine/Source/Editor/VREditor/VREditorQuickMenu.cpp
--- a/Engine/Source/Editor/VREditor/VREditorQuickMenu.cpp
+++ b/Engine/Source/Editor/VREditor/VREditorQuickMenu.cpp
@@ -319,6 +319,17 @@ bool UVREditorQuickMenu::IsTutorialVisible() const
 	return ( GetOwner()->GetOwner().IsShowingEditorUIPanel( FVREditorUISystem::EEditorUIPanel::Tutorial ) );
 }
 
+void UVREditorQuickMenu::OnLightButtonClicked(const bool bIsChecked)
+{
+	const int32 HandIndexWithoutQuickMenu = GetOwner()->GetDockedTo() == AVREditorFloatingUI::EDockedTo::LeftArm ? VREditorConstants::RightHandIndex : VREditorConstants::LeftHandIndex;
+	GetOwner()->GetOwner().GetOwner().ToggleFlashlight(HandIndexWithoutQuickMenu);
+}
+
+bool UVREditorQuickMenu::IsLightVisible() const
+{
+	return (GetOwner()->GetOwner().GetOwner().IsFlashlightOn());
+}
+
 
 void UVREditorQuickMenu::OnAssetEditorButtonClicked( const bool bIsChecked )
 {
diff --git a/Engine/Source/Editor/VREditor/VREditorQuickMenu.h b/Engine/Source/Editor/VREditor/VREditorQuickMenu.h
--- a/Engine/Source/Editor/VREditor/VREditorQuickMenu.h
+++ b/Engine/Source/Editor/VREditor/VREditorQuickMenu.h
@@ -132,6 +132,12 @@ public:
 	UFUNCTION( BlueprintPure, Category = "Default" )
 	bool IsTutorialVisible() const;
 
+	UFUNCTION(BlueprintCallable, Category = "Default")
+	void OnLightButtonClicked(const bool bIsChecked);
+
+	UFUNCTION(BlueprintPure, Category = "Default")
+	bool IsLightVisible() const;
+
 	UFUNCTION( BlueprintCallable, Category = "Default" )
 	void OnAssetEditorButtonClicked( const bool bIsChecked );
 
