VR Editor :
- Fix mac warning for WorldInteraction constructor list order
- Fix mac warning for UniversalGizmoHandle using the same variable name

diff --git a/Engine/Source/Editor/VREditor/Gizmo/VREditorUniformScaleGizmoHandle.cpp b/Engine/Source/Editor/VREditor/Gizmo/VREditorUniformScaleGizmoHandle.cpp
--- a/Engine/Source/Editor/VREditor/Gizmo/VREditorUniformScaleGizmoHandle.cpp
+++ b/Engine/Source/Editor/VREditor/Gizmo/VREditorUniformScaleGizmoHandle.cpp
@@ -98,7 +98,7 @@ EGizmoHandleTypes UVREditorUniformScaleGizmoHandleGroup::GetHandleType() const
 	return EGizmoHandleTypes::Scale;
 }
 
-void UVREditorUniformScaleGizmoHandleGroup::SetUsePivotPointAsLocation( bool bUsePivotAsLocation )
+void UVREditorUniformScaleGizmoHandleGroup::SetUsePivotPointAsLocation( const bool bInUsePivotAsLocation )
 {
-	this->bUsePivotAsLocation = bUsePivotAsLocation;
+	bUsePivotAsLocation = bInUsePivotAsLocation;
 }
diff --git a/Engine/Source/Editor/VREditor/Gizmo/VREditorUniformScaleGizmoHandle.h b/Engine/Source/Editor/VREditor/Gizmo/VREditorUniformScaleGizmoHandle.h
--- a/Engine/Source/Editor/VREditor/Gizmo/VREditorUniformScaleGizmoHandle.h
+++ b/Engine/Source/Editor/VREditor/Gizmo/VREditorUniformScaleGizmoHandle.h
@@ -29,7 +29,7 @@ public:
 	virtual EGizmoHandleTypes GetHandleType() const override;
 
 	/** Sets if the pivot point is used as location for the handle */
-	void SetUsePivotPointAsLocation( bool bUsePivotAsLocation );
+	void SetUsePivotPointAsLocation( const bool bInUsePivotAsLocation );
 
 private:
 
diff --git a/Engine/Source/Editor/VREditor/VREditorWorldInteraction.cpp b/Engine/Source/Editor/VREditor/VREditorWorldInteraction.cpp
--- a/Engine/Source/Editor/VREditor/VREditorWorldInteraction.cpp
+++ b/Engine/Source/Editor/VREditor/VREditorWorldInteraction.cpp
@@ -280,10 +280,10 @@ FVREditorWorldInteraction::FVREditorWorldInteraction( FVREditorMode& InitOwner )
 	  TeleportStartLocation( FVector::ZeroVector ),
 	  TeleportGoalLocation( FVector::ZeroVector ),
 	  bJustTeleported( false ),
-	  FloatingUIAssetDraggedFrom( nullptr ),
 	  TeleportSound( nullptr ),
 	  DropMaterialOrMaterialSound( nullptr ),
-	  DockSelectDistance( 100 )
+	  DockSelectDistance( 100 ),
+	  FloatingUIAssetDraggedFrom( nullptr )
 {
 	// Register to find out about VR input events
 	Owner.OnVRAction().AddRaw( this, &FVREditorWorldInteraction::OnVRAction );
