VR Editor: Adding screenshot button to the VR Editor Quick Menu. The screenshot does undistorted left eye + right eye, and does not capture UI.

diff --git a/Engine/Source/Editor/VREditor/VREditorQuickMenu.cpp b/Engine/Source/Editor/VREditor/VREditorQuickMenu.cpp
--- a/Engine/Source/Editor/VREditor/VREditorQuickMenu.cpp
+++ b/Engine/Source/Editor/VREditor/VREditorQuickMenu.cpp
@@ -330,6 +330,29 @@ bool UVREditorQuickMenu::IsLightVisible() const
 	return (GetOwner()->GetOwner().GetOwner().IsFlashlightOn());
 }
 
+void UVREditorQuickMenu::OnScreenshotButtonClicked()
+{
+	// @todo vreditor: update after direct buffer grab changes
+	if ((GetOwner()->GetOwner().GetOwner().GetHMDDeviceType() == EHMDDeviceType::DT_OculusRift))
+	{	
+		TSharedRef<SWidget> WindowRef = GetOwner()->GetOwner().GetOwner().GetLevelViewportPossessedForVR().AsWidget();
+		TArray<FColor> OutImageData;
+		FIntVector OutImageSize;
+		FSlateApplication::Get().TakeScreenshot(WindowRef, OutImageData, OutImageSize);
+		float ScreenPercentAdjust = (2048.0f / OutImageSize.X) * 100.0f;
+		GEngine->DeferredCommands.Add(FString::Printf(TEXT("HMD SP %3.0f"), ScreenPercentAdjust));
+	}
+	else if(GetOwner()->GetOwner().GetOwner().GetHMDDeviceType() == EHMDDeviceType::DT_SteamVR)
+	{
+		GEngine->DeferredCommands.Add(FString::Printf(TEXT("r.screenpercentage 100")));
+	}
+
+	// Take a screenshot from the HMD, left eye + right eye undistorted
+	// @todo vreditor: verify an option to make hands visible
+	const bool bShowUI = false;
+	const bool bUseCustomFilename = true;
+	FScreenshotRequest::RequestScreenshot(FString(), bShowUI, bUseCustomFilename);
+}
 
 void UVREditorQuickMenu::OnAssetEditorButtonClicked( const bool bIsChecked )
 {
diff --git a/Engine/Source/Editor/VREditor/VREditorQuickMenu.h b/Engine/Source/Editor/VREditor/VREditorQuickMenu.h
--- a/Engine/Source/Editor/VREditor/VREditorQuickMenu.h
+++ b/Engine/Source/Editor/VREditor/VREditorQuickMenu.h
@@ -138,6 +138,9 @@ public:
 	UFUNCTION(BlueprintPure, Category = "Default")
 	bool IsLightVisible() const;
 
+	UFUNCTION(BlueprintCallable, Category = "Default")
+	void OnScreenshotButtonClicked();
+
 	UFUNCTION( BlueprintCallable, Category = "Default" )
 	void OnAssetEditorButtonClicked( const bool bIsChecked );
 
