VR Editor: Fix for UE-27727, unable to undo a duplication

diff --git a/Engine/Source/Editor/VREditor/VREditorMode.cpp b/Engine/Source/Editor/VREditor/VREditorMode.cpp
--- a/Engine/Source/Editor/VREditor/VREditorMode.cpp
+++ b/Engine/Source/Editor/VREditor/VREditorMode.cpp
@@ -2626,9 +2626,9 @@ void FVREditorMode::Paste()
 
 void FVREditorMode::Duplicate()
 {
-	ABrush::SetSuppressBSPRegeneration( true );
-	GEditor->edactDuplicateSelected( GetWorld()->GetCurrentLevel(), false );
-	ABrush::SetSuppressBSPRegeneration( false );
+	// @todo vreditor: Needs CanExecute()  (see LevelEditorActions.cpp)
+	GUnrealEd->Exec( GetWorld(), TEXT( "DUPLICATE" ) );
+
 }
 
 
