VR Editor: Fixed asset editor UI popping up whenever a Modes icon was clicked in VR

diff --git a/Engine/Source/Editor/VREditor/VREditorUISystem.cpp b/Engine/Source/Editor/VREditor/VREditorUISystem.cpp
--- a/Engine/Source/Editor/VREditor/VREditorUISystem.cpp
+++ b/Engine/Source/Editor/VREditor/VREditorUISystem.cpp
@@ -29,6 +29,7 @@
 #include "SScaleBox.h"
 #include "SDPIScaler.h"
 #include "SWidget.h"
+#include "SDockTab.h"
 
 namespace VREd
 {
@@ -1307,7 +1308,12 @@ void FVREditorUISystem::OnProxyTabLaunched(TSharedPtr<SDockTab> NewTab)
 
 void FVREditorUISystem::OnAttentionDrawnToTab(TSharedPtr<SDockTab> NewTab)
 {
-	ShowAssetEditor();
+	// @todo vreditor: When clicking on Modes icons in the Modes panel, the LevelEditor tab is invoked which causes an empty asset editor to pop-up
+	static FName LevelEditorTabID( "LevelEditor" );
+	if( NewTab->GetLayoutIdentifier().TabType != LevelEditorTabID )
+	{
+		ShowAssetEditor();
+	}
 }
 
 void FVREditorUISystem::ShowAssetEditor()
