VR Editor: Adding access to the World Settings UI from the Quick Menu (UETOOL-722). Temp. quick menu button reorg.

diff --git a/Engine/Source/Editor/VREditor/VREditorQuickMenu.cpp b/Engine/Source/Editor/VREditor/VREditorQuickMenu.cpp
--- a/Engine/Source/Editor/VREditor/VREditorQuickMenu.cpp
+++ b/Engine/Source/Editor/VREditor/VREditorQuickMenu.cpp
@@ -365,4 +365,15 @@ bool UVREditorQuickMenu::IsAssetEditorVisible() const
 	return ( GetOwner()->GetOwner().IsShowingEditorUIPanel( FVREditorUISystem::EEditorUIPanel::AssetEditor ) );
 }
 
+void UVREditorQuickMenu::OnWorldSettingsButtonClicked(const bool bIsChecked)
+{
+	GetOwner()->GetOwner().TogglePanelVisibility(FVREditorUISystem::EEditorUIPanel::WorldSettings);
+
+}
+
+bool UVREditorQuickMenu::IsWorldSettingsVisible() const
+{
+	return (GetOwner()->GetOwner().IsShowingEditorUIPanel(FVREditorUISystem::EEditorUIPanel::WorldSettings));
+}
+
 #undef LOCTEXT_NAMESPACE
\ No newline at end of file
diff --git a/Engine/Source/Editor/VREditor/VREditorQuickMenu.h b/Engine/Source/Editor/VREditor/VREditorQuickMenu.h
--- a/Engine/Source/Editor/VREditor/VREditorQuickMenu.h
+++ b/Engine/Source/Editor/VREditor/VREditorQuickMenu.h
@@ -147,6 +147,12 @@ public:
 	UFUNCTION( BlueprintPure, Category = "Default" )
 	bool IsAssetEditorVisible() const;
 
+	UFUNCTION(BlueprintCallable, Category = "Default")
+	void OnWorldSettingsButtonClicked(const bool bIsChecked);
+
+	UFUNCTION(BlueprintPure, Category = "Default")
+	bool IsWorldSettingsVisible() const;
+
 	UFUNCTION(BlueprintImplementableEvent, meta=(DisplayName = "Refresh"))
 	void RefreshUI();
 
diff --git a/Engine/Source/Editor/VREditor/VREditorUISystem.cpp b/Engine/Source/Editor/VREditor/VREditorUISystem.cpp
--- a/Engine/Source/Editor/VREditor/VREditorUISystem.cpp
+++ b/Engine/Source/Editor/VREditor/VREditorUISystem.cpp
@@ -25,6 +25,10 @@
 // Actor Details, Modes
 #include "LevelEditor.h"
 
+// World Settings
+#include "Editor/PropertyEditor/Public/PropertyEditorModule.h"
+#include "Editor/PropertyEditor/Public/IDetailsView.h"
+
 #include "SLevelViewport.h"
 #include "SScaleBox.h"
 #include "SDPIScaler.h"
@@ -873,6 +877,37 @@ void FVREditorUISystem::CreateUIs()
 
 			FAssetEditorManager::Get().OnAssetEditorOpened().AddRaw(this, &FVREditorUISystem::OnAssetEditorOpened);
 		}
+		// create the world settings menu
+		{
+			FPropertyEditorModule& PropertyEditorModule = FModuleManager::Get().LoadModuleChecked< FPropertyEditorModule >( "PropertyEditor" );
+
+			FDetailsViewArgs DetailsViewArgs( false, false, true, FDetailsViewArgs::HideNameArea, false, GUnrealEd );
+			DetailsViewArgs.bShowActorLabel = false;
+
+			TSharedRef<IDetailsView> WorldSettings = PropertyEditorModule.CreateDetailView( DetailsViewArgs );
+
+			UWorld* World = GetOwner().GetWorld();
+			if ( World != NULL )
+			{
+				WorldSettings->SetObject( World->GetWorldSettings() );
+			}
+
+			TSharedRef<SWidget> WidgetToDraw = 
+				SNew( SDPIScaler )
+				.DPIScale( VREd::EditorUIScale->GetFloat() )
+				[
+					WorldSettings
+				]
+			;
+
+			const bool bWithSceneComponent = false;
+			AVREditorFloatingUI* WorldSettingsUI = GetOwner().SpawnTransientSceneActor< AVREditorDockableWindow >( TEXT( "WorldSettingsUI" ), bWithSceneComponent );
+			WorldSettingsUI->SetSlateWidget( *this, WidgetToDraw, DefaultResolution, VREd::EditorUISize->GetFloat(), AVREditorFloatingUI::EDockedTo::Nothing );
+			WorldSettingsUI->ShowUI( false );
+			FloatingUIs.Add( WorldSettingsUI );
+
+			EditorUIPanels[ (int32)EEditorUIPanel::WorldSettings ] = WorldSettingsUI;
+		}
 	}
 }
 
diff --git a/Engine/Source/Editor/VREditor/VREditorUISystem.h b/Engine/Source/Editor/VREditor/VREditorUISystem.h
--- a/Engine/Source/Editor/VREditor/VREditorUISystem.h
+++ b/Engine/Source/Editor/VREditor/VREditorUISystem.h
@@ -17,7 +17,7 @@ public:
 		Modes,
 		Tutorial,
 		AssetEditor,
-
+		WorldSettings,
 		// ...
 
 		TotalCount,
