VR Editor: Adding return values to the quick menu functions for more discrete button refreshes

diff --git a/Engine/Source/Editor/VREditor/VREditorQuickMenu.cpp b/Engine/Source/Editor/VREditor/VREditorQuickMenu.cpp
--- a/Engine/Source/Editor/VREditor/VREditorQuickMenu.cpp
+++ b/Engine/Source/Editor/VREditor/VREditorQuickMenu.cpp
@@ -54,9 +54,11 @@ void UVREditorQuickMenu::OnPasteButtonClicked()
 }
 
 
-void UVREditorQuickMenu::OnTranslationSnapButtonClicked( const bool bIsChecked )
+bool UVREditorQuickMenu::OnTranslationSnapButtonClicked( const bool bIsChecked )
 {
 	GUnrealEd->Exec( GetWorld(), *FString::Printf( TEXT( "MODE GRID=%d" ), bIsChecked ? 1 : 0 ) );
+
+	return IsTranslationSnapEnabled();
 }
 
 
@@ -66,7 +68,7 @@ bool UVREditorQuickMenu::IsTranslationSnapEnabled() const
 }
 
 
-void UVREditorQuickMenu::OnTranslationSnapSizeButtonClicked()
+FText UVREditorQuickMenu::OnTranslationSnapSizeButtonClicked()
 {
 	const TArray<float>& GridSizes = GEditor->GetCurrentPositionGridArray();
 	const int32 CurrentGridSize = GetDefault<ULevelEditorViewportSettings>()->CurrentPosGridSize;
@@ -78,6 +80,8 @@ void UVREditorQuickMenu::OnTranslationSnapSizeButtonClicked()
 	}
 
 	GEditor->SetGridSize( NewGridSize );
+
+	return GetTranslationSnapSizeText();
 }
 
 
@@ -87,7 +91,7 @@ FText UVREditorQuickMenu::GetTranslationSnapSizeText() const
 }
 
 
-void UVREditorQuickMenu::OnRotationSnapSizeButtonClicked()
+FText UVREditorQuickMenu::OnRotationSnapSizeButtonClicked()
 {
 	const TArray<float>& GridSizes = GEditor->GetCurrentRotationGridArray();
 	const int32 CurrentGridSize = GetDefault<ULevelEditorViewportSettings>()->CurrentRotGridSize;
@@ -100,6 +104,8 @@ void UVREditorQuickMenu::OnRotationSnapSizeButtonClicked()
 	}
 
 	GEditor->SetRotGridSize( NewGridSize, CurrentGridMode );
+
+	return GetRotationSnapSizeText();
 }
 
 
@@ -109,7 +115,7 @@ FText UVREditorQuickMenu::GetRotationSnapSizeText() const
 }
 
 
-void UVREditorQuickMenu::OnScaleSnapSizeButtonClicked()
+FText UVREditorQuickMenu::OnScaleSnapSizeButtonClicked()
 {
 	const ULevelEditorViewportSettings* ViewportSettings = GetMutableDefault<ULevelEditorViewportSettings>();
 	const int32 NumGridSizes = ViewportSettings->ScalingGridSizes.Num();
@@ -123,6 +129,8 @@ void UVREditorQuickMenu::OnScaleSnapSizeButtonClicked()
 	}
 
 	GEditor->SetScaleGridSize( NewGridSize );
+
+	return GetScaleSnapSizeText();
 }
 
 
@@ -132,9 +140,11 @@ FText UVREditorQuickMenu::GetScaleSnapSizeText() const
 }
 
 
-void UVREditorQuickMenu::OnRotationSnapButtonClicked( const bool bIsChecked )
+bool UVREditorQuickMenu::OnRotationSnapButtonClicked( const bool bIsChecked )
 {
 	GUnrealEd->Exec( GetWorld(), *FString::Printf( TEXT( "MODE ROTGRID=%d" ), bIsChecked ? 1 : 0 ) );
+
+	return IsRotationSnapEnabled();
 }
 
 
@@ -144,9 +154,11 @@ bool UVREditorQuickMenu::IsRotationSnapEnabled() const
 }
 
 
-void UVREditorQuickMenu::OnScaleSnapButtonClicked( const bool bIsChecked )
+bool UVREditorQuickMenu::OnScaleSnapButtonClicked( const bool bIsChecked )
 {
 	GUnrealEd->Exec( GetWorld(), *FString::Printf( TEXT( "MODE SCALEGRID=%d" ), bIsChecked ? 1 : 0 ) );
+
+	return IsScaleSnapEnabled();
 }
 
 
@@ -156,9 +168,11 @@ bool UVREditorQuickMenu::IsScaleSnapEnabled() const
 }
 
 
-void UVREditorQuickMenu::OnGizmoCoordinateSystemButtonClicked()
+FText UVREditorQuickMenu::OnGizmoCoordinateSystemButtonClicked()
 {
 	GetOwner()->GetOwner().GetOwner().CycleTransformGizmoCoordinateSpace();
+
+	return GetGizmoCoordinateSystemText();
 }
 
 
@@ -169,10 +183,11 @@ FText UVREditorQuickMenu::GetGizmoCoordinateSystemText() const
 }
 
 
-void UVREditorQuickMenu::OnGizmoModeButtonClicked()
+FText UVREditorQuickMenu::OnGizmoModeButtonClicked()
 {
 	const EGizmoHandleTypes CurrentGizmoMode = GetOwner()->GetOwner().GetOwner().GetCurrentGizmoType();
 	GetOwner()->GetOwner().GetOwner().CycleTransformGizmoHandleType();
+	return GetGizmoModeText();
 }
 
 
@@ -207,7 +222,7 @@ FText UVREditorQuickMenu::GetGizmoModeText() const
 }
 
 
-void UVREditorQuickMenu::OnSimulateButtonClicked( const bool bIsChecked )
+bool UVREditorQuickMenu::OnSimulateButtonClicked( const bool bIsChecked )
 {
 	if( bIsChecked != GEditor->bIsSimulatingInEditor )
 	{
@@ -234,6 +249,8 @@ void UVREditorQuickMenu::OnSimulateButtonClicked( const bool bIsChecked )
 			GUnrealEd->RequestToggleBetweenPIEandSIE();
 		}
 	}
+
+	return IsSimulatingEnabled();
 }
 
 
@@ -243,9 +260,11 @@ bool UVREditorQuickMenu::IsSimulatingEnabled() const
 }
 
 
-void UVREditorQuickMenu::OnContentBrowserButtonClicked( const bool bIsChecked )
+bool UVREditorQuickMenu::OnContentBrowserButtonClicked( const bool bIsChecked )
 {
 	GetOwner()->GetOwner().TogglePanelVisibility( FVREditorUISystem::EEditorUIPanel::ContentBrowser );
+
+	return IsContentBrowserVisible();
 }
 
 
@@ -255,9 +274,11 @@ bool UVREditorQuickMenu::IsContentBrowserVisible() const
 }
 
 
-void UVREditorQuickMenu::OnWorldOutlinerButtonClicked( const bool bIsChecked )
+bool UVREditorQuickMenu::OnWorldOutlinerButtonClicked( const bool bIsChecked )
 {
 	GetOwner()->GetOwner().TogglePanelVisibility( FVREditorUISystem::EEditorUIPanel::WorldOutliner );
+
+	return IsWorldOutlinerVisible();
 }
 
 
@@ -267,9 +288,11 @@ bool UVREditorQuickMenu::IsWorldOutlinerVisible() const
 }
 
 
-void UVREditorQuickMenu::OnActorDetailsButtonClicked( const bool bIsChecked )
+bool UVREditorQuickMenu::OnActorDetailsButtonClicked( const bool bIsChecked )
 {
 	GetOwner()->GetOwner().TogglePanelVisibility( FVREditorUISystem::EEditorUIPanel::ActorDetails );
+
+	return IsActorDetailsVisible();
 }
 
 
@@ -279,9 +302,11 @@ bool UVREditorQuickMenu::IsActorDetailsVisible() const
 }
 
 
-void UVREditorQuickMenu::OnModesButtonClicked( const bool bIsChecked )
+bool UVREditorQuickMenu::OnModesButtonClicked( const bool bIsChecked )
 {
 	GetOwner()->GetOwner().TogglePanelVisibility( FVREditorUISystem::EEditorUIPanel::Modes );
+
+	return IsModesVisible();
 }
 
 
@@ -291,7 +316,7 @@ bool UVREditorQuickMenu::IsModesVisible() const
 }
 
 
-void UVREditorQuickMenu::OnGameModeButtonClicked( const bool bIsChecked )
+bool UVREditorQuickMenu::OnGameModeButtonClicked( const bool bIsChecked )
 {
 	FEditorViewportClient* ViewportClient = GetOwner()->GetOwner().GetOwner().GetLevelViewportPossessedForVR().GetViewportClient().Get();
 	const bool bIsInGameView = ViewportClient->IsInGameView();
@@ -299,6 +324,8 @@ void UVREditorQuickMenu::OnGameModeButtonClicked( const bool bIsChecked )
 	{
 		ViewportClient->SetGameView( bIsChecked );
 	}
+
+	return IsGameModeEnabled();
 }
 
 
@@ -309,9 +336,11 @@ bool UVREditorQuickMenu::IsGameModeEnabled() const
 }
 
 
-void UVREditorQuickMenu::OnTutorialButtonClicked( const bool bIsChecked )
+bool UVREditorQuickMenu::OnTutorialButtonClicked( const bool bIsChecked )
 {
 	GetOwner()->GetOwner().TogglePanelVisibility( FVREditorUISystem::EEditorUIPanel::Tutorial );
+
+	return IsTutorialVisible();
 }
 
 bool UVREditorQuickMenu::IsTutorialVisible() const
@@ -319,10 +348,12 @@ bool UVREditorQuickMenu::IsTutorialVisible() const
 	return ( GetOwner()->GetOwner().IsShowingEditorUIPanel( FVREditorUISystem::EEditorUIPanel::Tutorial ) );
 }
 
-void UVREditorQuickMenu::OnLightButtonClicked(const bool bIsChecked)
+bool UVREditorQuickMenu::OnLightButtonClicked(const bool bIsChecked)
 {
 	const int32 HandIndexWithoutQuickMenu = GetOwner()->GetDockedTo() == AVREditorFloatingUI::EDockedTo::LeftArm ? VREditorConstants::RightHandIndex : VREditorConstants::LeftHandIndex;
 	GetOwner()->GetOwner().GetOwner().ToggleFlashlight(HandIndexWithoutQuickMenu);
+
+	return IsLightVisible();
 }
 
 bool UVREditorQuickMenu::IsLightVisible() const
@@ -354,10 +385,10 @@ void UVREditorQuickMenu::OnScreenshotButtonClicked()
 	FScreenshotRequest::RequestScreenshot(FString(), bShowUI, bUseCustomFilename);
 }
 
-void UVREditorQuickMenu::OnAssetEditorButtonClicked( const bool bIsChecked )
+bool UVREditorQuickMenu::OnAssetEditorButtonClicked( const bool bIsChecked )
 {
 	GetOwner()->GetOwner().TogglePanelVisibility( FVREditorUISystem::EEditorUIPanel::AssetEditor );
-	
+	return IsAssetEditorVisible();
 }
 
 bool UVREditorQuickMenu::IsAssetEditorVisible() const
@@ -365,10 +396,11 @@ bool UVREditorQuickMenu::IsAssetEditorVisible() const
 	return ( GetOwner()->GetOwner().IsShowingEditorUIPanel( FVREditorUISystem::EEditorUIPanel::AssetEditor ) );
 }
 
-void UVREditorQuickMenu::OnWorldSettingsButtonClicked(const bool bIsChecked)
+bool UVREditorQuickMenu::OnWorldSettingsButtonClicked(const bool bIsChecked)
 {
 	GetOwner()->GetOwner().TogglePanelVisibility(FVREditorUISystem::EEditorUIPanel::WorldSettings);
 
+	return IsWorldSettingsVisible();
 }
 
 bool UVREditorQuickMenu::IsWorldSettingsVisible() const
diff --git a/Engine/Source/Editor/VREditor/VREditorQuickMenu.h b/Engine/Source/Editor/VREditor/VREditorQuickMenu.h
--- a/Engine/Source/Editor/VREditor/VREditorQuickMenu.h
+++ b/Engine/Source/Editor/VREditor/VREditorQuickMenu.h
@@ -43,97 +43,97 @@ public:
 	void OnPasteButtonClicked();
 
 	UFUNCTION( BlueprintCallable, Category = "Default" )
-	void OnTranslationSnapButtonClicked( const bool bIsChecked );
+	bool OnTranslationSnapButtonClicked( const bool bIsChecked );
 
 	UFUNCTION( BlueprintPure, Category = "Default" )
 	bool IsTranslationSnapEnabled() const;
 
 	UFUNCTION( BlueprintCallable, Category = "Default" )
-	void OnTranslationSnapSizeButtonClicked();
+	FText OnTranslationSnapSizeButtonClicked();
 
 	UFUNCTION( BlueprintPure, Category = "Default" )
 	FText GetTranslationSnapSizeText() const;
 
 	UFUNCTION( BlueprintCallable, Category = "Default" )
-	void OnRotationSnapButtonClicked( const bool bIsChecked );
+	bool OnRotationSnapButtonClicked( const bool bIsChecked );
 
 	UFUNCTION( BlueprintPure, Category = "Default" )
 	bool IsRotationSnapEnabled() const;
 
 	UFUNCTION( BlueprintCallable, Category = "Default" )
-	void OnRotationSnapSizeButtonClicked();
+	FText OnRotationSnapSizeButtonClicked();
 
 	UFUNCTION( BlueprintPure, Category = "Default" )
 	FText GetRotationSnapSizeText() const;
 
 	UFUNCTION( BlueprintCallable, Category = "Default" )
-	void OnScaleSnapButtonClicked( const bool bIsChecked );
+	bool OnScaleSnapButtonClicked( const bool bIsChecked );
 
 	UFUNCTION( BlueprintPure, Category = "Default" )
 	bool IsScaleSnapEnabled() const;
 
 	UFUNCTION( BlueprintCallable, Category = "Default" )
-	void OnScaleSnapSizeButtonClicked();
+	FText OnScaleSnapSizeButtonClicked();
 
 	UFUNCTION( BlueprintPure, Category = "Default" )
 	FText GetScaleSnapSizeText() const;
 
 	UFUNCTION( BlueprintCallable, Category = "Default" )
-	void OnGizmoCoordinateSystemButtonClicked();
+	FText OnGizmoCoordinateSystemButtonClicked();
 
 	UFUNCTION( BlueprintPure, Category = "Default" )
 	FText GetGizmoCoordinateSystemText() const;
 
 	UFUNCTION( BlueprintCallable, Category = "Default" )
-	void OnGizmoModeButtonClicked();
+	FText OnGizmoModeButtonClicked();
 
 	UFUNCTION( BlueprintPure, Category = "Default" )
 	FText GetGizmoModeText() const;
 
 	UFUNCTION( BlueprintCallable, Category = "Default" )
-	void OnSimulateButtonClicked( const bool bIsChecked );
+	bool OnSimulateButtonClicked( const bool bIsChecked );
 
 	UFUNCTION( BlueprintPure, Category = "Default" )
 	bool IsSimulatingEnabled() const;
 
 	UFUNCTION( BlueprintCallable, Category = "Default" )
-	void OnContentBrowserButtonClicked( const bool bIsChecked );
+	bool OnContentBrowserButtonClicked( const bool bIsChecked );
 
 	UFUNCTION( BlueprintPure, Category = "Default" )
 	bool IsContentBrowserVisible() const;
 
 	UFUNCTION( BlueprintCallable, Category = "Default" )
-	void OnWorldOutlinerButtonClicked( const bool bIsChecked );
+	bool OnWorldOutlinerButtonClicked( const bool bIsChecked );
 
 	UFUNCTION( BlueprintPure, Category = "Default" )
 	bool IsWorldOutlinerVisible() const;
 
 	UFUNCTION( BlueprintCallable, Category = "Default" )
-	void OnActorDetailsButtonClicked( const bool bIsChecked );
+	bool OnActorDetailsButtonClicked( const bool bIsChecked );
 
 	UFUNCTION( BlueprintPure, Category = "Default" )
 	bool IsActorDetailsVisible() const;
 
 	UFUNCTION( BlueprintCallable, Category = "Default" )
-	void OnModesButtonClicked( const bool bIsChecked );
+	bool OnModesButtonClicked( const bool bIsChecked );
 
 	UFUNCTION( BlueprintPure, Category = "Default" )
 	bool IsModesVisible() const;
 
 	UFUNCTION( BlueprintCallable, Category = "Default" )
-	void OnGameModeButtonClicked( const bool bIsChecked );
+	bool OnGameModeButtonClicked( const bool bIsChecked );
 
 	UFUNCTION( BlueprintPure, Category = "Default" )
 	bool IsGameModeEnabled() const;
 
 	UFUNCTION( BlueprintCallable, Category = "Default" )
-	void OnTutorialButtonClicked( const bool bIsChecked );
+	bool OnTutorialButtonClicked( const bool bIsChecked );
 
 	UFUNCTION( BlueprintPure, Category = "Default" )
 	bool IsTutorialVisible() const;
 
 	UFUNCTION(BlueprintCallable, Category = "Default")
-	void OnLightButtonClicked(const bool bIsChecked);
+	bool OnLightButtonClicked(const bool bIsChecked);
 
 	UFUNCTION(BlueprintPure, Category = "Default")
 	bool IsLightVisible() const;
@@ -142,13 +142,13 @@ public:
 	void OnScreenshotButtonClicked();
 
 	UFUNCTION( BlueprintCallable, Category = "Default" )
-	void OnAssetEditorButtonClicked( const bool bIsChecked );
+	bool OnAssetEditorButtonClicked( const bool bIsChecked );
 
 	UFUNCTION( BlueprintPure, Category = "Default" )
 	bool IsAssetEditorVisible() const;
 
 	UFUNCTION(BlueprintCallable, Category = "Default")
-	void OnWorldSettingsButtonClicked(const bool bIsChecked);
+	bool OnWorldSettingsButtonClicked(const bool bIsChecked);
 
 	UFUNCTION(BlueprintPure, Category = "Default")
 	bool IsWorldSettingsVisible() const;
