VR Editor: Changing doubleclick event to be set on pressed instead of released. UETOOL-878

diff --git a/Engine/Source/Editor/VREditor/VREditorUISystem.cpp b/Engine/Source/Editor/VREditor/VREditorUISystem.cpp
--- a/Engine/Source/Editor/VREditor/VREditorUISystem.cpp
+++ b/Engine/Source/Editor/VREditor/VREditorUISystem.cpp
@@ -264,11 +264,23 @@ void FVREditorUISystem::OnVRAction( FEditorViewportClient& ViewportClient, const
 									FReply Reply = FReply::Unhandled();
 									if( Event == IE_Pressed )
 									{
-										Reply = FSlateApplication::Get().RoutePointerDownEvent( WidgetPathUnderFinger, PointerEvent );
+										const double CurrentTime = FPlatformTime::Seconds();
+										if (CurrentTime - Hand.LastClickPressTime <= VREd::DoubleClickTime->GetFloat())
+										{
+											// Trigger a double click event!
+											Reply = FSlateApplication::Get().RoutePointerDoubleClickEvent(WidgetPathUnderFinger, PointerEvent);
+										}
+										else
+										{
+											Reply = FSlateApplication::Get().RoutePointerDownEvent(WidgetPathUnderFinger, PointerEvent);
+										}
+										
 										Hand.bIsClickingOnUI = true;
 										Hand.bIsRightClickingOnUI = bIsRightClicking;
 										bIsInputCaptured = true;
 
+										Hand.LastClickPressTime = CurrentTime;
+
 										// Play a haptic effect on press
 										const float Strength = VREd::UIPressHapticFeedbackStrength->GetFloat();
 										Owner.PlayHapticEffect(
@@ -278,15 +290,6 @@ void FVREditorUISystem::OnVRAction( FEditorViewportClient& ViewportClient, const
 									else if( Event == IE_Released )
 									{
 										Reply = FSlateApplication::Get().RoutePointerUpEvent( WidgetPathUnderFinger, PointerEvent );
-
-										const double CurrentTime = FPlatformTime::Seconds();
-										if( CurrentTime - Hand.LastClickReleaseTime <= VREd::DoubleClickTime->GetFloat() )
-										{
-											// Trigger a double click event!
-											Reply = FSlateApplication::Get().RoutePointerDoubleClickEvent( WidgetPathUnderFinger, PointerEvent );
-										}
-
-										Hand.LastClickReleaseTime = CurrentTime;
 									}
 								}
 							}
diff --git a/Engine/Source/Editor/VREditor/VirtualHand.h b/Engine/Source/Editor/VREditor/VirtualHand.h
--- a/Engine/Source/Editor/VREditor/VirtualHand.h
+++ b/Engine/Source/Editor/VREditor/VirtualHand.h
@@ -147,8 +147,8 @@ struct FVirtualHand
 	/** When bIsClickingOnUI is true, this will be true if we're "right clicking".  That is, the Modifier key was held down at the time that the user clicked */
 	bool bIsRightClickingOnUI;
 
-	/** Last real time that we released the 'SelectAndMove' button on UI.  This is used to detect double-clicks. */
-	double LastClickReleaseTime;
+	/** Last real time that we pressed the 'SelectAndMove' button on UI.  This is used to detect double-clicks. */
+	double LastClickPressTime;
 
 	/** The last real time we played a haptic effect.  This is used to turn off haptic effects shortly after they're triggered. */
 	double LastHapticTime;
@@ -301,7 +301,7 @@ struct FVirtualHand
 		ClickingOnComponent = nullptr;
 		bIsClickingOnUI = false;
 		bIsRightClickingOnUI = false;
-		LastClickReleaseTime = 0.0;
+		LastClickPressTime = 0.0;
 		LastHapticTime = 0.0;
 		
 		bHasUIInFront = false;
