VR Editor:

diff --git a/Engine/Source/Editor/VREditor/VREditorInteractor.cpp b/Engine/Source/Editor/VREditor/VREditorInteractor.cpp
new file mode 100644
--- /dev/null
+++ b/Engine/Source/Editor/VREditor/VREditorInteractor.cpp
@@ -0,0 +1,1133 @@
+// Copyright 1998-2015 Epic Games, Inc. All Rights Reserved.
+
+#include "VREditorModule.h"
+#include "VREditorInteractor.h"
+#include "VREditorWorldInteraction.h"
+
+#include "VREditorMode.h"
+#include "VREditorUISystem.h"
+#include "VREditorFloatingText.h"
+#include "MotionControllerComponent.h"
+
+#include "Interactables/VREditorButton.h"
+#include "VREditorDockableWindow.h"
+
+namespace VREd
+{
+	static FAutoConsoleVariable HelpLabelFadeDuration( TEXT( "VREd.HelpLabelFadeDuration" ), 0.4f, TEXT( "Duration to fade controller help labels in and out" ) );
+	static FAutoConsoleVariable HelpLabelFadeDistance( TEXT( "VREd.HelpLabelFadeDistance" ), 30.0f, TEXT( "Distance at which controller help labels should appear (in cm)" ) );
+	
+	static FAutoConsoleVariable LaserPointerRadius( TEXT( "VREd.LaserPointerRadius" ), 0.5f, TEXT( "Radius of the laser pointer line" ) );
+	static FAutoConsoleVariable LaserPointerHoverBallRadius( TEXT( "VREd.LaserPointerHoverBallRadius" ), 1.5f, TEXT( "Radius of the visual cue for a hovered object along the laser pointer ray" ) );
+	static FAutoConsoleVariable LaserPointerLightPullBackDistance( TEXT( "VREd.LaserPointerLightPullBackDistance" ), 2.5f, TEXT( "How far to pull back our little hover light from the impact surface" ) );
+	static FAutoConsoleVariable LaserPointerLightRadius( TEXT( "VREd.LaserPointLightRadius" ), 20.0f, TEXT( "How big our hover light is" ) );
+	static FAutoConsoleVariable OculusLaserPointerStartOffset( TEXT( "VREd.OculusLaserPointerStartOffset" ), 2.8f, TEXT( "How far to offset the start of the laser pointer to avoid overlapping the hand mesh geometry (Oculus)" ) );
+	static FAutoConsoleVariable ViveLaserPointerStartOffset( TEXT( "VREd.ViveLaserPointerStartOffset" ), 1.25f /* 8.5f */, TEXT( "How far to offset the start of the laser pointer to avoid overlapping the hand mesh geometry (Vive)" ) );
+
+	static FAutoConsoleVariable TrackpadAbsoluteDragSpeed( TEXT( "VREd.TrackpadAbsoluteDragSpeed" ), 40.0f, TEXT( "How fast objects move toward or away when you drag on the touchpad while carrying them" ) );
+	static FAutoConsoleVariable TrackpadRelativeDragSpeed( TEXT( "VREd.TrackpadRelativeDragSpeed" ), 8.0f, TEXT( "How fast objects move toward or away when you hold a direction on an analog stick while carrying them" ) );
+	static FAutoConsoleVariable InvertTrackpadVertical( TEXT( "VREd.InvertTrackpadVertical" ), 1, TEXT( "Toggles inverting the touch pad vertical axis" ) );
+	static FAutoConsoleVariable MinVelocityForInertia( TEXT( "VREd.MinVelocityForInertia" ), 1.0f, TEXT( "Minimum velocity (in cm/frame in unscaled room space) before inertia will kick in when releasing objects (or the world)" ) );
+	static FAutoConsoleVariable MinJoystickOffsetBeforeRadialMenu( TEXT( "VREd.MinJoystickOffsetBeforeRadialMenu" ), 0.15f, TEXT( "Toggles inverting the touch pad vertical axis" ) );
+	static FAutoConsoleVariable TriggerLightlyPressedLockTime( TEXT( "VREd.TriggerLightlyPressedLockTime" ), 0.15f, TEXT( "If the trigger remains lightly pressed for longer than this, we'll continue to treat it as a light press in some cases" ) );
+}
+
+namespace VREditorKeyNames
+{
+	// @todo vreditor input: Ideally these would not be needed, but SteamVR fires off it's "trigger pressed" event
+	// well before the trigger is fully down (*click*)
+	static const FName MotionController_Left_FullyPressedTriggerAxis( "MotionController_Left_FullyPressedTriggerAxis" );
+	static const FName MotionController_Right_FullyPressedTriggerAxis( "MotionController_Right_FullyPressedTriggerAxis" );
+
+	static const FName MotionController_Left_LightlyPressedTriggerAxis( "MotionController_Left_LightlyPressedTriggerAxis" );
+	static const FName MotionController_Right_LightlyPressedTriggerAxis( "MotionController_Right_LightlyPressedTriggerAxis" );
+}
+
+namespace SteamVRControllerKeyNames
+{
+	static const FGamepadKeyNames::Type Touch0( "Steam_Touch_0" );
+	static const FGamepadKeyNames::Type Touch1( "Steam_Touch_1" );
+}
+
+#define LOCTEXT_NAMESPACE "VREditor"
+
+UVREditorInteractor::UVREditorInteractor( const FObjectInitializer& Initializer ) :
+	UViewportInteractorMotionController( Initializer ),
+	VRMode( nullptr ),
+	bHasUIInFront( false ),
+	bHasUIOnForearm( false ),
+	bIsClickingOnUI( false ),
+	bIsRightClickingOnUI( false ),
+	bIsHoveringOverUI( false ),
+	UIScrollVelocity( 0.0f ),
+	bIsTouchingTrackpad( false ),
+	TrackpadPosition( FVector2D::ZeroVector ),
+	LastTrackpadPosition( FVector2D::ZeroVector ),
+	bIsTrackpadPositionValid { false, false },
+	LastTrackpadPositionUpdateTime( FTimespan::MinValue() ),
+	bIsModifierPressed( false ),
+	LastClickReleaseTime( 0.0f ),
+	bWantHelpLabels( false ),
+	HelpLabelShowOrHideStartTime( FTimespan::MinValue() )
+{
+
+}
+
+UVREditorInteractor::~UVREditorInteractor()
+{
+	Shutdown();
+}
+
+void UVREditorInteractor::Init( const EControllerHand InHandSide, FVREditorMode* InVRMode )
+{
+	VRMode = InVRMode;
+	ControllerHandSide = InHandSide;
+	bHaveMotionController = true;
+	KeyToActionMap.Reset();
+
+	// Setup keys
+	if ( InHandSide == EControllerHand::Left )
+	{
+		AddKeyAction( EKeys::MotionController_Left_Grip1, FViewportActionKeyInput( ViewportWorldActionTypes::WorldMovement ) );
+		AddKeyAction( UViewportInteractorMotionController::MotionController_Left_FullyPressedTriggerAxis, FViewportActionKeyInput( ViewportWorldActionTypes::SelectAndMove ) );
+		AddKeyAction( UViewportInteractorMotionController::MotionController_Left_LightlyPressedTriggerAxis, FViewportActionKeyInput( ViewportWorldActionTypes::SelectAndMove_LightlyPressed ) );
+		AddKeyAction( EKeys::MotionController_Left_Thumbstick, FViewportActionKeyInput( VRActionTypes::ConfirmRadialSelection ) );
+
+		AddKeyAction( SteamVRControllerKeyNames::Touch0, FViewportActionKeyInput( VRActionTypes::Touch ) );
+		AddKeyAction( EKeys::MotionController_Left_TriggerAxis, FViewportActionKeyInput( UViewportInteractorMotionController::TriggerAxis ) );
+		AddKeyAction( EKeys::MotionController_Left_Thumbstick_X, FViewportActionKeyInput( UViewportInteractorMotionController::TrackpadPositionX ) );
+		AddKeyAction( EKeys::MotionController_Left_Thumbstick_Y, FViewportActionKeyInput( UViewportInteractorMotionController::TrackpadPositionY ) );	
+		AddKeyAction( EKeys::MotionController_Left_Thumbstick, FViewportActionKeyInput( VRActionTypes::ConfirmRadialSelection ) );
+
+		if ( GetVRMode().GetHMDDeviceType() == EHMDDeviceType::DT_SteamVR )
+		{
+			AddKeyAction( EKeys::MotionController_Left_Shoulder, FViewportActionKeyInput( VRActionTypes::Modifier ) );
+		}
+		else
+		{
+			AddKeyAction( EKeys::MotionController_Left_FaceButton1, FViewportActionKeyInput( VRActionTypes::Modifier ) );
+		}
+	}
+	else if ( InHandSide == EControllerHand::Right )
+	{
+		AddKeyAction( EKeys::MotionController_Right_Grip1, FViewportActionKeyInput( ViewportWorldActionTypes::WorldMovement ) );
+		AddKeyAction( UViewportInteractorMotionController::MotionController_Right_FullyPressedTriggerAxis, FViewportActionKeyInput( ViewportWorldActionTypes::SelectAndMove ) );
+		AddKeyAction( UViewportInteractorMotionController::MotionController_Right_LightlyPressedTriggerAxis, FViewportActionKeyInput( ViewportWorldActionTypes::SelectAndMove_LightlyPressed ) );
+		AddKeyAction( EKeys::MotionController_Right_Thumbstick, FViewportActionKeyInput( VRActionTypes::ConfirmRadialSelection ) );
+
+		AddKeyAction( SteamVRControllerKeyNames::Touch1, FViewportActionKeyInput( VRActionTypes::Touch ) );
+		AddKeyAction( EKeys::MotionController_Right_TriggerAxis, FViewportActionKeyInput( UViewportInteractorMotionController::TriggerAxis ) );
+		AddKeyAction( EKeys::MotionController_Right_Thumbstick_X, FViewportActionKeyInput( UViewportInteractorMotionController::TrackpadPositionX ) );
+		AddKeyAction( EKeys::MotionController_Right_Thumbstick_Y, FViewportActionKeyInput( UViewportInteractorMotionController::TrackpadPositionY ) );
+		AddKeyAction( EKeys::MotionController_Right_Thumbstick, FViewportActionKeyInput( VRActionTypes::ConfirmRadialSelection ) );
+
+		if ( GetVRMode().GetHMDDeviceType() == EHMDDeviceType::DT_SteamVR )
+		{
+			AddKeyAction( EKeys::MotionController_Right_Shoulder, FViewportActionKeyInput( VRActionTypes::Modifier ) );
+		}
+		else
+		{
+			AddKeyAction( EKeys::MotionController_Right_FaceButton1, FViewportActionKeyInput( VRActionTypes::Modifier ) );
+		}
+	}
+}
+
+void UVREditorInteractor::SetupComponent( AActor* OwningActor )
+{
+	// Setup a motion controller component.  This allows us to take advantage of late frame updates, so
+	// our motion controllers won't lag behind the HMD
+	{
+		MotionControllerComponent = NewObject<UMotionControllerComponent>( OwningActor );
+		OwningActor->AddOwnedComponent( MotionControllerComponent );
+		MotionControllerComponent->SetupAttachment( OwningActor->GetRootComponent() );
+		MotionControllerComponent->RegisterComponent();
+
+		MotionControllerComponent->SetMobility( EComponentMobility::Movable );
+		MotionControllerComponent->SetCollisionEnabled( ECollisionEnabled::NoCollision );
+
+		MotionControllerComponent->Hand = ControllerHandSide;
+
+		MotionControllerComponent->bDisableLowLatencyUpdate = false;
+	}
+
+	// Hand mesh
+	{
+		HandMeshComponent = NewObject<UStaticMeshComponent>( OwningActor );
+		OwningActor->AddOwnedComponent( HandMeshComponent );
+		HandMeshComponent->SetupAttachment( MotionControllerComponent );
+		HandMeshComponent->RegisterComponent();
+
+		// @todo vreditor extensibility: We need this to be able to be overridden externally, or simply based on the HMD name (but allowing external folders)
+		FString MeshName;
+		FString MaterialName;
+		if ( GetVRMode().GetHMDDeviceType() == EHMDDeviceType::DT_SteamVR )
+		{
+			MeshName = TEXT( "/Engine/VREditor/Devices/Vive/VivePreControllerMesh" );
+			MaterialName = TEXT( "/Engine/VREditor/Devices/Vive/VivePreControllerMaterial_Inst" );
+		}
+		else
+		{
+			MeshName = TEXT( "/Engine/VREditor/Devices/Oculus/OculusControllerMesh" );
+			MaterialName = TEXT( "/Engine/VREditor/Devices/Oculus/OculusControllerMaterial_Inst" );
+		}
+
+		UStaticMesh* HandMesh = LoadObject<UStaticMesh>( nullptr, *MeshName );
+		check( HandMesh != nullptr );
+
+		HandMeshComponent->SetStaticMesh( HandMesh );
+		HandMeshComponent->SetMobility( EComponentMobility::Movable );
+		HandMeshComponent->SetCollisionEnabled( ECollisionEnabled::NoCollision );
+
+		UMaterialInstance* HandMeshMaterialInst = LoadObject<UMaterialInstance>( nullptr, *MaterialName );
+		check( HandMeshMaterialInst != nullptr );
+		HandMeshMID = UMaterialInstanceDynamic::Create( HandMeshMaterialInst, GetTransientPackage() );
+		check( HandMeshMID != nullptr );
+		HandMeshComponent->SetMaterial( 0, HandMeshMID );
+	}
+
+	// Laser pointer
+	{
+		LaserPointerMeshComponent = NewObject<UStaticMeshComponent>( OwningActor );
+		OwningActor->AddOwnedComponent( LaserPointerMeshComponent );
+		LaserPointerMeshComponent->SetupAttachment( MotionControllerComponent );
+		LaserPointerMeshComponent->RegisterComponent();
+
+		UStaticMesh* LaserPointerMesh = LoadObject<UStaticMesh>( nullptr, TEXT( "/Engine/VREditor/LaserPointer/LaserPointerMesh" ) );
+		check( LaserPointerMesh != nullptr );
+
+		LaserPointerMeshComponent->SetStaticMesh( LaserPointerMesh );
+		LaserPointerMeshComponent->SetMobility( EComponentMobility::Movable );
+		LaserPointerMeshComponent->SetCollisionEnabled( ECollisionEnabled::NoCollision );
+
+		UMaterialInstance* LaserPointerMaterialInst = LoadObject<UMaterialInstance>( nullptr, TEXT( "/Engine/VREditor/LaserPointer/LaserPointerMaterialInst" ) );
+		check( LaserPointerMaterialInst != nullptr );
+		LaserPointerMID = UMaterialInstanceDynamic::Create( LaserPointerMaterialInst, GetTransientPackage() );
+		check( LaserPointerMID != nullptr );
+		LaserPointerMeshComponent->SetMaterial( 0, LaserPointerMID );
+
+		UMaterialInstance* TranslucentLaserPointerMaterialInst = LoadObject<UMaterialInstance>( nullptr, TEXT( "/Engine/VREditor/LaserPointer/TranslucentLaserPointerMaterialInst" ) );
+		check( TranslucentLaserPointerMaterialInst != nullptr );
+		TranslucentLaserPointerMID = UMaterialInstanceDynamic::Create( TranslucentLaserPointerMaterialInst, GetTransientPackage() );
+		check( TranslucentLaserPointerMID != nullptr );
+		LaserPointerMeshComponent->SetMaterial( 1, TranslucentLaserPointerMID );
+	}
+
+	// Hover cue for laser pointer
+	{
+		HoverMeshComponent = NewObject<UStaticMeshComponent>( OwningActor );
+		OwningActor->AddOwnedComponent( HoverMeshComponent );
+		HoverMeshComponent->SetupAttachment( OwningActor->GetRootComponent() );
+		HoverMeshComponent->RegisterComponent();
+
+		UStaticMesh* HoverMesh = LoadObject<UStaticMesh>( nullptr, TEXT( "/Engine/VREditor/LaserPointer/HoverMesh" ) );
+		check( HoverMesh != nullptr );
+		HoverMeshComponent->SetStaticMesh( HoverMesh );
+		HoverMeshComponent->SetMobility( EComponentMobility::Movable );
+		HoverMeshComponent->SetCollisionEnabled( ECollisionEnabled::NoCollision );
+
+		// Add a light!
+		{
+			HoverPointLightComponent = NewObject<UPointLightComponent>( OwningActor );
+			OwningActor->AddOwnedComponent( HoverPointLightComponent );
+			HoverPointLightComponent->SetupAttachment( HoverMeshComponent );
+			HoverPointLightComponent->RegisterComponent();
+
+			HoverPointLightComponent->SetLightColor( FLinearColor::Red );
+			//Hand.HoverPointLightComponent->SetLightColor( FLinearColor( 0.0f, 1.0f, 0.2f, 1.0f ) );
+			HoverPointLightComponent->SetIntensity( 30.0f );	// @todo: VREditor tweak
+			HoverPointLightComponent->SetMobility( EComponentMobility::Movable );
+			HoverPointLightComponent->SetAttenuationRadius( VREd::LaserPointerLightRadius->GetFloat() );
+			HoverPointLightComponent->bUseInverseSquaredFalloff = false;
+			HoverPointLightComponent->SetCastShadows( false );
+		}
+	}
+}
+
+void UVREditorInteractor::Shutdown()
+{
+	for ( auto& KeyAndValue : HelpLabels )
+	{
+		AFloatingText* FloatingText = KeyAndValue.Value;
+		GetVRMode().DestroyTransientActor( FloatingText );
+	}
+
+	HelpLabels.Empty();	
+
+	VRMode = nullptr;
+
+	Super::Shutdown();
+}
+
+void UVREditorInteractor::Tick( const float DeltaTime )
+{
+	Super::Tick( DeltaTime );
+	
+	UpdateRadialMenuInput( DeltaTime );
+
+	{
+		const float WorldScaleFactor = WorldInteraction->GetWorldScaleFactor();
+
+		// We don't bother drawing hand meshes if we're in "forced VR mode" (because they'll just be on top of the camera).
+		// Also, don't bother drawing hands if we're not currently tracking them.
+		if ( GetVRMode().IsActuallyUsingVR() && bHaveMotionController )
+		{
+			HandMeshComponent->SetVisibility( true );
+		}
+		else
+		{
+			HandMeshComponent->SetVisibility( false );
+		}
+
+		FVector LaserPointerStart, LaserPointerEnd;
+		if ( GetLaserPointer( /* Out */ LaserPointerStart, /* Out */ LaserPointerEnd ) )
+		{
+			// Only show the laser if we're actually in VR
+			LaserPointerMeshComponent->SetVisibility( GetVRMode().IsActuallyUsingVR() );
+
+			// NOTE: We don't need to set the laser pointer location and rotation, as the MotionControllerComponent will do
+			// that later in the frame.  
+
+			// If we're actively dragging something around, then we'll crop the laser length to the hover impact
+			// point.  Otherwise, we always want the laser to protrude through hovered objects, so that you can
+			// interact with translucent gizmo handles that are occluded by geometry
+			FVector LaserPointerImpactPoint = LaserPointerEnd;
+			if ( GetInteractorData().bIsHovering && ( GetInteractorData().DraggingMode != EViewportInteractionDraggingMode::Nothing || bIsHoveringOverUI ) )
+			{
+				LaserPointerImpactPoint = GetInteractorData().HoverLocation;
+			}
+
+			// Apply rotation offset to the laser direction
+			const float LaserPointerRotationOffset = 0.0f;
+			LaserPointerMeshComponent->SetRelativeRotation( FRotator( LaserPointerRotationOffset, 0.0f, 0.0f ) );
+
+			const FVector LaserPointerDirection = ( LaserPointerImpactPoint - LaserPointerStart ).GetSafeNormal();
+
+			// Offset the beginning of the laser pointer a bit, so that it doesn't overlap the hand mesh
+			const float LaserPointerStartOffset =
+				WorldScaleFactor *
+				( GetVRMode().GetHMDDeviceType() == EHMDDeviceType::DT_OculusRift ? VREd::OculusLaserPointerStartOffset->GetFloat() : VREd::ViveLaserPointerStartOffset->GetFloat() );
+
+			const float LaserPointerLength = FMath::Max( 0.000001f, ( LaserPointerImpactPoint - LaserPointerStart ).Size() - LaserPointerStartOffset );
+
+			LaserPointerMeshComponent->SetRelativeLocation( FRotator( LaserPointerRotationOffset, 0.0f, 0.0f ).RotateVector( FVector( LaserPointerStartOffset, 0.0f, 0.0f ) ) );
+
+			// The laser pointer needs to stay the same size relative to our tracking space, so we inverse compensate for world to meters scale here
+			float LaserPointerRadius = VREd::LaserPointerRadius->GetFloat() * WorldScaleFactor;
+			float HoverMeshRadius = VREd::LaserPointerHoverBallRadius->GetFloat() * WorldScaleFactor;
+
+			// If we're hovering over something really close to the camera, go ahead and shrink the effect
+			// @todo vreditor: Can we make this actually just sized based on distance automatically?  The beam and impact point are basically a cone.
+			if ( bIsHoveringOverUI )
+			{
+				LaserPointerRadius *= 0.35f;	// @todo vreditor tweak
+				HoverMeshRadius *= 0.35f;	// @todo vreditor tweak
+			}
+
+			LaserPointerMeshComponent->SetRelativeScale3D( FVector( LaserPointerLength, LaserPointerRadius * 2.0f, LaserPointerRadius * 2.0f ) );
+
+			if ( IsHovering() )
+			{
+				// The hover effect needs to stay the same size relative to our tracking space, so we inverse compensate for world to meters scale here
+				HoverMeshComponent->SetRelativeScale3D( FVector( HoverMeshRadius * 2.0f ) );
+				HoverMeshComponent->SetVisibility( true );
+				HoverMeshComponent->SetWorldLocation( GetHoverLocation() );
+
+				// Show the light too, unless it's on top of UI.  It looks too distracting on top of UI.
+				HoverPointLightComponent->SetVisibility( !bIsHoveringOverUI );
+
+				// Update radius for world scaling
+				HoverPointLightComponent->SetAttenuationRadius( VREd::LaserPointerLightRadius->GetFloat() * WorldScaleFactor );
+
+				// Pull hover light back a bit from the end of the ray
+				const float PullBackAmount = VREd::LaserPointerLightPullBackDistance->GetFloat() * WorldInteraction->GetWorldScaleFactor();
+				HoverPointLightComponent->SetWorldLocation( GetHoverLocation() - PullBackAmount * LaserPointerDirection );
+			}
+			else
+			{
+				HoverMeshComponent->SetVisibility( false );
+				HoverPointLightComponent->SetVisibility( false );
+			}
+		}
+		else
+		{
+			LaserPointerMeshComponent->SetVisibility( false );
+			HoverMeshComponent->SetVisibility( false );
+			HoverPointLightComponent->SetVisibility( false );
+		}
+	}
+
+	// Updating laser colors for both hands
+	{
+		FVREditorMode::EColors ResultColorType = FVREditorMode::EColors::DefaultColor;
+		float CrawlSpeed = 0.0f;
+		float CrawlFade = 0.0f;
+			
+		const EViewportInteractionDraggingMode DraggingMode = GetDraggingMode();
+
+		//@todo VREditor
+		const bool bIsDraggingWorldWithTwoHands = false;
+			( DraggingMode == EViewportInteractionDraggingMode::AssistingDrag && GetOtherInteractor()->GetDraggingMode() == EViewportInteractionDraggingMode::World ) ||
+			( DraggingMode == EViewportInteractionDraggingMode::World && GetOtherInteractor()->GetDraggingMode() == EViewportInteractionDraggingMode::AssistingDrag );
+
+		if ( bIsDraggingWorldWithTwoHands )
+		{
+			ResultColorType = FVREditorMode::EColors::WorldDraggingColor_TwoHands;
+		}
+		else if ( DraggingMode == EViewportInteractionDraggingMode::World )
+		{
+			//@todo VREditor find a way to make the teleporter do it
+// 			if ( GetVRMode().GetWorldInteraction().IsTeleporting() ) //@todo VREditor: Make teleport an action
+// 			{
+// 				ResultColorType = FVREditorMode::EColors::TeleportColor;
+// 			}
+// 			else
+// 			{
+				// We can teleport in this mode, so animate the laser a bit
+				CrawlFade = 1.0f;
+				CrawlSpeed = 5.0f;
+				ResultColorType = FVREditorMode::EColors::WorldDraggingColor_OneHand;
+//			}
+		}
+		else if ( DraggingMode == EViewportInteractionDraggingMode::ActorsAtLaserImpact ||
+			DraggingMode == EViewportInteractionDraggingMode::Material ||
+			DraggingMode == EViewportInteractionDraggingMode::ActorsFreely ||
+			DraggingMode == EViewportInteractionDraggingMode::ActorsWithGizmo ||
+			DraggingMode == EViewportInteractionDraggingMode::AssistingDrag ||
+			DraggingMode == EViewportInteractionDraggingMode::Interactable ||
+			( GetVRMode().GetUISystem().IsInteractorDraggingDockUI( this ) && GetVRMode().GetUISystem().IsDraggingDockUI() ) )
+		{
+			ResultColorType = FVREditorMode::EColors::SelectionColor;
+		}
+		
+		const FLinearColor ResultColor = GetVRMode().GetColor( ResultColorType );
+		SetLaserVisuals( ResultColor, CrawlFade, CrawlSpeed );
+	}
+
+	UpdateHelpLabels();
+}
+
+FHitResult UVREditorInteractor::GetHitResultFromLaserPointer( TArray<AActor*>* OptionalListOfIgnoredActors /*= nullptr*/, const bool bIgnoreGizmos /*= false*/,
+	TArray<UClass*>* ObjectsInFrontOfGizmo /*= nullptr */, const bool bEvenIfBlocked /*= false */, const float LaserLengthOverride /*= 0.0f */ )
+{
+	if ( OptionalListOfIgnoredActors == nullptr )
+	{
+		TArray<AActor*> IgnoredActors;
+		OptionalListOfIgnoredActors = &IgnoredActors;
+	}
+
+	OptionalListOfIgnoredActors->Add( GetVRMode().GetAvatarMeshActor() );
+	
+	// Ignore UI widgets too
+	if ( GetDraggingMode() == EViewportInteractionDraggingMode::ActorsAtLaserImpact )
+	{
+		for( TActorIterator<AVREditorFloatingUI> UIActorIt( WorldInteraction->GetViewportWorld() ); UIActorIt; ++UIActorIt )
+		{
+			OptionalListOfIgnoredActors->Add( *UIActorIt );
+		}
+	}
+
+	if ( ObjectsInFrontOfGizmo == nullptr )
+	{
+		TArray<UClass*> PriorityOverGizmoObjects;
+		ObjectsInFrontOfGizmo = &PriorityOverGizmoObjects;
+	}
+
+	ObjectsInFrontOfGizmo->Add( AVREditorButton::StaticClass() );
+	ObjectsInFrontOfGizmo->Add( AVREditorDockableWindow::StaticClass() );
+
+	return UViewportInteractor::GetHitResultFromLaserPointer( OptionalListOfIgnoredActors, bIgnoreGizmos, ObjectsInFrontOfGizmo, bEvenIfBlocked, LaserLengthOverride );
+}
+
+void UVREditorInteractor::ResetHoverState( const float DeltaTime )
+{
+	bIsHoveringOverUI = false;
+}
+
+void UVREditorInteractor::CalculateDragRay()
+{
+	const FTimespan CurrentTime = FTimespan::FromSeconds( FPlatformTime::Seconds() );
+	const float WorldScaleFactor = WorldInteraction->GetWorldScaleFactor();
+
+	// If we're dragging an object, go ahead and slide the object along the ray based on how far they slide their touch
+	// Make sure they are touching the trackpad, otherwise we get bad data
+	if ( bIsTrackpadPositionValid[ 1 ] )
+	{
+		const bool bIsAbsolute = ( GetVRMode().GetHMDDeviceType() == EHMDDeviceType::DT_SteamVR );
+		float SlideDelta = GetTrackpadSlideDelta() * WorldScaleFactor;
+
+		if ( !FMath::IsNearlyZero( SlideDelta ) )
+		{
+			InteractorData.DragRayLength += SlideDelta;
+
+			InteractorData.DragRayLengthVelocity = 0.0f;
+
+			// Don't apply inertia unless the user dragged a decent amount this frame
+			if ( bIsAbsolute && FMath::Abs( SlideDelta ) >= VREd::MinVelocityForInertia->GetFloat() * WorldScaleFactor )
+			{
+				// Don't apply inertia if our data is sort of old
+				if ( CurrentTime - LastTrackpadPositionUpdateTime <= FTimespan::FromSeconds( 1.0f / 30.0f ) )
+				{
+					InteractorData.DragRayLengthVelocity = SlideDelta;
+				}
+			}
+
+			// Don't go too far
+			if ( InteractorData.DragRayLength < 0.0f )
+			{
+				InteractorData.DragRayLength = 0.0f;
+				InteractorData.DragRayLengthVelocity = 0.0f;
+			}
+		}
+	}
+	else
+	{
+		if ( !FMath::IsNearlyZero( InteractorData.DragRayLengthVelocity ) )
+		{
+			// Apply drag ray length inertia
+			InteractorData.DragRayLength += InteractorData.DragRayLengthVelocity;
+
+			// Don't go too far!
+			if ( InteractorData.DragRayLength < 0.0f )
+			{
+				InteractorData.DragRayLength = 0.0f;
+				InteractorData.DragRayLengthVelocity = 0.0f;
+			}
+
+			// Apply damping
+			FVector RayVelocityVector( InteractorData.DragRayLengthVelocity, 0.0f, 0.0f );
+			const bool bVelocitySensitive = true;
+			WorldInteraction->ApplyVelocityDamping( RayVelocityVector, bVelocitySensitive );
+			InteractorData.DragRayLengthVelocity = RayVelocityVector.X;
+		}
+		else
+		{
+			InteractorData.DragRayLengthVelocity = 0.0f;
+		}
+	}
+}
+
+void UVREditorInteractor::OnStartDragging( UActorComponent* ClickedComponent, const FVector& HitLocation, const bool bIsPlacingActors )
+{
+	// If the user is holding down the modifier key, go ahead and duplicate the selection first.  This needs to
+	// happen before we create our transformables, because duplicating actors will change which actors are
+	// selected!  Don't do this if we're placing objects right now though.
+	if ( bIsModifierPressed && !bIsPlacingActors )
+	{
+		WorldInteraction->Duplicate();
+	}
+}
+
+float UVREditorInteractor::GetTrackpadSlideDelta()
+{
+	const bool bIsAbsolute = ( GetVRMode().GetHMDDeviceType() == EHMDDeviceType::DT_SteamVR );
+	float SlideDelta = 0.0f;
+	if ( bIsTouchingTrackpad || !bIsAbsolute )
+	{
+		if ( bIsAbsolute )
+		{
+			SlideDelta = ( ( TrackpadPosition.Y - LastTrackpadPosition.Y ) * VREd::TrackpadAbsoluteDragSpeed->GetFloat() );
+		}
+		else
+		{
+			SlideDelta = ( TrackpadPosition.Y * VREd::TrackpadRelativeDragSpeed->GetFloat() );
+		}
+	}
+
+	return SlideDelta;
+}
+
+bool UVREditorInteractor::IsHoveringOverUI() const
+{
+	return bIsHoveringOverUI;
+}
+
+void UVREditorInteractor::SetHasUIInFront( const bool bInHasUIInFront )
+{
+	bHasUIInFront = bInHasUIInFront;
+}
+
+bool UVREditorInteractor::HasUIInFront() const
+{
+	return bHasUIInFront;
+}
+
+void UVREditorInteractor::SetHasUIOnForearm( const bool bInHasUIOnForearm )
+{
+	bHasUIOnForearm = bHasUIOnForearm;
+}
+
+bool UVREditorInteractor::HasUIOnForearm() const
+{
+	return bHasUIOnForearm;
+}
+
+bool UVREditorInteractor::IsTouchingTrackpad() const
+{
+	return bIsTouchingTrackpad;
+}
+
+FVector2D UVREditorInteractor::GetTrackpadPosition() const
+{
+	return TrackpadPosition;
+}
+
+FVector2D UVREditorInteractor::GetLastTrackpadPosition() const
+{
+	return LastTrackpadPosition;
+}
+
+UWidgetComponent* UVREditorInteractor::GetHoveringOverWidgetComponent() const
+{
+	return InteractorData.HoveringOverWidgetComponent;
+}
+
+void UVREditorInteractor::SetHoveringOverWidgetComponent( UWidgetComponent* NewHoveringOverWidgetComponent )
+{
+	InteractorData.HoveringOverWidgetComponent = NewHoveringOverWidgetComponent;
+}
+
+bool UVREditorInteractor::IsTrackpadPositionValid( const int32 AxisIndex ) const
+{
+	return bIsTrackpadPositionValid[ AxisIndex ];
+}
+
+FTimespan& UVREditorInteractor::GetLastTrackpadPositionUpdateTime()
+{
+	return LastTrackpadPositionUpdateTime;
+}
+
+bool UVREditorInteractor::IsModifierPressed() const
+{
+	return bIsModifierPressed;
+}
+
+void UVREditorInteractor::SetIsClickingOnUI( const bool bInIsClickingOnUI )
+{
+	bIsClickingOnUI = bInIsClickingOnUI;
+}
+
+bool UVREditorInteractor::IsClickingOnUI() const
+{
+	return bIsClickingOnUI;
+}
+
+void UVREditorInteractor::SetIsHoveringOverUI( const bool bInIsHoveringOverUI )
+{
+	bIsHoveringOverUI = bInIsHoveringOverUI;
+}
+
+void UVREditorInteractor::SetIsRightClickingOnUI( const bool bInIsRightClickingOnUI )
+{
+	bIsRightClickingOnUI = bInIsRightClickingOnUI;
+}
+
+bool UVREditorInteractor::IsRightClickingOnUI() const
+{
+	return bIsRightClickingOnUI;
+}
+
+void UVREditorInteractor::SetLastClickReleaseTime( const double InLastClickReleaseTime )
+{
+	LastClickReleaseTime = InLastClickReleaseTime;
+}
+
+double UVREditorInteractor::GetLastClickReleaseTime() const
+{
+	return LastClickReleaseTime;
+}
+
+void UVREditorInteractor::SetUIScrollVelocity( const float InUIScrollVelocity )
+{
+	UIScrollVelocity = InUIScrollVelocity;
+}
+
+float UVREditorInteractor::GetUIScrollVelocity() const
+{
+	return UIScrollVelocity;
+}
+
+void UVREditorInteractor::HandleInputKey( FViewportActionKeyInput& Action, const FKey Key, const EInputEvent Event, bool& bOutWasHandled )
+{
+	if( !bOutWasHandled )
+	{
+		bool bShouldAbsorbEvent = false;
+
+		// If "SelectAndMove" was pressed, we need to make sure that "SelectAndMove_LightlyPressed" is no longer
+		// in effect before we start handling "SelectAndMove".
+		if ( Action.ActionType == ViewportWorldActionTypes::SelectAndMove && IsTriggerAtLeastLightlyPressed() )
+		{
+			if ( Event == IE_Pressed )
+			{
+				// How long since the trigger was lightly pressed?
+				const float TimeSinceLightlyPressed = ( float ) ( FPlatformTime::Seconds() - GetRealTimeTriggerWasLightlyPressed() );
+				if ( !bIsClickingOnUI &&	// @todo vreditor: UI clicks happen with light presses; we never want to convert to a hard press!
+					GetDraggingMode() != EViewportInteractionDraggingMode::Material &&	// @todo vreditor: Material dragging happens with light presses, don't convert to a hard press!
+					GetDraggingMode() != EViewportInteractionDraggingMode::ActorsAtLaserImpact &&	// @todo vreditor: Actor dragging happens with light presses, don't convert to a hard press!
+					( ! AllowTriggerLightPressLocking() || TimeSinceLightlyPressed < VREd::TriggerLightlyPressedLockTime->GetFloat() ) )
+				{
+					SetTriggerAtLeastLightlyPressed( false );
+					SetTriggerBeenReleasedSinceLastPress( false );
+
+					// Synthesize an input key for releasing the light press
+					// NOTE: Intentionally re-entrant call here.
+					const EInputEvent InputEvent = IE_Released;
+					const bool bWasLightReleaseHandled = UViewportInteractor::HandleInputKey( GetControllerSide() == EControllerHand::Left ? VREditorKeyNames::MotionController_Left_LightlyPressedTriggerAxis : VREditorKeyNames::MotionController_Right_LightlyPressedTriggerAxis, InputEvent );
+				}
+				else
+				{
+					// The button was held in a "lightly pressed" state for a long enough time that we should just continue to
+					// treat it like a light press.  This makes it much easier to hold it in this state when you need to!
+					bShouldAbsorbEvent = true;
+				}
+			}
+			else
+			{
+				// Absorb the release of the SelectAndMove event if we are still in a lightly pressed state
+				bShouldAbsorbEvent = true;
+			}
+		}
+
+		if ( bShouldAbsorbEvent )
+		{
+			bOutWasHandled = true;
+		}
+	}
+
+	if ( !bOutWasHandled )
+	{
+		// Update touch state
+		if ( Action.ActionType == VRActionTypes::Touch )
+		{
+			if ( Event == IE_Pressed )
+			{
+				bIsTouchingTrackpad = true;
+			}
+			else if ( Event == IE_Released )
+			{
+				bIsTouchingTrackpad = false;
+				bIsTrackpadPositionValid[ 0 ] = false;
+				bIsTrackpadPositionValid[ 1 ] = false;
+			}
+		}
+
+		// Update modifier state
+		if ( Action.ActionType == VRActionTypes::Modifier )
+		{
+			bOutWasHandled = true;
+			if ( Event == IE_Pressed )
+			{
+				bIsModifierPressed = true;
+			}
+			else if ( Event == IE_Released )
+			{
+				bIsModifierPressed = false;
+			}
+		}
+	}
+
+	ApplyButtonPressColors( Action );
+}
+
+void UVREditorInteractor::HandleInputAxis( FViewportActionKeyInput& Action, const FKey Key, const float Delta, const float DeltaTime, bool& bOutWasHandled )
+{
+	Super::HandleInputAxis( Action, Key, Delta, DeltaTime, bOutWasHandled );
+
+	if ( !bOutWasHandled )
+	{
+		if ( Action.ActionType == VRActionTypes::TrackpadPositionX )
+		{
+			LastTrackpadPosition.X = bIsTrackpadPositionValid[ 0 ] ? TrackpadPosition.X : Delta;
+			LastTrackpadPositionUpdateTime = FTimespan::FromSeconds( FPlatformTime::Seconds() );
+			TrackpadPosition.X = Delta;
+			bIsTrackpadPositionValid[ 0 ] = true;
+		}
+			
+		if ( Action.ActionType == VRActionTypes::TrackpadPositionY )
+		{
+			float DeltaAxis = Delta;
+			if ( VREd::InvertTrackpadVertical->GetInt() != 0 )
+			{
+				DeltaAxis = -DeltaAxis;	// Y axis is inverted from HMD
+			}
+
+			LastTrackpadPosition.Y = bIsTrackpadPositionValid[ 1 ] ? TrackpadPosition.Y : DeltaAxis;
+			LastTrackpadPositionUpdateTime = FTimespan::FromSeconds( FPlatformTime::Seconds() );
+			TrackpadPosition.Y = DeltaAxis;
+			bIsTrackpadPositionValid[ 1 ] = true;
+		}
+	}
+}
+
+bool UVREditorInteractor::GetIsLaserBlocked()
+{
+	return bHasUIInFront;
+}
+
+/** Changes the color of the buttons on the handmesh */
+void UVREditorInteractor::ApplyButtonPressColors( const FViewportActionKeyInput& Action )
+{
+	const float PressStrength = 10.0f;
+	const FName ActionType = Action.ActionType;
+	const EInputEvent Event = Action.Event;
+
+	// Trigger
+	if ( ActionType == ViewportWorldActionTypes::SelectAndMove|| 
+		ActionType == ViewportWorldActionTypes::SelectAndMove_LightlyPressed )
+	{
+		static FName StaticTriggerParameter( "B1" );
+		SetMotionControllerButtonPressedVisuals( Event, StaticTriggerParameter, PressStrength );
+	}
+
+	// Shoulder button
+	if ( ActionType == ViewportWorldActionTypes::WorldMovement )
+	{
+		static FName StaticShoulderParameter( "B2" );
+		SetMotionControllerButtonPressedVisuals( Event, StaticShoulderParameter, PressStrength );
+	}
+
+	// Trackpad
+	if ( ActionType == VRActionTypes::ConfirmRadialSelection )
+	{
+		static FName StaticTrackpadParameter( "B3" );
+		SetMotionControllerButtonPressedVisuals( Event, StaticTrackpadParameter, PressStrength );
+	}
+
+	// Modifier
+	if ( ActionType == VRActionTypes::Modifier )
+	{
+		static FName StaticModifierParameter( "B4" );
+		SetMotionControllerButtonPressedVisuals( Event, StaticModifierParameter, PressStrength );
+	}
+}
+
+void UVREditorInteractor::SetMotionControllerButtonPressedVisuals( const EInputEvent Event, const FName& ParameterName, const float PressStrength )
+{
+	if ( Event == IE_Pressed )
+	{
+		HandMeshMID->SetScalarParameterValue( ParameterName, PressStrength );
+	}
+	else if ( Event == IE_Released )
+	{
+		HandMeshMID->SetScalarParameterValue( ParameterName, 0.0f );
+	}
+}
+
+void UVREditorInteractor::ShowHelpForHand( const bool bShowIt )
+{
+	if ( bShowIt != bWantHelpLabels )
+	{
+		bWantHelpLabels = bShowIt;
+
+		const FTimespan CurrentTime = FTimespan::FromSeconds( FApp::GetCurrentTime() );
+		const FTimespan TimeSinceStartedFadingOut = CurrentTime - HelpLabelShowOrHideStartTime;
+		const FTimespan HelpLabelFadeDuration = FTimespan::FromSeconds( VREd::HelpLabelFadeDuration->GetFloat() );
+
+		// If we were already fading, account for that here
+		if ( TimeSinceStartedFadingOut < HelpLabelFadeDuration )
+		{
+			// We were already fading, so we'll reverse the time value so it feels continuous
+			HelpLabelShowOrHideStartTime = CurrentTime - ( HelpLabelFadeDuration - TimeSinceStartedFadingOut );
+		}
+		else
+		{
+			HelpLabelShowOrHideStartTime = FTimespan::FromSeconds( FApp::GetCurrentTime() );
+		}
+
+		if ( bShowIt && HelpLabels.Num() == 0 )
+		{
+			for ( const auto& KeyToAction : KeyToActionMap )
+			{
+				const FKey Key = KeyToAction.Key;
+				const FViewportActionKeyInput& Action = KeyToAction.Value;
+
+				UStaticMeshSocket* Socket = FindMeshSocketForKey( HandMeshComponent->StaticMesh, Key );
+				if ( Socket != nullptr )
+				{
+					FText LabelText;
+					FString ComponentName;
+					
+					if ( Action.ActionType == VRActionTypes::Modifier )
+					{
+						LabelText = LOCTEXT( "ModifierHelp", "Modifier" );
+						ComponentName = TEXT( "ModifierHelp" );
+					}
+					else if ( Action.ActionType == ViewportWorldActionTypes::WorldMovement )
+					{
+						LabelText = LOCTEXT( "WorldMovementHelp", "Move World" );
+						ComponentName = TEXT( "WorldMovementHelp" );
+					}
+					else if ( Action.ActionType == ViewportWorldActionTypes::SelectAndMove )
+					{
+						LabelText = LOCTEXT( "SelectAndMoveHelp", "Select & Move" );
+						ComponentName = TEXT( "SelectAndMoveHelp" );
+					}
+					else if ( Action.ActionType == ViewportWorldActionTypes::SelectAndMove_LightlyPressed )
+					{
+						LabelText = LOCTEXT( "SelectAndMove_LightlyPressedHelp", "Select & Move" );
+						ComponentName = TEXT( "SelectAndMove_LightlyPressedHelp" );
+					}
+					else if ( Action.ActionType == VRActionTypes::Touch )
+					{
+						LabelText = LOCTEXT( "TouchHelp", "Slide" );
+						ComponentName = TEXT( "TouchHelp" );
+					}
+					else if ( Action.ActionType == ViewportWorldActionTypes::Undo )
+					{
+						LabelText = LOCTEXT( "UndoHelp", "Undo" );
+						ComponentName = TEXT( "UndoHelp" );
+					}
+					else if ( Action.ActionType == ViewportWorldActionTypes::Redo )
+					{
+						LabelText = LOCTEXT( "RedoHelp", "Redo" );
+						ComponentName = TEXT( "RedoHelp" );
+					}
+					else if ( Action.ActionType == ViewportWorldActionTypes::Delete )
+					{
+						LabelText = LOCTEXT( "DeleteHelp", "Delete" );
+						ComponentName = TEXT( "DeleteHelp" );
+					}
+					else if ( Action.ActionType == VRActionTypes::ConfirmRadialSelection )
+					{
+						LabelText = LOCTEXT( "ConfirmRadialSelectionHelp", "Radial Menu" );
+						ComponentName = TEXT( "ConfirmRadialSelectionHelp" );
+					}
+
+					const bool bWithSceneComponent = false;	// Nope, we'll spawn our own inside AFloatingText
+					check( VRMode );
+					AFloatingText* FloatingText = GetVRMode().SpawnTransientSceneActor<AFloatingText>( ComponentName, bWithSceneComponent );
+					FloatingText->SetText( LabelText );
+
+					HelpLabels.Add( Key, FloatingText );
+				}
+			}
+		}
+	}
+}
+
+
+void UVREditorInteractor::UpdateHelpLabels()
+{
+	const FTimespan HelpLabelFadeDuration = FTimespan::FromSeconds( VREd::HelpLabelFadeDuration->GetFloat() );
+
+	const FTransform HeadTransform = GetVRMode().GetHeadTransform();
+
+	// Only show help labels if the hand is pretty close to the face
+	const float DistanceToHead = ( GetTransform().GetLocation() - HeadTransform.GetLocation() ).Size();
+	const float MinDistanceToHeadForHelp = VREd::HelpLabelFadeDistance->GetFloat() * GetVRMode().GetWorldScaleFactor();	// (in cm)
+	bool bShowHelp = DistanceToHead <= MinDistanceToHeadForHelp;
+
+	// Don't show help if a UI is summoned on that hand
+	if ( bHasUIInFront || bHasUIOnForearm || GetVRMode().GetUISystem().IsShowingRadialMenu( this ) )
+	{
+		bShowHelp = false;
+	}
+
+	ShowHelpForHand( bShowHelp );
+
+	// Have the labels finished fading out?  If so, we'll kill their actors!
+	const FTimespan CurrentTime = FTimespan::FromSeconds( FApp::GetCurrentTime() );
+	const FTimespan TimeSinceStartedFadingOut = CurrentTime - HelpLabelShowOrHideStartTime;
+	if ( !bWantHelpLabels && ( TimeSinceStartedFadingOut > HelpLabelFadeDuration ) )
+	{
+		// Get rid of help text
+		for ( auto& KeyAndValue : HelpLabels )
+		{
+			AFloatingText* FloatingText = KeyAndValue.Value;
+			GetVRMode().DestroyTransientActor( FloatingText );
+		}
+		HelpLabels.Reset();
+	}
+	else
+	{
+		// Update fading state
+		float FadeAlpha = FMath::Clamp( ( float ) TimeSinceStartedFadingOut.GetTotalSeconds() / ( float ) HelpLabelFadeDuration.GetTotalSeconds(), 0.0f, 1.0f );
+		if ( !bWantHelpLabels )
+		{
+			FadeAlpha = 1.0f - FadeAlpha;
+		}
+
+		// Exponential falloff, so the fade is really obvious (gamma/HDR)
+		FadeAlpha = FMath::Pow( FadeAlpha, 3.0f );
+
+		for ( auto& KeyAndValue : HelpLabels )
+		{
+			const FKey Key = KeyAndValue.Key;
+			AFloatingText* FloatingText = KeyAndValue.Value;
+
+			UStaticMeshSocket* Socket = FindMeshSocketForKey( HandMeshComponent->StaticMesh, Key );
+			check( Socket != nullptr );
+			FTransform SocketRelativeTransform( Socket->RelativeRotation, Socket->RelativeLocation, Socket->RelativeScale );
+
+			// Oculus has asymmetrical controllers, so we the sock transform horizontally
+			if ( ControllerHandSide == EControllerHand::Right &&
+				GetVRMode().GetHMDDeviceType() == EHMDDeviceType::DT_OculusRift )
+			{
+				const FVector Scale3D = SocketRelativeTransform.GetLocation();
+				SocketRelativeTransform.SetLocation( FVector( Scale3D.X, -Scale3D.Y, Scale3D.Z ) );
+			}
+
+			// Make sure the labels stay the same size even when the world is scaled
+			FTransform HandTransformWithWorldToMetersScaling = GetTransform();
+			HandTransformWithWorldToMetersScaling.SetScale3D( HandTransformWithWorldToMetersScaling.GetScale3D() * FVector( GetVRMode().GetWorldScaleFactor() ) );
+
+			// Position right on top of the controller itself
+			FTransform FloatingTextTransform = SocketRelativeTransform * HandTransformWithWorldToMetersScaling;
+			FloatingText->SetActorTransform( FloatingTextTransform );
+
+			// Orientate it toward the viewer
+			FloatingText->Update( HeadTransform.GetLocation() );
+
+			// Update fade state
+			FloatingText->SetOpacity( FadeAlpha );
+		}
+	}
+}
+
+
+UStaticMeshSocket* UVREditorInteractor::FindMeshSocketForKey( UStaticMesh* StaticMesh, const FKey Key )
+{
+	// @todo vreditor: Hard coded mapping of socket names (e.g. "Shoulder") to expected names of sockets in the static mesh
+	FName SocketName = NAME_None;
+	if ( Key == EKeys::MotionController_Left_Shoulder || Key == EKeys::MotionController_Right_Shoulder )
+	{
+		static FName ShoulderSocketName( "Shoulder" );
+		SocketName = ShoulderSocketName;
+	}
+	else if ( Key == EKeys::MotionController_Left_Trigger || Key == EKeys::MotionController_Right_Trigger ||
+		Key == VREditorKeyNames::MotionController_Left_FullyPressedTriggerAxis || Key == VREditorKeyNames::MotionController_Right_FullyPressedTriggerAxis ||
+		Key == VREditorKeyNames::MotionController_Left_LightlyPressedTriggerAxis || Key == VREditorKeyNames::MotionController_Right_LightlyPressedTriggerAxis )
+	{
+		static FName TriggerSocketName( "Trigger" );
+		SocketName = TriggerSocketName;
+	}
+	else if ( Key == EKeys::MotionController_Left_Grip1 || Key == EKeys::MotionController_Right_Grip1 )
+	{
+		static FName GripSocketName( "Grip" );
+		SocketName = GripSocketName;
+	}
+	else if ( Key == EKeys::MotionController_Left_Thumbstick || Key == EKeys::MotionController_Right_Thumbstick )
+	{
+		static FName ThumbstickSocketName( "Thumbstick" );
+		SocketName = ThumbstickSocketName;
+	}
+	else if ( Key == SteamVRControllerKeyNames::Touch0 || Key == SteamVRControllerKeyNames::Touch1 )
+	{
+		static FName TouchSocketName( "Touch" );
+		SocketName = TouchSocketName;
+	}
+	else if ( Key == EKeys::MotionController_Left_Thumbstick_Down || Key == EKeys::MotionController_Right_Thumbstick_Down )
+	{
+		static FName DownSocketName( "Down" );
+		SocketName = DownSocketName;
+	}
+	else if ( Key == EKeys::MotionController_Left_Thumbstick_Up || Key == EKeys::MotionController_Right_Thumbstick_Up )
+	{
+		static FName UpSocketName( "Up" );
+		SocketName = UpSocketName;
+	}
+	else if ( Key == EKeys::MotionController_Left_Thumbstick_Left || Key == EKeys::MotionController_Right_Thumbstick_Left )
+	{
+		static FName LeftSocketName( "Left" );
+		SocketName = LeftSocketName;
+	}
+	else if ( Key == EKeys::MotionController_Left_Thumbstick_Right || Key == EKeys::MotionController_Right_Thumbstick_Right )
+	{
+		static FName RightSocketName( "Right" );
+		SocketName = RightSocketName;
+	}
+	else if ( Key == EKeys::MotionController_Left_FaceButton1 || Key == EKeys::MotionController_Right_FaceButton1 )
+	{
+		static FName FaceButton1SocketName( "FaceButton1" );
+		SocketName = FaceButton1SocketName;
+	}
+	else if ( Key == EKeys::MotionController_Left_FaceButton2 || Key == EKeys::MotionController_Right_FaceButton2 )
+	{
+		static FName FaceButton2SocketName( "FaceButton2" );
+		SocketName = FaceButton2SocketName;
+	}
+	else if ( Key == EKeys::MotionController_Left_FaceButton3 || Key == EKeys::MotionController_Right_FaceButton3 )
+	{
+		static FName FaceButton3SocketName( "FaceButton3" );
+		SocketName = FaceButton3SocketName;
+	}
+	else if ( Key == EKeys::MotionController_Left_FaceButton4 || Key == EKeys::MotionController_Right_FaceButton4 )
+	{
+		static FName FaceButton4SocketName( "FaceButton4" );
+		SocketName = FaceButton4SocketName;
+	}
+	else
+	{
+		// Not a key that we care about
+	}
+
+	if ( SocketName != NAME_None )
+	{
+		UStaticMeshSocket* Socket = StaticMesh->FindSocket( SocketName );
+		if ( Socket != nullptr )
+		{
+			return Socket;
+		}
+	}
+
+	return nullptr;
+};
+
+void UVREditorInteractor::SetLaserVisuals( const FLinearColor& NewColor, const float CrawlFade, const float CrawlSpeed )
+{
+	static FName StaticLaserColorParameterName( "LaserColor" );
+	LaserPointerMID->SetVectorParameterValue( StaticLaserColorParameterName, NewColor );
+	TranslucentLaserPointerMID->SetVectorParameterValue( StaticLaserColorParameterName, NewColor );
+
+	static FName StaticCrawlParameterName( "Crawl" );
+	LaserPointerMID->SetScalarParameterValue( StaticCrawlParameterName, CrawlFade );
+	TranslucentLaserPointerMID->SetScalarParameterValue( StaticCrawlParameterName, CrawlFade );
+
+	static FName StaticCrawlSpeedParameterName( "CrawlSpeed" );
+	LaserPointerMID->SetScalarParameterValue( StaticCrawlSpeedParameterName, CrawlSpeed );
+	TranslucentLaserPointerMID->SetScalarParameterValue( StaticCrawlSpeedParameterName, CrawlSpeed );
+
+	static FName StaticHandTrimColorParameter( "TrimGlowColor" );
+	HandMeshMID->SetVectorParameterValue( StaticHandTrimColorParameter, NewColor );
+
+	HoverPointLightComponent->SetLightColor( NewColor );
+}
+
+void UVREditorInteractor::UpdateRadialMenuInput( const float DeltaTime )
+{
+	UVREditorUISystem& UISystem = GetVRMode().GetUISystem();
+	UVREditorInteractor* OtherInteractor = Cast<UVREditorInteractor>( GetOtherInteractor() );
+
+	//Update the radial menu
+	EViewportInteractionDraggingMode DraggingMode = GetDraggingMode();
+	if ( !bHasUIInFront && 
+		( bIsTrackpadPositionValid[ 0 ] && bIsTrackpadPositionValid[ 1 ] ) &&
+		DraggingMode != EViewportInteractionDraggingMode::ActorsWithGizmo &&
+		DraggingMode != EViewportInteractionDraggingMode::ActorsFreely &&
+		DraggingMode != EViewportInteractionDraggingMode::ActorsAtLaserImpact &&
+		DraggingMode != EViewportInteractionDraggingMode::AssistingDrag && 
+		!UISystem.IsInteractorDraggingDockUI( this ) &&
+		!UISystem.IsShowingRadialMenu( OtherInteractor ) )
+	{
+		const EHMDDeviceType::Type HMDDevice = GetHMDDeviceType();
+
+		// Spawn the radial menu if we are using the touchpad for steamvr or the analog stick for the oculus
+		if ( ( HMDDevice == EHMDDeviceType::DT_SteamVR && bIsTouchingTrackpad ) ||
+			 ( HMDDevice == EHMDDeviceType::DT_OculusRift && TrackpadPosition.Size() > VREd::MinJoystickOffsetBeforeRadialMenu->GetFloat() ) )
+		{
+			UISystem.TryToSpawnRadialMenu( this );
+		}
+		else
+		{
+			// Hide it if we are not using the touchpad or analog stick
+			UISystem.HideRadialMenu( this );
+		}
+
+		// Update the radial menu if we are already showing the radial menu
+		if ( UISystem.IsShowingRadialMenu( this ) )
+		{
+			UISystem.UpdateRadialMenu( this );
+		}
+	}
+}
+
+#undef LOCTEXT_NAMESPACE
\ No newline at end of file
