VR Editor:

diff --git a/Engine/Source/Editor/VREditor/VREditorMode.cpp b/Engine/Source/Editor/VREditor/VREditorMode.cpp
--- a/Engine/Source/Editor/VREditor/VREditorMode.cpp
+++ b/Engine/Source/Editor/VREditor/VREditorMode.cpp
@@ -7,6 +7,8 @@
 #include "VREditorTransformGizmo.h"
 #include "VREditorFloatingText.h"
 #include "VREditorFloatingUI.h"
+#include "Teleporter/VREditorTeleporter.h"
+
 
 #include "CameraController.h"
 #include "DynamicMeshBuilder.h"
@@ -19,26 +21,17 @@
 #include "IHeadMountedDisplay.h"
 
 #include "Editor/LevelEditor/Public/LevelEditorActions.h"
+#include "Editor/ViewportInteraction/Public/ViewportInteraction.h"
+#include "VREditorInteractor.h"
+
+#include "Interactables/VREditorButton.h"
 
 #define LOCTEXT_NAMESPACE "VREditorMode"
 
 namespace VREd
 {
-	static FAutoConsoleVariable LaserPointerMaxLength( TEXT( "VREd.LaserPointerMaxLength" ), 10000.0f, TEXT( "Maximum length of the laser pointer line" ) );
-	static FAutoConsoleVariable LaserPointerRadius( TEXT( "VREd.LaserPointerRadius" ), 0.5f, TEXT( "Radius of the laser pointer line" ) );
-	static FAutoConsoleVariable LaserPointerHoverBallRadius( TEXT( "VREd.LaserPointerHoverBallRadius" ), 1.5f, TEXT( "Radius of the visual cue for a hovered object along the laser pointer ray" ) );
-	static FAutoConsoleVariable LaserPointerLightPullBackDistance( TEXT( "VREd.LaserPointerLightPullBackDistance" ), 2.5f, TEXT( "How far to pull back our little hover light from the impact surface" ) );
-	static FAutoConsoleVariable LaserPointerLightRadius( TEXT( "VREd.LaserPointLightRadius" ), 20.0f, TEXT( "How big our hover light is" ) );
-	static FAutoConsoleVariable OculusLaserPointerRotationOffset( TEXT( "VREd.OculusLaserPointerRotationOffset" ), 0.0f, TEXT( "How much to rotate the laser pointer (pitch) relative to the forward vector of the controller (Oculus)" ) );
-	static FAutoConsoleVariable ViveLaserPointerRotationOffset( TEXT( "VREd.ViveLaserPointerRotationOffset" ), /* -57.8f */ 0.0f, TEXT( "How much to rotate the laser pointer (pitch) relative to the forward vector of the controller (Vive)" ) );
-	static FAutoConsoleVariable OculusLaserPointerStartOffset( TEXT( "VREd.OculusLaserPointerStartOffset" ), 2.8f, TEXT( "How far to offset the start of the laser pointer to avoid overlapping the hand mesh geometry (Oculus)" ) );
-	static FAutoConsoleVariable ViveLaserPointerStartOffset( TEXT( "VREd.ViveLaserPointerStartOffset" ), 1.25f /* 8.5f */, TEXT( "How far to offset the start of the laser pointer to avoid overlapping the hand mesh geometry (Vive)" ) );
-
 	static FAutoConsoleVariable UseMouseAsHandInForcedVRMode( TEXT( "VREd.UseMouseAsHandInForcedVRMode" ), 1, TEXT( "When in forced VR mode, enabling this setting uses the mouse cursor as a virtual hand instead of motion controllers" ) );
 
-	static FAutoConsoleVariable HelpLabelFadeDuration( TEXT( "VREd.HelpLabelFadeDuration" ), 0.4f, TEXT( "Duration to fade controller help labels in and out" ) );
-	static FAutoConsoleVariable HelpLabelFadeDistance( TEXT( "VREd.HelpLabelFadeDistance" ), 30.0f, TEXT( "Distance at which controller help labels should appear (in cm)" ) );
-
 	static FAutoConsoleVariable GridMovementTolerance( TEXT( "VREd.GridMovementTolerance" ), 0.1f, TEXT( "Tolerance for movement when the grid must disappear" ) );
 	static FAutoConsoleVariable GridScaleMultiplier( TEXT( "VREd.GridScaleMultiplier" ), 35.0f, TEXT( "Scale of the grid" ) );
 	static FAutoConsoleVariable GridFadeMultiplier( TEXT( "VREd.GridFadeMultiplier" ), 3.0f, TEXT( "Grid fade in/out speed, in 'fades per second'" ) );
@@ -46,11 +39,7 @@ namespace VREd
 	static FAutoConsoleVariable GridMaxOpacity( TEXT( "VREd.GridMaxFade" ), 0.8f, TEXT( "Grid maximum opacity" ) );
 	static FAutoConsoleVariable GridHeightOffset( TEXT( "VREd.GridHeightOffset" ), 0.0f, TEXT( "Height offset for the world movement grid.  Useful when tracking space is not properly calibrated" ) );
 
-	static FAutoConsoleVariable TriggerFullyPressedThreshold( TEXT( "VREd.TriggerFullyPressedThreshold" ), 0.95f, TEXT( "Minimum trigger threshold before we consider the trigger 'fully pressed'" ) );
-	static FAutoConsoleVariable TriggerFullyPressedReleaseThreshold( TEXT( "VREd.TriggerFullyPressedReleaseThreshold" ), 0.8f, TEXT( "After fully pressing the trigger, if the axis falls below this threshold we no longer consider it fully pressed" ) );
-	static FAutoConsoleVariable TriggerLightlyPressedThreshold( TEXT( "VREd.TriggerLightlyPressedThreshold" ), 0.03f, TEXT( "Minimum trigger threshold before we consider the trigger at least 'lightly pressed'" ) );
 	static FAutoConsoleVariable TriggerLightlyPressedLockTime( TEXT( "VREd.TriggerLightlyPressedLockTime" ), 0.15f, TEXT( "If the trigger remains lightly pressed for longer than this, we'll continue to treat it as a light press in some cases" ) );
-	static FAutoConsoleVariable TriggerDeadZone( TEXT( "VREd.TriggerDeadZone" ), 0.01f, TEXT( "Trigger dead zone.  The trigger must be fully released before we'll trigger a new 'light press'" ) );
 
 	static FAutoConsoleVariable WorldMovementFogOpacity( TEXT( "VREd.WorldMovementFogOpacity" ), 0.8f, TEXT( "How opaque the fog should be at the 'end distance' (0.0 - 1.0)" ) );
 	static FAutoConsoleVariable WorldMovementFogStartDistance( TEXT( "VREd.WorldMovementFogStartDistance" ), 300.0f, TEXT( "How far away fog will start rendering while in world movement mode" ) );
@@ -60,41 +49,17 @@ namespace VREd
 	static FAutoConsoleVariable ScaleProgressBarLength( TEXT( "VREd.ScaleProgressBarLength" ), 50.0f, TEXT( "Length of the progressbar that appears when scaling" ) );
 	static FAutoConsoleVariable ScaleProgressBarRadius( TEXT( "VREd.ScaleProgressBarRadius" ), 1.0f, TEXT( "Radius of the progressbar that appears when scaling" ) );
 
-	static FAutoConsoleVariable MinHapticTimeForRift( TEXT( "VREd.MinHapticTimeForRift" ), 0.005f, TEXT( "How long to play haptic effects on the Rift" ) );
-	static FAutoConsoleVariable SleepForRiftHaptics( TEXT( "VREd.SleepForRiftHaptics" ), 1, TEXT( "When enabled, we'll sleep the game thread mid-frame to wait for haptic effects to finish.  This can be devasting to performance!" ) );
-
 	static FAutoConsoleVariable InvertTrackpadVertical( TEXT( "VREd.InvertTrackpadVertical" ), 1, TEXT( "Toggles inverting the touch pad vertical axis" ) );
 	static FAutoConsoleVariable ForceOculusMirrorMode( TEXT( "VREd.ForceOculusMirrorMode" ), 3, TEXT( "Which Oculus display mirroring mode to use (see MirrorWindowModeType in OculusRiftHMD.h)" ) );
 
 	static FAutoConsoleVariable ShowMovementGrid( TEXT( "VREd.ShowMovementGrid" ), 1, TEXT( "Showing the ground movement grid" ) );
 	static FAutoConsoleVariable ShowWorldMovementPostProcess( TEXT( "VREd.ShowWorldMovementPostProcess" ), 1, TEXT( "Showing the movement post processing" ) );
-
-	static FAutoConsoleVariable ScaleMax( TEXT( "VREd.ScaleMax" ), 6000.0f, TEXT( "Maximum world scale in centimeters" ) );
-	static FAutoConsoleVariable ScaleMin( TEXT( "VREd.ScaleMin" ), 10.0f, TEXT( "Minimum world scale in centimeters" ) );
 }
 
 FEditorModeID FVREditorMode::VREditorModeID( "VREditor" );
 bool FVREditorMode::bActuallyUsingVR = true;
 
 // @todo vreditor: Hacky that we have to import these this way. (Plugin has them in a .cpp, not exported)
-namespace SteamVRControllerKeyNames
-{
-	static const FGamepadKeyNames::Type Touch0( "Steam_Touch_0" );
-	static const FGamepadKeyNames::Type Touch1( "Steam_Touch_1" );
-}
-
-
-namespace VREditorKeyNames
-{
-	// @todo vreditor input: Ideally these would not be needed, but SteamVR fires off it's "trigger pressed" event
-	// well before the trigger is fully down (*click*)
-	static const FName MotionController_Left_FullyPressedTriggerAxis( "MotionController_Left_FullyPressedTriggerAxis" );
-	static const FName MotionController_Right_FullyPressedTriggerAxis( "MotionController_Right_FullyPressedTriggerAxis" );
-
-	static const FName MotionController_Left_LightlyPressedTriggerAxis( "MotionController_Left_LightlyPressedTriggerAxis" );
-	static const FName MotionController_Right_LightlyPressedTriggerAxis( "MotionController_Right_LightlyPressedTriggerAxis" );
-}
-
 
 FVREditorMode::FVREditorMode()
 	: bWantsToExitMode( false ),
@@ -109,18 +74,13 @@ FVREditorMode::FVREditorMode()
 	  WorldMovementGridOpacity( 0.0f ),	// NOTE: Intentionally not quite zero so that we update the MIDs on the first frame
 	  bIsDrawingWorldMovementPostProcess( false ),
 	  WorldMovementPostProcessMaterial( nullptr ),
-	  SnapGridActor( nullptr ),
-	  SnapGridMeshComponent( nullptr ),
-	  SnapGridMID( nullptr ),
 	  ScaleProgressMeshComponent( nullptr ),
 	  CurrentScaleProgressMeshComponent( nullptr ),
 	  UserScaleIndicatorText( nullptr ),
 	  PostProcessComponent( nullptr ),
 	  MotionControllerID( 0 ),	// @todo vreditor minor: We only support a single controller, and we assume the first controller are the motion controls
-	  LastFrameNumberInputWasPolled( 0 ),
-	  UISystem(),
-	  WorldInteraction(),
-	  CurrentGizmoType( EGizmoHandleTypes::All ),
+	  UISystem( nullptr ),
+	  WorldInteraction( nullptr ),
 	  bFirstTick( true ),
 	  bWasInWorldSpaceBeforeScaleMode( false )
 {
@@ -142,9 +102,6 @@ FVREditorMode::~FVREditorMode()
 
 void FVREditorMode::Enter()
 {
-	InputProcessor = MakeShareable(new FVREditorInputProcessor(this));
-	FSlateApplication::Get().SetInputPreProcessor(true, InputProcessor);
-
 	// @todo vreditor urgent: Turn on global editor hacks for VR Editor mode
 	GEnableVREditorHacks = true;
 
@@ -155,7 +112,6 @@ void FVREditorMode::Enter()
 
 	WorldMovementGridOpacity = 0.0f;
 	bIsDrawingWorldMovementPostProcess = false;
-	LastFrameNumberInputWasPolled = 0;
 
 	// Take note of VREditor activation
 	if( FEngineAnalytics::IsAvailable() )
@@ -163,69 +119,6 @@ void FVREditorMode::Enter()
 		FEngineAnalytics::GetProvider().RecordEvent( TEXT( "Editor.Usage.EnterVREditorMode" ) );
 	}
 
-	// Fully reset all hand state.  We don't want anything carrying over from the previous session.
-	for( int32 HandIndex = 0; HandIndex < VREditorConstants::NumVirtualHands; ++HandIndex )
-	{
-		FVirtualHand& Hand = VirtualHands[ HandIndex ];
-		Hand = FVirtualHand();
-	}
-
-	// Call parent implementation
-	FEdMode::Enter();
-
-	{
-		KeyToActionMap.Reset();
-
-		if( !bActuallyUsingVR && VREd::UseMouseAsHandInForcedVRMode->GetInt() != 0 )
-		{
-			// For mouse control, when in forced VR mode
-			KeyToActionMap.Add( EKeys::LeftMouseButton, FVRAction( EVRActionType::SelectAndMove, VREditorConstants::LeftHandIndex ) );
-			KeyToActionMap.Add( EKeys::MiddleMouseButton, FVRAction( EVRActionType::SelectAndMove_LightlyPressed, VREditorConstants::LeftHandIndex ) );
-			//			KeyToActionMap.Add( EKeys::RightMouseButton, FVRAction( EVRActionType::WorldMovement, VREditorConstants::LeftHandIndex ) );
-//			KeyToActionMap.Add( EKeys::MiddleMouseButton, FVRAction( EVRActionType::Selection, VREditorConstants::LeftHandIndex ) );
-			KeyToActionMap.Add( EKeys::LeftControl, FVRAction( EVRActionType::Modifier, VREditorConstants::LeftHandIndex ) );
-		}
-
-		// Motion controllers
-		KeyToActionMap.Add( EKeys::MotionController_Left_Grip1, FVRAction( EVRActionType::WorldMovement, VREditorConstants::LeftHandIndex ) );
-		KeyToActionMap.Add( EKeys::MotionController_Right_Grip1, FVRAction( EVRActionType::WorldMovement, VREditorConstants::RightHandIndex ) );
-		if( GetHMDDeviceType() == EHMDDeviceType::DT_SteamVR )
-		{
-			KeyToActionMap.Add( EKeys::MotionController_Left_Shoulder /* Red face button */, FVRAction( EVRActionType::Modifier, VREditorConstants::LeftHandIndex ) );
-			KeyToActionMap.Add( EKeys::MotionController_Right_Shoulder /* Red face button */, FVRAction( EVRActionType::Modifier, VREditorConstants::RightHandIndex ) );
-		}
-		else
-		{
-			KeyToActionMap.Add( EKeys::MotionController_Left_FaceButton1, FVRAction( EVRActionType::Modifier, VREditorConstants::LeftHandIndex ) );
-			KeyToActionMap.Add( EKeys::MotionController_Right_FaceButton1, FVRAction( EVRActionType::Modifier, VREditorConstants::RightHandIndex ) );
-		}
-
-// 		KeyToActionMap.Add( EKeys::MotionController_Left_Trigger, FVRAction( EVRActionType::SelectAndMove, VREditorConstants::LeftHandIndex ) );
-// 		KeyToActionMap.Add( EKeys::MotionController_Right_Trigger, FVRAction( EVRActionType::SelectAndMove, VREditorConstants::RightHandIndex ) );
-		KeyToActionMap.Add( VREditorKeyNames::MotionController_Left_FullyPressedTriggerAxis, FVRAction( EVRActionType::SelectAndMove, VREditorConstants::LeftHandIndex ) );
-		KeyToActionMap.Add( VREditorKeyNames::MotionController_Right_FullyPressedTriggerAxis, FVRAction( EVRActionType::SelectAndMove, VREditorConstants::RightHandIndex ) );
-		KeyToActionMap.Add( VREditorKeyNames::MotionController_Left_LightlyPressedTriggerAxis, FVRAction( EVRActionType::SelectAndMove_LightlyPressed, VREditorConstants::LeftHandIndex ) );
-		KeyToActionMap.Add( VREditorKeyNames::MotionController_Right_LightlyPressedTriggerAxis, FVRAction( EVRActionType::SelectAndMove_LightlyPressed, VREditorConstants::RightHandIndex ) );
-
-		KeyToActionMap.Add( EKeys::MotionController_Left_Thumbstick, FVRAction( EVRActionType::ConfirmRadialSelection, VREditorConstants::LeftHandIndex ) );
- 		KeyToActionMap.Add( EKeys::MotionController_Right_Thumbstick, FVRAction( EVRActionType::ConfirmRadialSelection, VREditorConstants::RightHandIndex ) );
-
-// 		KeyToActionMap.Add( EKeys::MotionController_Left_FaceButton2 /* Touchpad right */, FVRAction( EVRActionType::Redo, VREditorConstants::LeftHandIndex ) );
-// 		KeyToActionMap.Add( EKeys::MotionController_Right_FaceButton2 /* Touchpad right */, FVRAction( EVRActionType::Redo, VREditorConstants::RightHandIndex ) );
-// 		KeyToActionMap.Add( EKeys::MotionController_Left_FaceButton3 /* Touchpad down */, FVRAction( EVRActionType::Delete, VREditorConstants::LeftHandIndex ) );
-// 		KeyToActionMap.Add( EKeys::MotionController_Right_FaceButton3 /* Touchpad down */, FVRAction( EVRActionType::Delete, VREditorConstants::RightHandIndex ) );
-// 		KeyToActionMap.Add( EKeys::MotionController_Left_FaceButton4 /* Touchpad left */, FVRAction( EVRActionType::Undo, VREditorConstants::LeftHandIndex ) );
-// 		KeyToActionMap.Add( EKeys::MotionController_Right_FaceButton4 /* Touchpad left */, FVRAction( EVRActionType::Undo, VREditorConstants::RightHandIndex ) );
-		KeyToActionMap.Add( SteamVRControllerKeyNames::Touch0, FVRAction( EVRActionType::Touch, VREditorConstants::LeftHandIndex ) );
- 		KeyToActionMap.Add( SteamVRControllerKeyNames::Touch1, FVRAction( EVRActionType::Touch, VREditorConstants::RightHandIndex ) );
-		KeyToActionMap.Add( EKeys::MotionController_Left_TriggerAxis, FVRAction( EVRActionType::TriggerAxis, VREditorConstants::LeftHandIndex ) );
-		KeyToActionMap.Add( EKeys::MotionController_Right_TriggerAxis, FVRAction( EVRActionType::TriggerAxis, VREditorConstants::RightHandIndex ) );
-		KeyToActionMap.Add( EKeys::MotionController_Left_Thumbstick_X, FVRAction( EVRActionType::TrackpadPositionX, VREditorConstants::LeftHandIndex ) );
-		KeyToActionMap.Add( EKeys::MotionController_Right_Thumbstick_X, FVRAction( EVRActionType::TrackpadPositionX, VREditorConstants::RightHandIndex ) );
-		KeyToActionMap.Add( EKeys::MotionController_Left_Thumbstick_Y, FVRAction( EVRActionType::TrackpadPositionY, VREditorConstants::LeftHandIndex ) );
-		KeyToActionMap.Add( EKeys::MotionController_Right_Thumbstick_Y, FVRAction( EVRActionType::TrackpadPositionY, VREditorConstants::RightHandIndex ) );
-	}
-
 	// Setting up colors
 	Colors.SetNumZeroed( (int32)EColors::TotalCount );
 	{
@@ -268,7 +161,7 @@ void FVREditorMode::Enter()
 				bSummonNewWindow = false;
 			}
 		}
-
+		
 		TSharedPtr< SLevelViewport > VREditorLevelViewport;
 		if( bSummonNewWindow )
 		{
@@ -341,7 +234,6 @@ void FVREditorMode::Enter()
 			FLevelEditorViewportClient& VRViewportClient = VREditorLevelViewport->GetLevelViewportClient();
 			FEditorViewportClient& VREditorViewportClient = VRViewportClient;
 
-
 			// Make sure we are in perspective mode
 			// @todo vreditor: We should never allow ortho switching while in VR
 			SavedEditorState.ViewportType = VREditorViewportClient.GetViewportType();
@@ -429,10 +321,34 @@ void FVREditorMode::Enter()
 		check( WorldMovementPostProcessMaterial != nullptr );
 	}
 
+	// Setup sub systems
+	{
+		// Setup world interaction
+		WorldInteraction = NewObject<UVREditorWorldInteraction>();
+		WorldInteraction->SetOwner( this );
+		WorldInteraction->SetIsUsingVR( bActuallyUsingVR );
+		WorldInteraction->Init( VREditorLevelViewportWeakPtr.Pin()->GetViewportClient().Get() );
+
+		LeftHandInteractor = NewObject<UVREditorInteractor>( WorldInteraction );
+		LeftHandInteractor->Init( EControllerHand::Left, this );
+		WorldInteraction->AddInteractor( LeftHandInteractor );
+
+		RightHandInteractor = NewObject<UVREditorInteractor>( WorldInteraction );
+		RightHandInteractor->Init( EControllerHand::Right, this );
+		WorldInteraction->AddInteractor( RightHandInteractor );
 
-	UISystem.Reset( new FVREditorUISystem( *this ) );
+		WorldInteraction->PairInteractors( LeftHandInteractor, RightHandInteractor );
 
-	WorldInteraction.Reset( new FVREditorWorldInteraction( *this ) );
+		// Setup the UI system
+		UISystem = NewObject<UVREditorUISystem>();
+		UISystem->SetOwner( this );
+		UISystem->Init();
+
+		// Setup teleporter
+		TeleporterSystem = NewObject<UVREditorTeleporter>();
+		TeleporterSystem->Init( this );
+
+	}
 
 	//// If we're not actually in VR mode, and instead we've forced it, we should spawn some UI at the Origin to help test
 	//if ( !bActuallyUsingVR )
@@ -448,9 +364,12 @@ void FVREditorMode::Enter()
 
 	if ( bActuallyUsingVR )
 	{
-		 SetWorldToMetersScale( 100.0f );
+		 WorldInteraction->SetWorldToMetersScale( 100.0f );
 	}
 
+	// Call parent implementation
+	FEdMode::Enter();
+
 	bFirstTick = true;
 	bIsFullyInitialized = true;
 }
@@ -462,12 +381,6 @@ void FVREditorMode::Exit()
 
 	FSlateApplication::Get().SetInputPreProcessor(false);
 
-	// Kill subsystems
-	{
-		WorldInteraction.Reset();
-		UISystem.Reset();
-	}
-
 	if( WorldMovementPostProcessMaterial != nullptr )
 	{
 		WorldMovementPostProcessMaterial->MarkPendingKill();
@@ -477,39 +390,15 @@ void FVREditorMode::Exit()
 	{
 		DestroyTransientActor( AvatarMeshActor );
 		AvatarMeshActor = nullptr;
+		
 		HeadMeshComponent = nullptr;
 		FlashlightComponent = nullptr;
 		WorldMovementGridMeshComponent = nullptr;
 		WorldMovementGridMID = nullptr;
-		DestroyTransientActor( SnapGridActor );
-		SnapGridActor = nullptr;
-		SnapGridMeshComponent = nullptr;
-		SnapGridMID = nullptr;
 		PostProcessComponent = nullptr;
 		ScaleProgressMeshComponent = nullptr;
 		CurrentScaleProgressMeshComponent = nullptr;
 		UserScaleIndicatorText = nullptr;
-		
-
-		for( int32 HandIndex = 0; HandIndex < VREditorConstants::NumVirtualHands; ++HandIndex )
-		{
-			FVirtualHand& Hand = VirtualHands[ HandIndex ];
-
-			Hand.MotionControllerComponent = nullptr;
-			Hand.HandMeshComponent = nullptr;
-			Hand.HoverMeshComponent = nullptr;
-			Hand.HoverPointLightComponent = nullptr;
-			Hand.LaserPointerMeshComponent = nullptr;
-			Hand.LaserPointerMID = nullptr;
-			Hand.TranslucentLaserPointerMID = nullptr;
-
-			for( auto& KeyAndValue : Hand.HelpLabels )
-			{
-				AFloatingText* FloatingText = KeyAndValue.Value;
-				DestroyTransientActor( FloatingText );
-			}
-			Hand.HelpLabels.Reset();
-		}
 	}
 
 	{
@@ -579,11 +468,33 @@ void FVREditorMode::Exit()
 	}
 
 	// Make sure we are no longer overriding WorldToMetersScale every frame
-	SetWorldToMetersScale( 0.0f );
+	WorldInteraction->SetWorldToMetersScale( 0.0f );
 
 	// Call parent implementation
 	FEdMode::Exit();
 
+	// Kill subsystems
+	if ( UISystem != nullptr )
+	{
+		UISystem->Shutdown();
+		UISystem->MarkPendingKill();
+		UISystem = nullptr;
+	}
+
+	if ( WorldInteraction != nullptr )
+	{
+		WorldInteraction->Shutdown();
+		WorldInteraction->MarkPendingKill();
+		WorldInteraction = nullptr;
+	}
+
+	if ( TeleporterSystem != nullptr )
+	{
+		TeleporterSystem->Shutdown();
+		TeleporterSystem->MarkPendingKill();
+		TeleporterSystem = nullptr;
+	}
+
 	// @todo vreditor urgent: Disable global editor hacks for VR Editor mode
 	GEnableVREditorHacks = false;
 }
@@ -615,119 +526,10 @@ void FVREditorMode::SpawnAvatarMeshActor()
 			HeadMeshComponent->SetCollisionEnabled( ECollisionEnabled::NoCollision );
 		}
 
-		for( int32 HandIndex = 0; HandIndex < ARRAY_COUNT( VirtualHands ); ++HandIndex )
-		{
-			FVirtualHand& Hand = VirtualHands[ HandIndex ];
-
-			// Setup a motion controller component.  This allows us to take advantage of late frame updates, so
-			// our motion controllers won't lag behind the HMD
-			{
-				Hand.MotionControllerComponent = NewObject<UMotionControllerComponent>( AvatarMeshActor );
-				AvatarMeshActor->AddOwnedComponent( Hand.MotionControllerComponent );
-				Hand.MotionControllerComponent->SetupAttachment( AvatarMeshActor->GetRootComponent() );
-				Hand.MotionControllerComponent->RegisterComponent();
-
-				Hand.MotionControllerComponent->SetMobility( EComponentMobility::Movable );
-				Hand.MotionControllerComponent->SetCollisionEnabled( ECollisionEnabled::NoCollision );
-
-				Hand.MotionControllerComponent->Hand = ( HandIndex == VREditorConstants::LeftHandIndex ) ? EControllerHand::Left : EControllerHand::Right;
-
-				Hand.MotionControllerComponent->bDisableLowLatencyUpdate = false;
-			}
-
-			// Hand mesh
-			{
-				Hand.HandMeshComponent = NewObject<UStaticMeshComponent>( AvatarMeshActor );
-				AvatarMeshActor->AddOwnedComponent( Hand.HandMeshComponent );
-				Hand.HandMeshComponent->SetupAttachment( Hand.MotionControllerComponent );
-				Hand.HandMeshComponent->RegisterComponent();
-
-				// @todo vreditor extensibility: We need this to be able to be overridden externally, or simply based on the HMD name (but allowing external folders)
-				FString MeshName;
-				FString MaterialName;
-				if( GetHMDDeviceType() == EHMDDeviceType::DT_SteamVR )
-				{
-					MeshName = TEXT( "/Engine/VREditor/Devices/Vive/VivePreControllerMesh" );
-					MaterialName = TEXT( "/Engine/VREditor/Devices/Vive/VivePreControllerMaterial_Inst" );
-				}
-				else
-				{
-					MeshName = TEXT( "/Engine/VREditor/Devices/Oculus/OculusControllerMesh" );
-					MaterialName = TEXT( "/Engine/VREditor/Devices/Oculus/OculusControllerMaterial_Inst" );
-				}
-
-				UStaticMesh* HandMesh = LoadObject<UStaticMesh>( nullptr, *MeshName );
-				check( HandMesh != nullptr );
-
-				Hand.HandMeshComponent->SetStaticMesh( HandMesh );
-				Hand.HandMeshComponent->SetMobility( EComponentMobility::Movable );
-				Hand.HandMeshComponent->SetCollisionEnabled( ECollisionEnabled::NoCollision );
-
-				UMaterialInstance* HandMeshMaterialInst = LoadObject<UMaterialInstance>( nullptr, *MaterialName );
-				check( HandMeshMaterialInst != nullptr );
-				Hand.HandMeshMID = UMaterialInstanceDynamic::Create( HandMeshMaterialInst, GetTransientPackage() );
-				check( Hand.HandMeshMID != nullptr );
-				Hand.HandMeshComponent->SetMaterial( 0, Hand.HandMeshMID );
-			}
-
-			// Laser pointer
-			{
-				Hand.LaserPointerMeshComponent = NewObject<UStaticMeshComponent>( AvatarMeshActor );
-				AvatarMeshActor->AddOwnedComponent( Hand.LaserPointerMeshComponent );
-				Hand.LaserPointerMeshComponent->SetupAttachment( Hand.MotionControllerComponent );
-				Hand.LaserPointerMeshComponent->RegisterComponent();
-
-				UStaticMesh* LaserPointerMesh = LoadObject<UStaticMesh>( nullptr, TEXT( "/Engine/VREditor/LaserPointer/LaserPointerMesh" ) );
-				check( LaserPointerMesh != nullptr );
-
-				Hand.LaserPointerMeshComponent->SetStaticMesh( LaserPointerMesh );
-				Hand.LaserPointerMeshComponent->SetMobility( EComponentMobility::Movable );
-				Hand.LaserPointerMeshComponent->SetCollisionEnabled( ECollisionEnabled::NoCollision );
-
-				UMaterialInstance* LaserPointerMaterialInst = LoadObject<UMaterialInstance>( nullptr, TEXT( "/Engine/VREditor/LaserPointer/LaserPointerMaterialInst" ) );
-				check( LaserPointerMaterialInst != nullptr );
-				Hand.LaserPointerMID = UMaterialInstanceDynamic::Create( LaserPointerMaterialInst, GetTransientPackage() );
-				check( Hand.LaserPointerMID != nullptr );
-				Hand.LaserPointerMeshComponent->SetMaterial( 0, Hand.LaserPointerMID );
-
-				UMaterialInstance* TranslucentLaserPointerMaterialInst = LoadObject<UMaterialInstance>( nullptr, TEXT( "/Engine/VREditor/LaserPointer/TranslucentLaserPointerMaterialInst" ) );
-				check( TranslucentLaserPointerMaterialInst != nullptr );
-				Hand.TranslucentLaserPointerMID = UMaterialInstanceDynamic::Create( TranslucentLaserPointerMaterialInst, GetTransientPackage() );
-				check( Hand.TranslucentLaserPointerMID != nullptr );
-				Hand.LaserPointerMeshComponent->SetMaterial( 1, Hand.TranslucentLaserPointerMID );
-			}
-
-			// Hover cue for laser pointer
-			{
-				Hand.HoverMeshComponent = NewObject<UStaticMeshComponent>( AvatarMeshActor );
-				AvatarMeshActor->AddOwnedComponent( Hand.HoverMeshComponent );
-				Hand.HoverMeshComponent->SetupAttachment( AvatarMeshActor->GetRootComponent() );
-				Hand.HoverMeshComponent->RegisterComponent();
-
-				UStaticMesh* HoverMesh = LoadObject<UStaticMesh>( nullptr, TEXT( "/Engine/VREditor/LaserPointer/HoverMesh" ) );
-				check( HoverMesh != nullptr );
-				Hand.HoverMeshComponent->SetStaticMesh( HoverMesh );
-				Hand.HoverMeshComponent->SetMobility( EComponentMobility::Movable );
-				Hand.HoverMeshComponent->SetCollisionEnabled( ECollisionEnabled::NoCollision );
-
-				// Add a light!
-				{
-					Hand.HoverPointLightComponent = NewObject<UPointLightComponent>( AvatarMeshActor );
-					AvatarMeshActor->AddOwnedComponent( Hand.HoverPointLightComponent );
-					Hand.HoverPointLightComponent->SetupAttachment( Hand.HoverMeshComponent );
-					Hand.HoverPointLightComponent->RegisterComponent();
-
-					Hand.HoverPointLightComponent->SetLightColor( FLinearColor::Red );
-//					Hand.HoverPointLightComponent->SetLightColor( FLinearColor( 0.0f, 1.0f, 0.2f, 1.0f ) );
-					Hand.HoverPointLightComponent->SetIntensity( 30.0f );	// @todo vreditor tweak
-					Hand.HoverPointLightComponent->SetMobility( EComponentMobility::Movable );
-					Hand.HoverPointLightComponent->SetAttenuationRadius( VREd::LaserPointerLightRadius->GetFloat() );
-					Hand.HoverPointLightComponent->bUseInverseSquaredFalloff = false;
-					Hand.HoverPointLightComponent->SetCastShadows( false );
-				}
-			}
-		}
-
+		//@todo VREditor: Hardcoded interactors
+		LeftHandInteractor->SetupComponent( AvatarMeshActor );
+		RightHandInteractor->SetupComponent( AvatarMeshActor );
+		
 		// World movement grid mesh
 		{
 			WorldMovementGridMeshComponent = NewObject<UStaticMeshComponent>( AvatarMeshActor );
@@ -752,33 +554,6 @@ void FVREditorMode::SpawnAvatarMeshActor()
 			WorldMovementGridMeshComponent->SetVisibility( false );
 		}
 
-		// Snap grid mesh
-		{
-			const bool bWithSceneComponent = false;
-			SnapGridActor = SpawnTransientSceneActor<AActor>( TEXT( "SnapGrid" ), bWithSceneComponent );
-
-			SnapGridMeshComponent = NewObject<UStaticMeshComponent>( SnapGridActor );
-			SnapGridActor->AddOwnedComponent( SnapGridMeshComponent );
-			SnapGridActor->SetRootComponent( SnapGridMeshComponent );
-			SnapGridMeshComponent->RegisterComponent();
-
-			UStaticMesh* GridMesh = LoadObject<UStaticMesh>( nullptr, TEXT( "/Engine/VREditor/SnapGrid/SnapGridPlaneMesh" ) );
-			check( GridMesh != nullptr );
-			SnapGridMeshComponent->SetStaticMesh( GridMesh );
-			SnapGridMeshComponent->SetMobility( EComponentMobility::Movable );
-			SnapGridMeshComponent->SetCollisionEnabled( ECollisionEnabled::NoCollision );
-
-			UMaterialInterface* GridMaterial = LoadObject<UMaterialInterface>( nullptr, TEXT( "/Engine/VREditor/SnapGrid/SnapGridMaterial" ) );
-			check( GridMaterial != nullptr );
-
-			SnapGridMID = UMaterialInstanceDynamic::Create( GridMaterial, GetTransientPackage() );
-			check( SnapGridMID != nullptr );
-			SnapGridMeshComponent->SetMaterial( 0, SnapGridMID );
-
-			// The grid starts off hidden
-			SnapGridMeshComponent->SetVisibility( false );
-		}
-
 		{
 			{
 				UMaterialInterface* UserScaleIndicatorMaterial = LoadObject<UMaterialInterface>( nullptr, TEXT( "/Engine/VREditor/LaserPointer/LaserPointerMaterialInst" ) );
@@ -919,53 +694,24 @@ void FVREditorMode::Tick( FEditorViewportClient* ViewportClient, float DeltaTime
 	FEdMode::Tick( ViewportClient, DeltaTime );
 
 	// Only if this is our VR viewport. Remember, editor modes get a chance to tick and receive input for each active viewport.
-	if( ViewportClient != GetLevelViewportPossessedForVR().GetViewportClient().Get() )
+	if ( ViewportClient != GetLevelViewportPossessedForVR().GetViewportClient().Get() || bWantsToExitMode )
 	{
 		return;
 	}
-
-	
-	// Cancel any haptic effects that have been playing too long
-	StopOldHapticEffects();
-
-	// Get the latest controller data, and fill in our VirtualHands with fresh transforms
-	PollInputIfNeeded();
 	
-
-	// Update hover.  Note that hover can also be updated when ticking our sub-systems below.
+	if ( AvatarMeshActor == nullptr )
 	{
-		for( int32 HandIndex = 0; HandIndex < ARRAY_COUNT( VirtualHands ); ++HandIndex )
-		{
-			FVirtualHand& Hand = VirtualHands[ HandIndex ];
-			Hand.bIsHovering = false;
-			Hand.bIsHoveringOverUI = false;
-
-			// Extensibility: Allow other modes to get a shot at this event before we handle it
-			bool bHandled = false;
-			FVector HoverImpactPoint = FVector::ZeroVector;
-			bool bIsHoveringOverUI = false;
-			OnVRHoverUpdateEvent.Broadcast( *ViewportClient, HandIndex, /* In/Out */ HoverImpactPoint, /* In/Out */ bIsHoveringOverUI, /* In/Out */ bHandled );
-
-			if( bHandled )
-			{
-				Hand.bIsHovering = true;
-				Hand.bIsHoveringOverUI = bIsHoveringOverUI;
-				Hand.HoverLocation = HoverImpactPoint;
-			}
-		}
+		SpawnAvatarMeshActor();
 	}
 
+	TickHandle.Broadcast( DeltaTime );
+
 	WorldInteraction->Tick( ViewportClient, DeltaTime );
 
 	UISystem->Tick( ViewportClient, DeltaTime );
 
 	// Update avatar meshes
 	{
-		if( AvatarMeshActor == nullptr )
-		{
-			SpawnAvatarMeshActor();
-		}
-
 		// Move our avatar mesh along with the room.  We need our hand components to remain the same coordinate space as the 
 		{
 			const FTransform RoomTransform = GetRoomTransform();
@@ -996,140 +742,12 @@ void FVREditorMode::Tick( FEditorViewportClient* ViewportClient, float DeltaTime
 
 		// Scale the grid so that it stays the same size in the tracking space
 		WorldMovementGridMeshComponent->SetRelativeScale3D( FVector( GetWorldScaleFactor() ) * VREd::GridScaleMultiplier->GetFloat() );
-
 		WorldMovementGridMeshComponent->SetRelativeLocation( FVector( GetWorldScaleFactor() ) * FVector( 0.0f, 0.0f, VREd::GridHeightOffset->GetFloat() ) );
 
-		for( int32 HandIndex = 0; HandIndex < ARRAY_COUNT( VirtualHands ); ++HandIndex )
-		{
-			FVirtualHand& Hand = VirtualHands[ HandIndex ];
-
-			// @todo vreditor: Manually ticking motion controller components
-			Hand.MotionControllerComponent->TickComponent( DeltaTime, ELevelTick::LEVELTICK_PauseTick, nullptr );
-
-			// The hands need to stay the same size relative to our tracking space, so we inverse compensate for world to meters scale here
-			// NOTE: We don't need to set the hand mesh location and rotation, as the MotionControllerComponent does that itself
-			if( HandIndex == VREditorConstants::RightHandIndex && 
-				GetHMDDeviceType() == EHMDDeviceType::DT_OculusRift )	// Oculus has asymmetrical controllers, so we mirror the mesh horizontally
-			{
-				Hand.HandMeshComponent->SetRelativeScale3D( FVector( GetWorldScaleFactor(), -GetWorldScaleFactor(), GetWorldScaleFactor() ) );
-			}
-			else
-			{
-				Hand.HandMeshComponent->SetRelativeScale3D( FVector( GetWorldScaleFactor() ) );
-			}
-			
-			// We don't bother drawing hand meshes if we're in "forced VR mode" (because they'll just be on top of the camera).
-			// Also, don't bother drawing hands if we're not currently tracking them.
-			if( bActuallyUsingVR && Hand.bHaveMotionController )
-			{
-				Hand.HandMeshComponent->SetVisibility( true );
-			}
-			else
-			{
-				Hand.HandMeshComponent->SetVisibility( false );
-			}
-
-			FVector LaserPointerStart, LaserPointerEnd;
-			if( GetLaserPointer( HandIndex, /* Out */ LaserPointerStart, /* Out */ LaserPointerEnd ) )
-			{
-				// Only show the laser if we're actually in VR
-				Hand.LaserPointerMeshComponent->SetVisibility( bActuallyUsingVR );
-
-
-			    // NOTE: We don't need to set the laser pointer location and rotation, as the MotionControllerComponent will do
-			    // that later in the frame.  
-
-				// If we're actively dragging something around, then we'll crop the laser length to the hover impact
-				// point.  Otherwise, we always want the laser to protrude through hovered objects, so that you can
-				// interact with translucent gizmo handles that are occluded by geometry
-				FVector LaserPointerImpactPoint = LaserPointerEnd;
-				if( Hand.bIsHovering && ( Hand.DraggingMode != EVREditorDraggingMode::Nothing || Hand.bIsHoveringOverUI ) )
-				{
-					LaserPointerImpactPoint = Hand.HoverLocation;
-				}
-
-				// Apply rotation offset to the laser direction
-				const float LaserPointerRotationOffset = GetHMDDeviceType() == EHMDDeviceType::DT_OculusRift ? VREd::OculusLaserPointerRotationOffset->GetFloat() : VREd::ViveLaserPointerRotationOffset->GetFloat();
-				Hand.LaserPointerMeshComponent->SetRelativeRotation( FRotator( LaserPointerRotationOffset, 0.0f, 0.0f ) );
-
-				const FVector LaserPointerDirection = ( LaserPointerImpactPoint - LaserPointerStart ).GetSafeNormal();
-
-				// Offset the beginning of the laser pointer a bit, so that it doesn't overlap the hand mesh
-				const float LaserPointerStartOffset =
-					GetWorldScaleFactor() *
-					( GetHMDDeviceType() == EHMDDeviceType::DT_OculusRift ? VREd::OculusLaserPointerStartOffset->GetFloat() : VREd::ViveLaserPointerStartOffset->GetFloat() );
-
-				const float LaserPointerLength = FMath::Max( 0.000001f, ( LaserPointerImpactPoint - LaserPointerStart ).Size() - LaserPointerStartOffset );
-
-				Hand.LaserPointerMeshComponent->SetRelativeLocation( FRotator( LaserPointerRotationOffset, 0.0f, 0.0f ).RotateVector( FVector( LaserPointerStartOffset, 0.0f, 0.0f ) ) );
-
-				// The laser pointer needs to stay the same size relative to our tracking space, so we inverse compensate for world to meters scale here
-				float LaserPointerRadius = VREd::LaserPointerRadius->GetFloat() * GetWorldScaleFactor();
-				float HoverMeshRadius = VREd::LaserPointerHoverBallRadius->GetFloat() * GetWorldScaleFactor();
-
-				// If we're hovering over something really close to the camera, go ahead and shrink the effect
-				// @todo vreditor: Can we make this actually just sized based on distance automatically?  The beam and impact point are basically a cone.
-				if( Hand.bIsHoveringOverUI )
-				{
-					LaserPointerRadius *= 0.35f;	// @todo vreditor tweak
-					HoverMeshRadius *= 0.35f;	// @todo vreditor tweak
-				}
-
-				Hand.LaserPointerMeshComponent->SetRelativeScale3D( FVector( LaserPointerLength, LaserPointerRadius * 2.0f, LaserPointerRadius * 2.0f ) );
-
-				if( Hand.bIsHovering )
-				{
-					// The hover effect needs to stay the same size relative to our tracking space, so we inverse compensate for world to meters scale here
-					Hand.HoverMeshComponent->SetRelativeScale3D( FVector( HoverMeshRadius * 2.0f ) );
-					Hand.HoverMeshComponent->SetVisibility( true );
-					Hand.HoverMeshComponent->SetWorldLocation( Hand.HoverLocation );
-
-					// Show the light too, unless it's on top of UI.  It looks too distracting on top of UI.
-					Hand.HoverPointLightComponent->SetVisibility( !Hand.bIsHoveringOverUI );
-
-					// Update radius for world scaling
-					Hand.HoverPointLightComponent->SetAttenuationRadius( VREd::LaserPointerLightRadius->GetFloat() * GetWorldScaleFactor() );
-
-
-					// Pull hover light back a bit from the end of the ray
-					const float PullBackAmount = VREd::LaserPointerLightPullBackDistance->GetFloat() * GetWorldScaleFactor();
-					Hand.HoverPointLightComponent->SetWorldLocation( Hand.HoverLocation - PullBackAmount * LaserPointerDirection );
-				}
-				else
-				{
-					Hand.HoverMeshComponent->SetVisibility( false );
-					Hand.HoverPointLightComponent->SetVisibility( false );
-				}
-
-
-				// Update laser pointer materials
-				{
-					// @todo vreditor: Hook up variety of colors and "crawling"
-					const bool CrawlFade = 1.0f;
-					const float CrawlSpeed = 10.0f;
-
-					SetLaserVisuals( HandIndex, GetColor( EColors::DefaultColor ) );
-
-					static FName StaticLengthParameterName( "Length" );
-					Hand.LaserPointerMID->SetScalarParameterValue( StaticLengthParameterName, LaserPointerLength );
-					Hand.TranslucentLaserPointerMID->SetScalarParameterValue( StaticLengthParameterName, LaserPointerLength );
-				}
-			}
-			else
-			{
-				Hand.LaserPointerMeshComponent->SetVisibility( false );
-				Hand.HoverMeshComponent->SetVisibility( false );
-				Hand.HoverPointLightComponent->SetVisibility( false );
-			}
-		}
-
-	
-		// Update the user scale indicator
+		// Update the user scale indicator //@todo
 		{
-			FVirtualHand LeftHand = GetVirtualHand( VREditorConstants::LeftHandIndex );
-			FVirtualHand RightHand = GetVirtualHand( VREditorConstants::RightHandIndex );
-			if ( ( LeftHand.DraggingMode == EVREditorDraggingMode::World && RightHand.DraggingMode == EVREditorDraggingMode::AssistingDrag ) ||
-				 ( LeftHand.DraggingMode == EVREditorDraggingMode::AssistingDrag && RightHand.DraggingMode == EVREditorDraggingMode::World ) )
+			if ( ( LeftHandInteractor->GetDraggingMode() == EViewportInteractionDraggingMode::World && RightHandInteractor->GetDraggingMode() == EViewportInteractionDraggingMode::AssistingDrag ) ||
+				 ( LeftHandInteractor->GetDraggingMode() == EViewportInteractionDraggingMode::AssistingDrag && RightHandInteractor->GetDraggingMode() == EViewportInteractionDraggingMode::World ) )
 			{
 				// Setting all components to be visible
 				CurrentScaleProgressMeshComponent->SetVisibility( true );
@@ -1137,8 +755,8 @@ void FVREditorMode::Tick( FEditorViewportClient* ViewportClient, float DeltaTime
 				UserScaleIndicatorText->SetVisibility( true );
 
 				// Here we calculate the distance, direction and center of the two hands
-				const FVector ScaleIndicatorStartPosition = LeftHand.Transform.GetLocation();
-				const FVector ScaleIndicatorEndPosition = RightHand.Transform.GetLocation();
+				const FVector ScaleIndicatorStartPosition = LeftHandInteractor->GetTransform().GetLocation();
+				const FVector ScaleIndicatorEndPosition = RightHandInteractor->GetTransform().GetLocation();
 				const FVector ScaleIndicatorDirection = (ScaleIndicatorEndPosition - ScaleIndicatorStartPosition).GetSafeNormal();
 
 				const float ScaleIndicatorLength = FMath::Max( 0.000001f, (ScaleIndicatorEndPosition - ScaleIndicatorStartPosition).Size() );
@@ -1162,7 +780,7 @@ void FVREditorMode::Tick( FEditorViewportClient* ViewportClient, float DeltaTime
 
 				// Setting the transform for the current scale progressbar from the center
 				{
-					const float CurrentProgressScale = (Scale * Scale) * (ProgressbarLength / (VREd::ScaleMax->GetFloat() / 100));
+					const float CurrentProgressScale = (Scale * Scale) * (ProgressbarLength / ( WorldInteraction->GetMaxScale() / 100));
 					const FVector CurrentProgressStart = MiddleLocation - (CurrentProgressScale * 0.5f) * ScaleIndicatorDirection;
 					
 					CurrentScaleProgressMeshComponent->SetWorldTransform( FTransform (  ScaleIndicatorDirection.ToOrientationRotator(),
@@ -1197,53 +815,6 @@ void FVREditorMode::Tick( FEditorViewportClient* ViewportClient, float DeltaTime
 		}
 	}
 
-	// Updating laser colors for both hands
-	{
-		for( int32 HandIndex = 0; HandIndex < VREditorConstants::NumVirtualHands; ++HandIndex )
-		{
-			FLinearColor ResultColor = GetColor( EColors::DefaultColor );
-			float CrawlSpeed = 0.0f;
-			float CrawlFade = 0.0f;
-
-			const FVirtualHand& Hand = GetVirtualHand( HandIndex );
-			const FVirtualHand& OtherHand = GetOtherHand( HandIndex );
-
-			const bool bIsDraggingWorldWithTwoHands =
-				( Hand.DraggingMode == EVREditorDraggingMode::AssistingDrag && OtherHand.DraggingMode == EVREditorDraggingMode::World ) ||
-				( Hand.DraggingMode == EVREditorDraggingMode::World && OtherHand.DraggingMode == EVREditorDraggingMode::AssistingDrag );
-
-			if( bIsDraggingWorldWithTwoHands )
-			{
-				ResultColor = GetColor( EColors::WorldDraggingColor_TwoHands );
-			}
-			else if( Hand.DraggingMode == EVREditorDraggingMode::World )
-			{
-				if( WorldInteraction->IsTeleporting() )
-				{
-					ResultColor = GetColor( EColors::TeleportColor );
-				}
-				else
-				{
-					// We can teleport in this mode, so animate the laser a bit
-					CrawlFade = 1.0f;
-					CrawlSpeed = 5.0f;
-					ResultColor = GetColor( EColors::WorldDraggingColor_OneHand );
-				}
-			}
-			else if( Hand.DraggingMode == EVREditorDraggingMode::ActorsAtLaserImpact ||
-				Hand.DraggingMode == EVREditorDraggingMode::Material ||
-				Hand.DraggingMode == EVREditorDraggingMode::ActorsFreely ||
-				Hand.DraggingMode == EVREditorDraggingMode::ActorsWithGizmo ||
-				Hand.DraggingMode == EVREditorDraggingMode::AssistingDrag || 
-				( GetUISystem().IsDraggingDockUI() && HandIndex == GetUISystem().GetDraggingDockUIHandIndex() ) )
-			{
-				ResultColor = GetColor( EColors::SelectionColor );
-			}
-
-			SetLaserVisuals( HandIndex, ResultColor, CrawlFade, CrawlSpeed );
-		}
-	}
-
 	// Updating the scale and intensity of the flashlight according to the world scale
 	if (FlashlightComponent)
 	{
@@ -1252,26 +823,24 @@ void FVREditorMode::Tick( FEditorViewportClient* ViewportClient, float DeltaTime
 		float UpdatedFalloffExponent = FMath::Clamp(CurrentFalloffExponent / GetWorldScaleFactor(), 2.0f, 16.0f);
 		FlashlightComponent->SetLightFalloffExponent(UpdatedFalloffExponent);
 	}
-	
 
-
-	// Updating the opacity and visibility of the grid according to the controllers
+	// Updating the opacity and visibility of the grid according to the controllers //@todo
 	{
 		if( VREd::ShowMovementGrid->GetInt() == 1)
 		{
 			// Show the grid full opacity when dragging or scaling
 			float GoalOpacity = 0.f;
-			if ( GetVirtualHand( VREditorConstants::LeftHandIndex ).DraggingMode == EVREditorDraggingMode::World || GetVirtualHand( VREditorConstants::RightHandIndex ).DraggingMode == EVREditorDraggingMode::World )
+			if ( LeftHandInteractor->GetDraggingMode() == EViewportInteractionDraggingMode::World || RightHandInteractor->GetDraggingMode() == EViewportInteractionDraggingMode::World )
 			{
 				GoalOpacity = 1.0f;
 			}
-			else if ( ( GetVirtualHand( VREditorConstants::LeftHandIndex ).LastDraggingMode == EVREditorDraggingMode::World && !GetVirtualHand( VREditorConstants::LeftHandIndex ).DragTranslationVelocity.IsNearlyZero( VREd::GridMovementTolerance->GetFloat() ) ) )
+			else if ( ( LeftHandInteractor->GetLastDraggingMode() == EViewportInteractionDraggingMode::World && !LeftHandInteractor->GetDragTranslationVelocity().IsNearlyZero( VREd::GridMovementTolerance->GetFloat() ) ) )
 			{
-				GoalOpacity = FMath::Clamp( GetVirtualHand( VREditorConstants::LeftHandIndex ).DragTranslationVelocity.Size() / VREd::GridFadeStartVelocity->GetFloat(), 0.0f, 1.0f );
+				GoalOpacity = FMath::Clamp( LeftHandInteractor->GetDragTranslationVelocity().Size() / VREd::GridFadeStartVelocity->GetFloat(), 0.0f, 1.0f );
 			}
-			else if( ( GetVirtualHand( VREditorConstants::RightHandIndex ).LastDraggingMode == EVREditorDraggingMode::World && !GetVirtualHand( VREditorConstants::RightHandIndex ).DragTranslationVelocity.IsNearlyZero( VREd::GridMovementTolerance->GetFloat() ) ) )
+			else if( ( RightHandInteractor->GetLastDraggingMode() == EViewportInteractionDraggingMode::World && !RightHandInteractor->GetDragTranslationVelocity().IsNearlyZero( VREd::GridMovementTolerance->GetFloat() ) ) )
 			{
-				GoalOpacity = FMath::Clamp( GetVirtualHand( VREditorConstants::RightHandIndex ).DragTranslationVelocity.Size() / VREd::GridFadeStartVelocity->GetFloat(), 0.0f, 1.0f );
+				GoalOpacity = FMath::Clamp( RightHandInteractor->GetDragTranslationVelocity().Size() / VREd::GridFadeStartVelocity->GetFloat(), 0.0f, 1.0f );
 			}
 
 			// Check the current opacity and add or subtract to reach the goal
@@ -1311,7 +880,6 @@ void FVREditorMode::Tick( FEditorViewportClient* ViewportClient, float DeltaTime
 		}
 	}
 
-
 	// Apply a post process effect while the user is moving the world around.  The effect "greys out" any pixels
 	// that are not nearby, and completely hides the skybox and other very distant objects.  This is to help
 	// prevent simulation sickness when moving/rotating/scaling the entire world around them.
@@ -1384,10 +952,6 @@ void FVREditorMode::Tick( FEditorViewportClient* ViewportClient, float DeltaTime
 		SetRoomTransform( ResultToWorld );
 	}
 
-
-	// Update floating help text
-	UpdateHelpLabels();
-
 	StopOldHapticEffects();
 
 	bFirstTick = false;
@@ -1402,156 +966,13 @@ bool FVREditorMode::InputKey(FEditorViewportClient* ViewportClient, FViewport* V
 		return InputKey(GetLevelViewportPossessedForVR().GetViewportClient().Get(), Viewport, Key, Event);
 	}
 
-	return FEdMode::InputKey(ViewportClient, Viewport, Key, Event);
-}
-
-bool FVREditorMode::HandleInputKey(FEditorViewportClient* ViewportClient, FViewport* Viewport, FKey Key, EInputEvent Event)
-{
-	bool bHandled = false;
-
-	FVRAction* KnownAction = KeyToActionMap.Find( Key );
-	if( KnownAction != nullptr )
-	{
-		FVRAction VRAction( KnownAction->ActionType, GetHandIndexFromKey( Key ) );
-		if( VRAction.HandIndex != INDEX_NONE )
-		{
-
-			FVirtualHand& Hand = VirtualHands[ VRAction.HandIndex ];
-
-			bool bShouldAbsorbEvent = false;
-
-			// If "SelectAndMove" was pressed, we need to make sure that "SelectAndMove_LightlyPressed" is no longer
-			// in effect before we start handling "SelectAndMove".
-			if( VRAction.ActionType == EVRActionType::SelectAndMove )
-			{
-				if( Hand.bIsTriggerAtLeastLightlyPressed )
-				{
-					if( Event == IE_Pressed )
-					{
-						// How long since the trigger was lightly pressed?
-						const float TimeSinceLightlyPressed = (float)( FPlatformTime::Seconds() - Hand.RealTimeTriggerWasLightlyPressed );
-						if( !Hand.bIsClickingOnUI &&	// @todo vreditor: UI clicks happen with light presses; we never want to convert to a hard press!
-							Hand.DraggingMode != EVREditorDraggingMode::Material &&	// @todo vreditor: Material dragging happens with light presses, don't convert to a hard press!
-							Hand.DraggingMode != EVREditorDraggingMode::ActorsAtLaserImpact &&	// @todo vreditor: Actor dragging happens with light presses, don't convert to a hard press!
-							( !Hand.bAllowTriggerLightPressLocking || TimeSinceLightlyPressed < VREd::TriggerLightlyPressedLockTime->GetFloat() ) )
-						{
-							Hand.bIsTriggerAtLeastLightlyPressed = false;
-							Hand.bHasTriggerBeenReleasedSinceLastPress = false;
-
-							// Synthesize an input key for releasing the light press
-							// NOTE: Intentionally re-entrant call here.
-							const EInputEvent InputEvent = IE_Released;
-							const bool bWasLightReleaseHandled = HandleInputKey(ViewportClient, Viewport, VRAction.HandIndex == VREditorConstants::LeftHandIndex ? VREditorKeyNames::MotionController_Left_LightlyPressedTriggerAxis : VREditorKeyNames::MotionController_Right_LightlyPressedTriggerAxis, InputEvent);
-						}
-						else
-						{
-							// The button was held in a "lightly pressed" state for a long enough time that we should just continue to
-							// treat it like a light press.  This makes it much easier to hold it in this state when you need to!
-							bShouldAbsorbEvent = true;
-						}
-					}
-					else
-					{
-						// Absorb the release of the SelectAndMove event if we are still in a lightly pressed state
-						bShouldAbsorbEvent = true;
-					}
-				}
-			}
-
-			if( bShouldAbsorbEvent )
-			{
-				bHandled = true;
-			}
-			else
-			{
-				// Update touch state
-				if( VRAction.ActionType == EVRActionType::Touch )
-				{
-					if( Event == IE_Pressed )
-					{
-						Hand.bIsTouchingTrackpad = true;
-					}
-					else if( Event == IE_Released )
-					{
-						Hand.bIsTouchingTrackpad = false;
-						Hand.bIsTrackpadPositionValid[ 0 ] = false;
-						Hand.bIsTrackpadPositionValid[ 1 ] = false;
-					}
-				}
-
-				// Update modifier state
-				if( VRAction.ActionType == EVRActionType::Modifier )
-				{
-					bHandled = true;
-					if( Event == IE_Pressed )
-					{
-						Hand.bIsModifierPressed = true;
-					}
-					else if( Event == IE_Released )
-					{
-						Hand.bIsModifierPressed = false;
-					}
-				}
-
-				// Undo/redo
-				if( VRAction.ActionType == EVRActionType::Undo )
-				{
-					if( Event == IE_Pressed || Event == IE_Repeat )
-					{
-						Undo();
-					}
-					bHandled = true;
-				}
-				else if( VRAction.ActionType == EVRActionType::Redo )
-				{
-					if( Event == IE_Pressed || Event == IE_Repeat )
-					{
-						Redo();
-					}
-					bHandled = true;
-				}
-
-				if( !bHandled )
-				{
-					// Extensibility: Allow other modes to get a shot at this event before we handle it
-					OnVRActionEvent.Broadcast( *ViewportClient, VRAction, Event, Hand.IsInputCaptured[ (int32)VRAction.ActionType ], /* In/Out */ bHandled );
-				}
-
-				if( !bHandled )
-				{
-					// If "select and move" was pressed but not handled, go ahead and deselect everything
-					if( VRAction.ActionType == EVRActionType::SelectAndMove && Event == IE_Pressed )
-					{
-						GEditor->SelectNone( true, true );
-					}
-				}
-
-				// Always handle mouse button presses when forcing VR mode, so that the editor doesn't select non-interactable actors
-				// on mouse down, etc.
-				if( !bActuallyUsingVR && VREd::UseMouseAsHandInForcedVRMode->GetInt() != 0 )
-				{
-					if( Key == EKeys::LeftMouseButton )
-					{
-						bHandled = true;
-					}
-				}
-			}
-
-			ApplyButtonPressColors( VRAction, Event );
-		}
-	}
-
-	if( Key == EKeys::Escape )	// @todo vreditor: Should be a bindable action, not hard coded
+	if ( Key == EKeys::Escape )
 	{
 		// User hit escape, so bail out of VR mode
 		StartExitingVRMode();
-
-		bHandled = true;
 	}
 
-	StopOldHapticEffects();
-
-	return bHandled;
+	return FEdMode::InputKey(ViewportClient, Viewport, Key, Event);
 }
 
 
@@ -1566,111 +987,11 @@ bool FVREditorMode::InputAxis(FEditorViewportClient* ViewportClient, FViewport*
 	return FEdMode::InputAxis(ViewportClient, Viewport, ControllerId, Key, Delta, DeltaTime);
 }
 
-bool FVREditorMode::HandleInputAxis(FEditorViewportClient* ViewportClient, FViewport* Viewport, int32 ControllerId, FKey Key, float Delta, float DeltaTime)
+bool FVREditorMode::IsCompatibleWith(FEditorModeID OtherModeID) const
 {
-	bool bHandled = false;
-
-	FVRAction* KnownAction = KeyToActionMap.Find( Key );
-	if( KnownAction != nullptr )	// Ignore key repeats
-	{
-		FVRAction VRAction( KnownAction->ActionType, GetHandIndexFromKey( Key ) );
-		FVirtualHand& Hand = VirtualHands[ VRAction.HandIndex ];
-
-		if( VRAction.HandIndex != INDEX_NONE )
-		{
-			if( VRAction.ActionType == EVRActionType::TriggerAxis )
-			{
-				// Store latest trigger value amount
-				Hand.SelectAndMoveTriggerValue = Delta;
-
-				// Synthesize "lightly pressed" events for the trigger
-				{
-					const bool bIsFullPressAlreadyCapturing = Hand.IsInputCaptured[ (int32)EVRActionType::SelectAndMove ];
-					if( !bIsFullPressAlreadyCapturing &&		// Don't fire if we're already capturing input for a full press
-						!Hand.bIsTriggerAtLeastLightlyPressed &&	// Don't fire unless we are already pressed
-						Hand.bHasTriggerBeenReleasedSinceLastPress &&	// Only if we've been fully released since the last time we fired
-						Delta >= VREd::TriggerLightlyPressedThreshold->GetFloat() )
-					{
-						Hand.bIsTriggerAtLeastLightlyPressed = true;
-						Hand.bAllowTriggerLightPressLocking = true;
-						Hand.RealTimeTriggerWasLightlyPressed = FPlatformTime::Seconds();
-						Hand.bHasTriggerBeenReleasedSinceLastPress = false;
-
-						// Synthesize an input key for this light press
-						const EInputEvent InputEvent = IE_Pressed;
-						const bool bWasLightPressHandled = HandleInputKey(ViewportClient, Viewport, VRAction.HandIndex == VREditorConstants::LeftHandIndex ? VREditorKeyNames::MotionController_Left_LightlyPressedTriggerAxis : VREditorKeyNames::MotionController_Right_LightlyPressedTriggerAxis, InputEvent);
-					}
-					else if( Hand.bIsTriggerAtLeastLightlyPressed && Delta < VREd::TriggerLightlyPressedThreshold->GetFloat() )
-					{
-						Hand.bIsTriggerAtLeastLightlyPressed = false;
-
-						// Synthesize an input key for this light press
-						const EInputEvent InputEvent = IE_Released;
-						const bool bWasLightReleaseHandled = HandleInputKey(ViewportClient, Viewport, VRAction.HandIndex == VREditorConstants::LeftHandIndex ? VREditorKeyNames::MotionController_Left_LightlyPressedTriggerAxis : VREditorKeyNames::MotionController_Right_LightlyPressedTriggerAxis, InputEvent);
-					}
-				}
-
-				if( Delta < VREd::TriggerDeadZone->GetFloat() )
-				{
-					Hand.bHasTriggerBeenReleasedSinceLastPress = true;
-				}
-
-				// Synthesize "fully pressed" events for the trigger
-				{
-					if( !Hand.bIsTriggerFullyPressed &&	// Don't fire unless we are already pressed
-						Delta >= VREd::TriggerFullyPressedThreshold->GetFloat() )
-					{
-						Hand.bIsTriggerFullyPressed = true;
-
-						// Synthesize an input key for this full press
-						const EInputEvent InputEvent = IE_Pressed;
-						HandleInputKey(ViewportClient, Viewport, VRAction.HandIndex == VREditorConstants::LeftHandIndex ? VREditorKeyNames::MotionController_Left_FullyPressedTriggerAxis : VREditorKeyNames::MotionController_Right_FullyPressedTriggerAxis, InputEvent);
-					}
-					else if( Hand.bIsTriggerFullyPressed && Delta < VREd::TriggerFullyPressedReleaseThreshold->GetFloat() )
-					{
-						Hand.bIsTriggerFullyPressed = false;
-
-						// Synthesize an input key for this full press
-						const EInputEvent InputEvent = IE_Released;
-						HandleInputKey(ViewportClient, Viewport, VRAction.HandIndex == VREditorConstants::LeftHandIndex ? VREditorKeyNames::MotionController_Left_FullyPressedTriggerAxis : VREditorKeyNames::MotionController_Right_FullyPressedTriggerAxis, InputEvent);
-					}
-				}
-			}
-
-			if( VRAction.ActionType == EVRActionType::TrackpadPositionX )
-			{
-				Hand.LastTrackpadPosition.X = Hand.bIsTrackpadPositionValid[0] ? Hand.TrackpadPosition.X : Delta;
-				Hand.LastTrackpadPositionUpdateTime = FTimespan::FromSeconds( FPlatformTime::Seconds() );
-				Hand.TrackpadPosition.X = Delta;
-				Hand.bIsTrackpadPositionValid[0] = true;
-			}
-
-			if( VRAction.ActionType == EVRActionType::TrackpadPositionY )
-			{
-				if( VREd::InvertTrackpadVertical->GetInt() != 0 )
-				{
-					Delta = -Delta;	// Y axis is inverted from HMD
-				}
-
-				Hand.LastTrackpadPosition.Y = Hand.bIsTrackpadPositionValid[1] ? Hand.TrackpadPosition.Y : Delta;
-				Hand.LastTrackpadPositionUpdateTime = FTimespan::FromSeconds( FPlatformTime::Seconds() );
-				Hand.TrackpadPosition.Y = Delta;
-				Hand.bIsTrackpadPositionValid[1] = true;
-			}
-		}
-	}
-
-	// ...
-	StopOldHapticEffects();
-
-	return bHandled;
-}
-
-bool FVREditorMode::IsCompatibleWith(FEditorModeID OtherModeID) const
-{
-	// We are compatible with all other modes!
-	return true;
-}
+	// We are compatible with all other modes!
+	return true;
+}
 
 
 void FVREditorMode::AddReferencedObjects( FReferenceCollector& Collector )
@@ -1679,38 +1000,23 @@ void FVREditorMode::AddReferencedObjects( FReferenceCollector& Collector )
 	Collector.AddReferencedObject( HeadMeshComponent );
 	Collector.AddReferencedObject( FlashlightComponent );
 	Collector.AddReferencedObject( WorldMovementGridMeshComponent );
-	Collector.AddReferencedObject( SnapGridActor );
-	Collector.AddReferencedObject( SnapGridMeshComponent );
-	Collector.AddReferencedObject( SnapGridMID );
 	Collector.AddReferencedObject( PostProcessComponent );
-
-	for( FVirtualHand& Hand : VirtualHands )
-	{
-		Collector.AddReferencedObject( Hand.MotionControllerComponent );
-		Collector.AddReferencedObject( Hand.HandMeshComponent );
-		Collector.AddReferencedObject( Hand.HandMeshMID );
-		Collector.AddReferencedObject( Hand.LaserPointerMeshComponent );
-		Collector.AddReferencedObject( Hand.LaserPointerMID );
-		Collector.AddReferencedObject( Hand.TranslucentLaserPointerMID );
-		Collector.AddReferencedObject( Hand.HoverMeshComponent );
-		Collector.AddReferencedObject( Hand.HoverPointLightComponent );
-
-		Collector.AddReferencedObjects( Hand.HelpLabels );
-
-		Collector.AddReferencedObject( Hand.HoveringOverWidgetComponent );
-	}
-
 	Collector.AddReferencedObject( WorldMovementGridMID );
 	Collector.AddReferencedObject( WorldMovementPostProcessMaterial );
 	Collector.AddReferencedObject( ScaleProgressMeshComponent );
 	Collector.AddReferencedObject( CurrentScaleProgressMeshComponent );
-	Collector.AddReferencedObject( UserScaleIndicatorText );
+	Collector.AddReferencedObject( UserScaleIndicatorText );	
+	Collector.AddReferencedObject( UISystem );
+	Collector.AddReferencedObject( WorldInteraction );
+	Collector.AddReferencedObject( TeleporterSystem );
+	Collector.AddReferencedObject( LeftHandInteractor );
+	Collector.AddReferencedObject( RightHandInteractor );
 }
 
 
 void FVREditorMode::Render( const FSceneView* SceneView, FViewport* Viewport, FPrimitiveDrawInterface* PDI )
 {
-	StopOldHapticEffects();
+	//StopOldHapticEffects(); //@todo vreditor
 
 	FEdMode::Render( SceneView, Viewport, PDI );
 
@@ -1718,316 +1024,58 @@ void FVREditorMode::Render( const FSceneView* SceneView, FViewport* Viewport, FP
 	{
 		// Let our subsystems render, too
 		UISystem->Render( SceneView, Viewport, PDI );
-		WorldInteraction->Render( SceneView, Viewport, PDI );
 	}
 }
 
+/************************************************************************/
+/* IVREditorMode interface                                              */
+/************************************************************************/
 
-float FVREditorMode::GetLaserPointerMaxLength() const
+AActor* FVREditorMode::GetAvatarMeshActor()
 {
-	return VREd::LaserPointerMaxLength->GetFloat() * GetWorldScaleFactor();
+	return AvatarMeshActor;
 }
 
-
-FTransform FVREditorMode::GetRoomTransform() const
+UWorld* FVREditorMode::GetWorld() const
 {
-	const FLevelEditorViewportClient& ViewportClient = GetLevelViewportPossessedForVR().GetLevelViewportClient();
-	const FTransform RoomTransform( 
-		ViewportClient.GetViewRotation().Quaternion(), 
-		ViewportClient.GetViewLocation(),
-		FVector( 1.0f ) );
-	return RoomTransform;
-}
-
-
-void FVREditorMode::SetRoomTransform( const FTransform& NewRoomTransform )
-{
-	FLevelEditorViewportClient& ViewportClient = GetLevelViewportPossessedForVR().GetLevelViewportClient();
-	ViewportClient.SetViewLocation( NewRoomTransform.GetLocation() );
-	ViewportClient.SetViewRotation( NewRoomTransform.GetRotation().Rotator() );
-
-	// Forcibly dirty the viewport camera location
-	const bool bDollyCamera = false;
-	ViewportClient.MoveViewportCamera( FVector::ZeroVector, FRotator::ZeroRotator, bDollyCamera );
+	return WorldInteraction->GetViewportWorld();
 }
 
-
-bool FVREditorMode::GetHandTransformAndForwardVector( int32 HandIndex, FTransform& OutHandTransform, FVector& OutForwardVector ) const
-{
-	FVREditorMode* MutableThis = const_cast<FVREditorMode*>( this );
-	MutableThis->PollInputIfNeeded();
-
-	const FVirtualHand& Hand = VirtualHands[ HandIndex ];
-	if( Hand.bHaveMotionController )
-	{
-		OutHandTransform = Hand.Transform;
-
-		const float LaserPointerRotationOffset = GetHMDDeviceType() == EHMDDeviceType::DT_OculusRift ? VREd::OculusLaserPointerRotationOffset->GetFloat() : VREd::ViveLaserPointerRotationOffset->GetFloat();
-		OutForwardVector = OutHandTransform.GetRotation().RotateVector( FRotator( LaserPointerRotationOffset, 0.0f, 0.0f ).RotateVector( FVector( 1.0f, 0.0f, 0.0f ) ) );
-
-		return true;
-	}
-
-	return false;
-}
-
-
-float FVREditorMode::GetMaxScale()
+FTransform FVREditorMode::GetRoomTransform() const
 {
-	return VREd::ScaleMax->GetFloat();
+	return WorldInteraction->GetRoomTransform();
 }
 
-float FVREditorMode::GetMinScale()
+void FVREditorMode::SetRoomTransform( const FTransform& NewRoomTransform )
 {
-	return VREd::ScaleMin->GetFloat();
+	WorldInteraction->SetRoomTransform( NewRoomTransform );
 }
 
-void FVREditorMode::SetWorldToMetersScale( const float NewWorldToMetersScale )
+FTransform FVREditorMode::GetRoomSpaceHeadTransform() const
 {
-	// @todo vreditor: This is bad because we're clobbering the world settings which will be saved with the map.  Instead we need to 
-	// be able to apply an override before the scene view gets it
-
-	ENGINE_API extern float GNewWorldToMetersScale;
-	GNewWorldToMetersScale = NewWorldToMetersScale;
+	return WorldInteraction->GetRoomSpaceHeadTransform();
 }
 
-bool FVREditorMode::GetLaserPointer(const int32 HandIndex, FVector& LaserPointerStart, FVector& LaserPointerEnd, const bool bEvenIfUIIsInFront, const float LaserLengthOverride) const
+FTransform FVREditorMode::GetHeadTransform() const
 {
-	const FVirtualHand& Hand = VirtualHands[ HandIndex ];
-
-	// If we're currently grabbing the world with both hands, then the laser pointer is not available
-	if (Hand.bHaveMotionController &&
-		!( Hand.DraggingMode == EVREditorDraggingMode::World && GetOtherHand( HandIndex ).DraggingMode == EVREditorDraggingMode::AssistingDrag ) &&
-		!( Hand.DraggingMode == EVREditorDraggingMode::AssistingDrag && GetOtherHand( HandIndex ).DraggingMode == EVREditorDraggingMode::World ) )  
-	{
-		// If we have UI attached to us, don't allow a laser pointer
-		if( bEvenIfUIIsInFront || !Hand.bHasUIInFront )
-		{
-			FTransform HandTransform;
-			FVector HandForwardVector;
-			if( GetHandTransformAndForwardVector( HandIndex, HandTransform, HandForwardVector ) )
-			{
-				LaserPointerStart = HandTransform.GetLocation();
-
-				float LaserLength = 0.0f;
-				if( LaserLengthOverride == 0.0f )
-				{
-					LaserLength = GetLaserPointerMaxLength();
-				}
-				else
-				{
-					LaserLength = LaserLengthOverride;
-				}
-
-				LaserPointerEnd = HandTransform.GetLocation() + HandForwardVector * LaserLength;
-
-				return true;
-			}
-		}
-	}
-
-	return false;
+	return WorldInteraction->GetHeadTransform();
 }
 
-
-bool FVREditorMode::IsInputCaptured( const FVRAction VRAction ) const
+const UViewportWorldInteraction& FVREditorMode::GetWorldInteraction() const
 {
-	check( VRAction.HandIndex >= 0 && VRAction.HandIndex < VREditorConstants::NumVirtualHands );
-	check( (int32)VRAction.ActionType >= 0 && VRAction.ActionType < EVRActionType::TotalActionTypes );
-	return VirtualHands[ VRAction.HandIndex ].IsInputCaptured[ (int32)VRAction.ActionType ];
+	return *WorldInteraction;
 }
 
-
-IVREditorMode::FHandInfo FVREditorMode::GetHandInfo( const int32 HandIndex ) const
+UViewportWorldInteraction& FVREditorMode::GetWorldInteraction()
 {
-	const FVirtualHand& Hand = GetVirtualHand( HandIndex );
-
-	IVREditorMode::FHandInfo HandInfo;
-	HandInfo.bIsModifierPressed = Hand.bIsModifierPressed;
-	HandInfo.SelectAndMoveTriggerValue = Hand.SelectAndMoveTriggerValue;
-	HandInfo.bIsHoveringOverUI = Hand.bIsHoveringOverUI;
-	return HandInfo;
+	return *WorldInteraction;
 }
 
-
-void FVREditorMode::SetAllowTriggerLightPressLocking( const int32 HandIndex, const bool bAllow )
+bool FVREditorMode::IsFullyInitialized() const
 {
-	GetVirtualHand( HandIndex ).bAllowTriggerLightPressLocking = bAllow;
+	return bIsFullyInitialized;
 }
 
-
-FHitResult FVREditorMode::GetHitResultFromLaserPointer( int32 HandIndex, TArray<AActor*>* OptionalListOfIgnoredActors, const bool bIgnoreGizmos, const bool bEvenIfUIIsInFront, const float LaserLengthOverride )
-{
-	FHitResult BestHitResult;
-	bool bBestHitResultSoFarIsUI = false;
-
-	FVector LaserPointerStart, LaserPointerEnd;
-	if( GetLaserPointer( HandIndex, LaserPointerStart, LaserPointerEnd, bEvenIfUIIsInFront, LaserLengthOverride ) )
-	{
-		// Twice twice.  Once for editor gizmos which are "on top" and always take precedence, then a second time
-		// for all of the scene objects
-
-		for( int32 PassIndex = bIgnoreGizmos ? 1 : 0; PassIndex < 2; ++PassIndex )
-		{
-			const bool bOnlyEditorGizmos = ( PassIndex == 0 );
-
-			const bool bTraceComplex = true;
-			FCollisionQueryParams TraceParams( NAME_None, bTraceComplex, nullptr );
-			
-			const FCollisionResponseParams& ResponseParam = FCollisionResponseParams::DefaultResponseParam;
-
-			const ECollisionChannel CollisionChannel = bOnlyEditorGizmos ? COLLISION_GIZMO : ECC_Visibility;
-
-			// Don't trace against our own head/hand meshes
-			TraceParams.AddIgnoredActor( AvatarMeshActor );
-
-			// Don't trace against our snap grid
-			TraceParams.AddIgnoredActor( SnapGridActor );
-
-			if( OptionalListOfIgnoredActors != nullptr )
-			{
-				TraceParams.AddIgnoredActors( *OptionalListOfIgnoredActors );
-			}
-
-			bool bHit = false;
-			FHitResult HitResult;
-			if( bOnlyEditorGizmos )
-			{
-				bHit = GetWorld()->LineTraceSingleByChannel( HitResult, LaserPointerStart, LaserPointerEnd, CollisionChannel, TraceParams, ResponseParam );
-			}
-			else
-			{
-				FCollisionObjectQueryParams EverythingButGizmos( FCollisionObjectQueryParams::AllObjects );
-				EverythingButGizmos.RemoveObjectTypesToQuery( COLLISION_GIZMO );
-				bHit = GetWorld()->LineTraceSingleByObjectType( HitResult, LaserPointerStart, LaserPointerEnd, EverythingButGizmos, TraceParams );
-			}
-			if( bHit )
-			{
-				// Is this better than what we have already?  Always prefer transform gizmos even if they were not using
-				// COLLISION_GIZMO (some gizmo handles are opaque and use ECC_Visibility as their collision channel)
-				// NOTE: We're treating components of floating UI actors as "gizmos" for the purpose of hit testing as long as the component is not the actual UI widget component
-				const bool bHitResultIsUI = UISystem->IsWidgetAnEditorUIWidget( HitResult.GetComponent() );
-				const bool bHitResultIsGizmo = HitResult.GetActor() != nullptr && ( ( HitResult.GetActor() == WorldInteraction->GetTransformGizmoActor() ) || ( !bHitResultIsUI && HitResult.GetActor()->IsA( AVREditorFloatingUI::StaticClass() ) ) );
-				const bool bBestHitResultSoFarIsGizmo = BestHitResult.GetActor() != nullptr && ( ( BestHitResult.GetActor() == WorldInteraction->GetTransformGizmoActor() ) || ( !bBestHitResultSoFarIsUI && BestHitResult.GetActor()->IsA( AVREditorFloatingUI::StaticClass() ) ) );
-				if( BestHitResult.GetActor() == nullptr ||	// Don't have anything yet?
-					( bHitResultIsUI && bBestHitResultSoFarIsGizmo ) ||
-					( bHitResultIsGizmo && !bBestHitResultSoFarIsUI ) )	// Always prefer gizmos, unless clicking on UI
-				{
-					BestHitResult = HitResult;
-					bBestHitResultSoFarIsUI = bHitResultIsUI;
-				}
-			}
-		}
-	}
-
-	return BestHitResult;
-}
-
-
-int32 FVREditorMode::GetHandIndexFromKey( const FKey& InKey ) const
-{
-	const FVRAction* Action = KeyToActionMap.Find( InKey );
-	return Action ? Action->HandIndex : INDEX_NONE;
-}
-
-
-void FVREditorMode::PollInputIfNeeded()
-{
-	if( LastFrameNumberInputWasPolled != GFrameNumber )
-	{
-		PollInputFromMotionControllers();
-		LastFrameNumberInputWasPolled = GFrameNumber;
-	}
-}
-
-
-void FVREditorMode::PollInputFromMotionControllers()
-{
-	const FTransform RoomTransform = GetRoomTransform();
-	for( int32 HandIndex = 0; HandIndex < VREditorConstants::NumVirtualHands; ++HandIndex )
-	{
-		FVirtualHand& Hand = VirtualHands[ HandIndex ];
-
-		Hand.bHaveMotionController = false;
-
-		Hand.LastTransform = Hand.Transform;
-		Hand.LastRoomSpaceTransform = Hand.RoomSpaceTransform;
-
-		// If we're in forced VR mode and were asked to use the mouse instead of motion controllers, go
-		// ahead and poll the mouse
-		if( !bActuallyUsingVR && VREd::UseMouseAsHandInForcedVRMode->GetInt() != 0 )
-		{
-			// The mouse will always be our left hand
-			if( HandIndex == VREditorConstants::LeftHandIndex )
-			{
-				Hand.bHaveMotionController = true;
-
-
-				FLevelEditorViewportClient& ViewportClient = this->VREditorLevelViewportWeakPtr.Pin()->GetLevelViewportClient();
-				FViewport* Viewport = ViewportClient.Viewport;
-				if( Viewport->GetSizeXY().GetMin() > 0 )
-				{
-					FSceneViewFamilyContext ViewFamily( FSceneViewFamily::ConstructionValues(
-						Viewport,
-						ViewportClient.GetScene(),
-						ViewportClient.EngineShowFlags )
-						.SetRealtimeUpdate( ViewportClient.IsRealtime() ) );
-					FSceneView* SceneView = ViewportClient.CalcSceneView( &ViewFamily );
-
-					const FIntPoint MousePosition( Viewport->GetMouseX(), Viewport->GetMouseY() );
-
-					{
-						const FViewportCursorLocation MouseViewportRay( SceneView, &ViewportClient, MousePosition.X, MousePosition.Y );
-						
-						const FVector WorldSpaceLocation = MouseViewportRay.GetOrigin();
-						const FQuat WorldSpaceOrientation = MouseViewportRay.GetDirection().ToOrientationQuat();
-						const FTransform HandToWorldTransform( WorldSpaceOrientation, WorldSpaceLocation, FVector( 1.0f ) );
-						
-						const FTransform WorldToRoomTransform = RoomTransform.Inverse();
-						
-						Hand.RoomSpaceTransform = HandToWorldTransform * WorldToRoomTransform;
-					}
-				}
-
-				Hand.Transform = Hand.RoomSpaceTransform * RoomTransform;
-			}
-			else
-			{
-				// Skip the right hand when using the mouse
-				continue;
-			}
-		}
-
-
-		// Generic motion controllers
-		if( !Hand.bHaveMotionController )
-		{
-			TArray<IMotionController*> MotionControllers = IModularFeatures::Get().GetModularFeatureImplementations<IMotionController>( IMotionController::GetModularFeatureName() );
-			for( auto MotionController : MotionControllers )
-			{
-				if( MotionController != nullptr && !Hand.bHaveMotionController )
-				{
-					FVector Location = FVector::ZeroVector;
-					FRotator Rotation = FRotator::ZeroRotator;
-
-					if( MotionController->GetControllerOrientationAndPosition( MotionControllerID, HandIndex == VREditorConstants::LeftHandIndex ? EControllerHand::Left : EControllerHand::Right, /* Out */ Rotation, /* Out */ Location ) )
-					{
-						Hand.bHaveMotionController = true;
-
-						Hand.RoomSpaceTransform = FTransform(
-							Rotation.Quaternion(),
-							Location,
-							FVector( 1.0f ) );
-
-						Hand.Transform = Hand.RoomSpaceTransform * RoomTransform;
-					}
-				}
-			}
-		}
-	}
-}
-
-
 AActor* FVREditorMode::SpawnTransientSceneActor( TSubclassOf<AActor> ActorClass, const FString& ActorName, const bool bWithSceneComponent ) const
 {
 	const bool bWasWorldPackageDirty = GetWorld()->GetOutermost()->IsDirty();
@@ -2080,438 +1128,6 @@ void FVREditorMode::DestroyTransientActor( AActor* Actor ) const
 	}
 }
 
-
-FTransform FVREditorMode::GetRoomSpaceHeadTransform() const
-{
-	FTransform HeadTransform = FTransform::Identity;
-	if( bActuallyUsingVR && GEngine->HMDDevice.IsValid() )
-	{
-		FQuat RoomSpaceHeadOrientation;
-		FVector RoomSpaceHeadLocation;
-		GEngine->HMDDevice->GetCurrentOrientationAndPosition( /* Out */ RoomSpaceHeadOrientation, /* Out */ RoomSpaceHeadLocation );
-
-		HeadTransform = FTransform(
-			RoomSpaceHeadOrientation,
-			RoomSpaceHeadLocation,
-			FVector( 1.0f ) );
-	}
-
-	return HeadTransform;
-}
-
-
-FTransform FVREditorMode::GetHeadTransform() const
-{
-	return GetRoomSpaceHeadTransform() * GetRoomTransform();
-}
-
-
-void FVREditorMode::PlayHapticEffect( const float LeftStrength, const float RightStrength )
-{
-	IInputInterface* InputInterface = FSlateApplication::Get().GetInputInterface();
-	if( InputInterface )
-	{
-		const double CurrentTime = FPlatformTime::Seconds();
-
-		// If we're dealing with an Oculus Rift, we have to setup haptic feedback directly.  Otherwise we can use our
-		// generic force feedback system
-		if( GetHMDDeviceType() == EHMDDeviceType::DT_OculusRift )
-		{
-			// Haptics are a little strong on Oculus Touch, so we scale them down a bit
-			const float HapticScaleForRift = 0.8f;
-
-			// Left hand
-			{
-				FHapticFeedbackValues HapticFeedbackValues;
-				HapticFeedbackValues.Amplitude = LeftStrength * HapticScaleForRift;
-				HapticFeedbackValues.Frequency = 0.5f;
-				InputInterface->SetHapticFeedbackValues( MotionControllerID, (int32)EControllerHand::Left, HapticFeedbackValues );
-
-				VirtualHands[ VREditorConstants::LeftHandIndex ].LastHapticTime = CurrentTime;
-			}
-
-			// Right hand
-			{
-				FHapticFeedbackValues HapticFeedbackValues;
-				HapticFeedbackValues.Amplitude = RightStrength * HapticScaleForRift;
-				HapticFeedbackValues.Frequency = 0.5f;
-				InputInterface->SetHapticFeedbackValues( MotionControllerID, (int32)EControllerHand::Right, HapticFeedbackValues );
-
-				VirtualHands[ VREditorConstants::RightHandIndex ].LastHapticTime = CurrentTime;
-			}
-		}
-		else
-		{
-			FForceFeedbackValues ForceFeedbackValues;
-			ForceFeedbackValues.LeftLarge = LeftStrength;
-			ForceFeedbackValues.RightLarge = RightStrength;
-
-			// @todo vreditor: If an Xbox controller is plugged in, this causes both the motion controllers and the Xbox controller to vibrate!
-			InputInterface->SetForceFeedbackChannelValues( GetMotionControllerID(), ForceFeedbackValues );
-
-			if( ForceFeedbackValues.LeftLarge > KINDA_SMALL_NUMBER )
-			{
-				VirtualHands[ VREditorConstants::LeftHandIndex ].LastHapticTime = CurrentTime;
-			}
-
-			if( ForceFeedbackValues.RightLarge > KINDA_SMALL_NUMBER )
-			{
-				VirtualHands[ VREditorConstants::RightHandIndex ].LastHapticTime = CurrentTime;
-			}
-		}
-	}
-
-	// @todo vreditor: We'll stop haptics right away because if the frame hitches, the controller will be left vibrating
-	StopOldHapticEffects();
-}
-
-
-void FVREditorMode::StopOldHapticEffects()
-{
-	const float MinHapticTime = VREd::MinHapticTimeForRift->GetFloat();
-
-	IInputInterface* InputInterface = FSlateApplication::Get().GetInputInterface();
-	if( InputInterface )
-	{
-		// If we're dealing with an Oculus Rift, we have to setup haptic feedback directly.  Otherwise we can use our
-		// generic force feedback system
-		if( GetHMDDeviceType() == EHMDDeviceType::DT_OculusRift )
-		{
-			bool bWaitingForMoreHaptics = false;
-
-			if( VirtualHands[ VREditorConstants::LeftHandIndex ].LastHapticTime != 0.0 ||
-				VirtualHands[ VREditorConstants::RightHandIndex ].LastHapticTime != 0.0 )
-			{
-				do
-				{
-					bWaitingForMoreHaptics = false;
-
-					const double CurrentTime = FPlatformTime::Seconds();
-
-					// Left hand
-					if( CurrentTime - VirtualHands[ VREditorConstants::LeftHandIndex ].LastHapticTime >= MinHapticTime )
-					{
-						FHapticFeedbackValues HapticFeedbackValues;
-						HapticFeedbackValues.Amplitude = 0.0f;
-						HapticFeedbackValues.Frequency = 0.0f;
-						InputInterface->SetHapticFeedbackValues( MotionControllerID, (int32)EControllerHand::Left, HapticFeedbackValues );
-
-						VirtualHands[ VREditorConstants::LeftHandIndex ].LastHapticTime = 0.0;
-					}
-					else if( VirtualHands[ VREditorConstants::LeftHandIndex ].LastHapticTime != 0.0 )
-					{
-						bWaitingForMoreHaptics = true;
-					}
-
-					// Right hand
-					if( CurrentTime - VirtualHands[ VREditorConstants::RightHandIndex ].LastHapticTime >= MinHapticTime )
-					{
-						FHapticFeedbackValues HapticFeedbackValues;
-						HapticFeedbackValues.Amplitude = 0.0f;
-						HapticFeedbackValues.Frequency = 0.0f;
-						InputInterface->SetHapticFeedbackValues( MotionControllerID, (int32)EControllerHand::Right, HapticFeedbackValues );
-
-						VirtualHands[ VREditorConstants::RightHandIndex ].LastHapticTime = 0.0;
-					}
-					else if( VirtualHands[ VREditorConstants::RightHandIndex ].LastHapticTime != 0.0 )
-					{
-						bWaitingForMoreHaptics = true;
-					}
-
-					if( bWaitingForMoreHaptics && VREd::SleepForRiftHaptics->GetInt() != 0 )
-					{
-						FPlatformProcess::Sleep( 0 );
-					}
-				}
-				// @todo vreditor urgent: This is needed so that haptics don't play too long.  Our Oculus code doesn't currently support 
-				// multi-threading, so we need to delay the main thread to make sure we stop it before it rumbles for more than an instant!
-				while( bWaitingForMoreHaptics && VREd::SleepForRiftHaptics->GetInt() != 0 );
-			}
-		}
-		else
-		{
-			// @todo vreditor: Do we need to cancel haptics for non-Rift devices?  Doesn't seem like it
-		}
-	}
-}
-
-void FVREditorMode::ShowHelpForHand( const int32 HandIndex, const bool bShowIt )
-{
-	FVirtualHand& Hand = VirtualHands[ HandIndex ];
-
-	if( bShowIt != Hand.bWantHelpLabels )
-	{
-		Hand.bWantHelpLabels = bShowIt;
-
-		const FTimespan CurrentTime = FTimespan::FromSeconds( FApp::GetCurrentTime() );
-		const FTimespan TimeSinceStartedFadingOut = CurrentTime - Hand.HelpLabelShowOrHideStartTime;
-		const FTimespan HelpLabelFadeDuration = FTimespan::FromSeconds( VREd::HelpLabelFadeDuration->GetFloat() );
-
-		// If we were already fading, account for that here
-		if( TimeSinceStartedFadingOut < HelpLabelFadeDuration )
-		{
-			// We were already fading, so we'll reverse the time value so it feels continuous
-			Hand.HelpLabelShowOrHideStartTime = CurrentTime - ( HelpLabelFadeDuration - TimeSinceStartedFadingOut );
-		}
-		else
-		{
-			Hand.HelpLabelShowOrHideStartTime = FTimespan::FromSeconds( FApp::GetCurrentTime() );
-		}
-
-		if( bShowIt && Hand.HelpLabels.Num() == 0 )
-		{
-			for( const auto& KeyToAction : KeyToActionMap )
-			{
-				const FKey Key = KeyToAction.Key;
-				const FVRAction& VRAction = KeyToAction.Value;
-
-				if( HandIndex == VRAction.HandIndex )
-				{
-					UStaticMeshSocket* Socket = FindMeshSocketForKey( Hand.HandMeshComponent->StaticMesh, Key );
-					if( Socket != nullptr )
-					{
-						FText LabelText;
-						FString ComponentName;
-
-						switch( VRAction.ActionType )
-						{
-							case EVRActionType::Modifier:
-								LabelText = LOCTEXT( "ModifierHelp", "Modifier" );
-								ComponentName = TEXT( "ModifierHelp" );
-								break;
-
-							case EVRActionType::WorldMovement:
-								LabelText = LOCTEXT( "WorldMovementHelp", "Move World" );
-								ComponentName = TEXT( "WorldMovementHelp" );
-								break;
-
-							case EVRActionType::SelectAndMove:
-								LabelText = LOCTEXT( "SelectAndMoveHelp", "Select & Move" );
-								ComponentName = TEXT( "SelectAndMoveHelp" );
-								break;
-
-							case EVRActionType::SelectAndMove_LightlyPressed:
-								LabelText = LOCTEXT( "SelectAndMove_LightlyPressedHelp", "Select & Move" );
-								ComponentName = TEXT( "SelectAndMove_LightlyPressedHelp" );
-								break;
-
-							case EVRActionType::Touch:
-								LabelText = LOCTEXT( "TouchHelp", "Slide" );
-								ComponentName = TEXT( "TouchHelp" );
-								break;
-
-							case EVRActionType::Undo:
-								LabelText = LOCTEXT( "UndoHelp", "Undo" );
-								ComponentName = TEXT( "UndoHelp" );
-								break;
-
-							case EVRActionType::Redo:
-								LabelText = LOCTEXT( "RedoHelp", "Redo" );
-								ComponentName = TEXT( "RedoHelp" );
-								break;
-
-							case EVRActionType::Delete:
-								LabelText = LOCTEXT( "DeleteHelp", "Delete" );
-								ComponentName = TEXT( "DeleteHelp" );
-								break;
-
-							case EVRActionType::ConfirmRadialSelection:
-								LabelText = LOCTEXT( "ConfirmRadialSelectionHelp", "Radial Menu" );
-								ComponentName = TEXT( "ConfirmRadialSelectionHelp" );
-								break;
-
-							default:
-								// If this goes off, a new EVRActionType was added and we need to update the above switch statement
-								check( 0 );
-								break;
-						}
-					
-						const bool bWithSceneComponent = false;	// Nope, we'll spawn our own inside AFloatingText
-						AFloatingText* FloatingText = SpawnTransientSceneActor<AFloatingText>( ComponentName, bWithSceneComponent );
-						FloatingText->SetText( LabelText );
-
-						Hand.HelpLabels.Add( Key, FloatingText );
-					}
-				}
-			}
-		}
-	}
-}
-
-
-void FVREditorMode::UpdateHelpLabels()
-{
-	const FTimespan HelpLabelFadeDuration = FTimespan::FromSeconds( VREd::HelpLabelFadeDuration->GetFloat() );
-
-	const FTransform HeadTransform = GetHeadTransform();
-	for( int32 HandIndex = 0; HandIndex < VREditorConstants::NumVirtualHands; ++HandIndex )
-	{
-		FVirtualHand& Hand = VirtualHands[ HandIndex ];
-
-		// Only show help labels if the hand is pretty close to the face
-		const float DistanceToHead = ( Hand.Transform.GetLocation() - HeadTransform.GetLocation() ).Size();
-		const float MinDistanceToHeadForHelp = VREd::HelpLabelFadeDistance->GetFloat() * GetWorldScaleFactor();	// (in cm)
-		bool bShowHelp = DistanceToHead <= MinDistanceToHeadForHelp;
-
-		// Don't show help if a UI is summoned on that hand
-		if( Hand.bHasUIInFront || Hand.bHasUIOnForearm || UISystem->IsShowingRadialMenu( HandIndex ) )
-		{
-			bShowHelp = false;
-		}
-
-		ShowHelpForHand( HandIndex, bShowHelp );
-
-		// Have the labels finished fading out?  If so, we'll kill their actors!
-		const FTimespan CurrentTime = FTimespan::FromSeconds( FApp::GetCurrentTime() );
-		const FTimespan TimeSinceStartedFadingOut = CurrentTime - Hand.HelpLabelShowOrHideStartTime;
-		if( !Hand.bWantHelpLabels && ( TimeSinceStartedFadingOut > HelpLabelFadeDuration ) )
-		{
-			// Get rid of help text
-			for( auto& KeyAndValue : Hand.HelpLabels )
-			{
-				AFloatingText* FloatingText = KeyAndValue.Value;
-				DestroyTransientActor( FloatingText );
-			}
-			Hand.HelpLabels.Reset();
-		}
-		else
-		{
-			// Update fading state
-			float FadeAlpha = FMath::Clamp( (float)TimeSinceStartedFadingOut.GetTotalSeconds() / (float)HelpLabelFadeDuration.GetTotalSeconds(), 0.0f, 1.0f );
-			if( !Hand.bWantHelpLabels )
-			{
-				FadeAlpha = 1.0f - FadeAlpha;
-			}
-
-			// Exponential falloff, so the fade is really obvious (gamma/HDR)
-			FadeAlpha = FMath::Pow( FadeAlpha, 3.0f );
-
-			for( auto& KeyAndValue : Hand.HelpLabels )
-			{
-				const FKey Key = KeyAndValue.Key;
-				AFloatingText* FloatingText = KeyAndValue.Value;
-
-				UStaticMeshSocket* Socket = FindMeshSocketForKey( Hand.HandMeshComponent->StaticMesh, Key );
-				check( Socket != nullptr );
-				FTransform SocketRelativeTransform( Socket->RelativeRotation, Socket->RelativeLocation, Socket->RelativeScale );
-
-				// Oculus has asymmetrical controllers, so we the sock transform horizontally
-				if( HandIndex == VREditorConstants::RightHandIndex &&
-					GetHMDDeviceType() == EHMDDeviceType::DT_OculusRift )
-				{
-					const FVector Scale3D = SocketRelativeTransform.GetLocation();
-					SocketRelativeTransform.SetLocation( FVector( Scale3D.X, -Scale3D.Y, Scale3D.Z ) );
-				}
-
-				// Make sure the labels stay the same size even when the world is scaled
-				FTransform HandTransformWithWorldToMetersScaling = Hand.Transform;
-				HandTransformWithWorldToMetersScaling.SetScale3D( HandTransformWithWorldToMetersScaling.GetScale3D() * FVector( GetWorldScaleFactor() ) );
-
-
-				// Position right on top of the controller itself
-				FTransform FloatingTextTransform = SocketRelativeTransform * HandTransformWithWorldToMetersScaling;
-				FloatingText->SetActorTransform( FloatingTextTransform );
-
-				// Orientate it toward the viewer
-				FloatingText->Update( HeadTransform.GetLocation() );
-
-				// Update fade state
-				FloatingText->SetOpacity( FadeAlpha );
-			}
-		}
-	}
-}
-
-
-UStaticMeshSocket* FVREditorMode::FindMeshSocketForKey( UStaticMesh* StaticMesh, const FKey Key )
-{
-	// @todo vreditor: Hard coded mapping of socket names (e.g. "Shoulder") to expected names of sockets in the static mesh
-	FName SocketName = NAME_None;
-	if( Key == EKeys::MotionController_Left_Shoulder || Key == EKeys::MotionController_Right_Shoulder )
-	{
-		static FName ShoulderSocketName( "Shoulder" );
-		SocketName = ShoulderSocketName;
-	}
-	else if( Key == EKeys::MotionController_Left_Trigger || Key == EKeys::MotionController_Right_Trigger ||
-		     Key == VREditorKeyNames::MotionController_Left_FullyPressedTriggerAxis || Key == VREditorKeyNames::MotionController_Right_FullyPressedTriggerAxis ||
-			 Key == VREditorKeyNames::MotionController_Left_LightlyPressedTriggerAxis || Key == VREditorKeyNames::MotionController_Right_LightlyPressedTriggerAxis )
-	{
-		static FName TriggerSocketName( "Trigger" );
-		SocketName = TriggerSocketName;
-	}
-	else if( Key == EKeys::MotionController_Left_Grip1 || Key == EKeys::MotionController_Right_Grip1 )
-	{
-		static FName GripSocketName( "Grip" );
-		SocketName = GripSocketName;
-	}
-	else if( Key == EKeys::MotionController_Left_Thumbstick || Key == EKeys::MotionController_Right_Thumbstick )
-	{
-		static FName ThumbstickSocketName( "Thumbstick" );
-		SocketName = ThumbstickSocketName;
-	}
-	else if( Key == SteamVRControllerKeyNames::Touch0 || Key == SteamVRControllerKeyNames::Touch1 )
-	{
-		static FName TouchSocketName( "Touch" );
-		SocketName = TouchSocketName;
-	}					   
-	else if( Key == EKeys::MotionController_Left_Thumbstick_Down || Key == EKeys::MotionController_Right_Thumbstick_Down )
-	{
-		static FName DownSocketName( "Down" );
-		SocketName = DownSocketName;
-	}
-	else if( Key == EKeys::MotionController_Left_Thumbstick_Up || Key == EKeys::MotionController_Right_Thumbstick_Up )
-	{
-		static FName UpSocketName( "Up" );
-		SocketName = UpSocketName;
-	}
-	else if( Key == EKeys::MotionController_Left_Thumbstick_Left || Key == EKeys::MotionController_Right_Thumbstick_Left )
-	{
-		static FName LeftSocketName( "Left" );
-		SocketName = LeftSocketName;
-	}
-	else if( Key == EKeys::MotionController_Left_Thumbstick_Right || Key == EKeys::MotionController_Right_Thumbstick_Right )
-	{
-		static FName RightSocketName( "Right" );
-		SocketName = RightSocketName;
-	}
-	else if( Key == EKeys::MotionController_Left_FaceButton1 || Key == EKeys::MotionController_Right_FaceButton1 )
-	{
-		static FName FaceButton1SocketName( "FaceButton1" );
-		SocketName = FaceButton1SocketName;
-	}
-	else if( Key == EKeys::MotionController_Left_FaceButton2 || Key == EKeys::MotionController_Right_FaceButton2 )
-	{
-		static FName FaceButton2SocketName( "FaceButton2" );
-		SocketName = FaceButton2SocketName;
-	}
-	else if( Key == EKeys::MotionController_Left_FaceButton3 || Key == EKeys::MotionController_Right_FaceButton3 )
-	{
-		static FName FaceButton3SocketName( "FaceButton3" );
-		SocketName = FaceButton3SocketName;
-	}
-	else if( Key == EKeys::MotionController_Left_FaceButton4 || Key == EKeys::MotionController_Right_FaceButton4 )
-	{
-		static FName FaceButton4SocketName( "FaceButton4" );
-		SocketName = FaceButton4SocketName;
-	}
-	else
-	{
-		// Not a key that we care about
-	}
-
-	if( SocketName != NAME_None )
-	{
-		UStaticMeshSocket* Socket = StaticMesh->FindSocket( SocketName );
-		if( Socket != nullptr )
-		{
-			return Socket;
-		}
-	}
-
-	return nullptr;
-};
-
-
 const SLevelViewport& FVREditorMode::GetLevelViewportPossessedForVR() const
 {
 	return *VREditorLevelViewportWeakPtr.Pin();
@@ -2524,6 +1140,11 @@ SLevelViewport& FVREditorMode::GetLevelViewportPossessedForVR()
 }
 
 
+float FVREditorMode::GetWorldScaleFactor() const
+{
+	return WorldInteraction->GetViewportWorld()->GetWorldSettings()->WorldToMeters / 100.0f;
+}
+
 void FVREditorMode::OnMapChange( uint32 MapChangeFlags )
 {
 	CleanUpActorsBeforeMapChangeOrSimulate();
@@ -2571,41 +1192,11 @@ void FVREditorMode::CleanUpActorsBeforeMapChangeOrSimulate()
 		WorldMovementPostProcessMaterial = nullptr;
 	}
 
-	DestroyTransientActor( SnapGridActor );
-	SnapGridActor = nullptr;
-	SnapGridMeshComponent = nullptr;
-
-	if( SnapGridMID != nullptr )
-	{
-		SnapGridMID->MarkPendingKill();
-		SnapGridMID = nullptr;
-	}
-
 	PostProcessComponent = nullptr;
 	ScaleProgressMeshComponent = nullptr;
 	CurrentScaleProgressMeshComponent = nullptr;
 	UserScaleIndicatorText = nullptr;
 
-	for( int32 HandIndex = 0; HandIndex < VREditorConstants::NumVirtualHands; ++HandIndex )
-	{
-		FVirtualHand& Hand = VirtualHands[ HandIndex ];
-
-		Hand.MotionControllerComponent = nullptr;
-		Hand.HandMeshComponent = nullptr;
-		Hand.HoverMeshComponent = nullptr;
-		Hand.HoverPointLightComponent = nullptr;
-		Hand.LaserPointerMeshComponent = nullptr;
-		Hand.LaserPointerMID = nullptr;
-		Hand.TranslucentLaserPointerMID = nullptr;
-
-		for( auto& KeyAndValue : Hand.HelpLabels )
-		{
-			AFloatingText* FloatingText = KeyAndValue.Value;
-			DestroyTransientActor( FloatingText );
-		}
-		Hand.HelpLabels.Reset();
-	}
-
 	if( UISystem != nullptr )
 	{
 		UISystem->CleanUpActorsBeforeMapChangeOrSimulate();
@@ -2613,62 +1204,13 @@ void FVREditorMode::CleanUpActorsBeforeMapChangeOrSimulate()
 
 	if( WorldInteraction != nullptr )
 	{
-		WorldInteraction->CleanUpActorsBeforeMapChangeOrSimulate();
+		WorldInteraction->Shutdown();
 	}
 }
 
-
-void FVREditorMode::Undo()
-{
-	GUnrealEd->Exec( GetWorld(), TEXT( "TRANSACTION UNDO" ) );
-}
-
-
-void FVREditorMode::Redo()
+void FVREditorMode::ToggleFlashlight( UVREditorInteractor* Interactor )
 {
-	GUnrealEd->Exec( GetWorld(), TEXT( "TRANSACTION REDO" ) );
-}
-
-
-void FVREditorMode::Copy()
-{
-	// @todo vreditor: Needs CanExecute()  (see LevelEditorActions.cpp)
-	GUnrealEd->Exec( GetWorld(), TEXT( "EDIT COPY" ) );
-}
-
-
-void FVREditorMode::Paste()
-{
-	// @todo vreditor: Needs CanExecute()  (see LevelEditorActions.cpp)
-	// @todo vreditor: Needs "paste here" style pasting (TO=HERE), but with ray
-	GUnrealEd->Exec( GetWorld(), TEXT( "EDIT PASTE" ) );
-}
-
-
-void FVREditorMode::Duplicate()
-{
-	// @todo vreditor: Needs CanExecute()  (see LevelEditorActions.cpp)
-	GUnrealEd->Exec( GetWorld(), TEXT( "DUPLICATE" ) );
-
-}
-
-
-void FVREditorMode::SnapSelectedActorsToGround()
-{
-	FLevelEditorModule& LevelEditorModule = FModuleManager::GetModuleChecked<FLevelEditorModule>( TEXT( "LevelEditor" ) );
-	const FLevelEditorCommands& Commands = LevelEditorModule.GetLevelEditorCommands();
-	const TSharedPtr< FUICommandList >& CommandList = GetLevelViewportPossessedForVR().GetParentLevelEditor().Pin()->GetLevelEditorActions();
-
-	CommandList->ExecuteAction( Commands.SnapBottomCenterBoundsToFloor.ToSharedRef() );
-	
-	// @todo vreditor: This should not be needed after to allow transformables to stop animating the transforms of actors after they come to rest
-	WorldInteraction->SetupTransformablesForSelectedActors();
-}
-
-void FVREditorMode::ToggleFlashlight( const int32 HandIndex )
-{
-
-	if (FlashlightComponent == nullptr)
+	if ( FlashlightComponent == nullptr )
 	{
 		FlashlightComponent = NewObject<USpotLightComponent>(AvatarMeshActor);
 		AvatarMeshActor->AddOwnedComponent(FlashlightComponent);
@@ -2684,117 +1226,15 @@ void FVREditorMode::ToggleFlashlight( const int32 HandIndex )
 		
 	}
 
-	FVirtualHand Hand = GetVirtualHand( HandIndex );
-	const FAttachmentTransformRules Rules = FAttachmentTransformRules( EAttachmentRule::KeepRelative, true );
-	FlashlightComponent->AttachToComponent( Hand.MotionControllerComponent, Rules );
+	const FAttachmentTransformRules AttachmentTransformRules = FAttachmentTransformRules( EAttachmentRule::KeepRelative, true );
+	FlashlightComponent->AttachToComponent( Interactor->GetMotionControllerComponent(), AttachmentTransformRules );
 	bIsFlashlightOn = !bIsFlashlightOn;
-	FlashlightComponent->SetVisibility( bIsFlashlightOn );
-	
-
-}
-
-void FVREditorMode::SetLaserVisuals( const int32 HandIndex, const FLinearColor& NewColor, const float CrawlFade, const float CrawlSpeed )
-{
-	FVirtualHand Hand = GetVirtualHand( HandIndex );
-
-	static FName StaticLaserColorParameterName( "LaserColor" );
-	Hand.LaserPointerMID->SetVectorParameterValue( StaticLaserColorParameterName, NewColor );
-	Hand.TranslucentLaserPointerMID->SetVectorParameterValue( StaticLaserColorParameterName, NewColor );
-
-	static FName StaticCrawlParameterName( "Crawl" );
-	Hand.LaserPointerMID->SetScalarParameterValue( StaticCrawlParameterName, CrawlFade );
-	Hand.TranslucentLaserPointerMID->SetScalarParameterValue( StaticCrawlParameterName, CrawlFade );
-
-	static FName StaticCrawlSpeedParameterName( "CrawlSpeed" );
-	Hand.LaserPointerMID->SetScalarParameterValue( StaticCrawlSpeedParameterName, CrawlSpeed );
-	Hand.TranslucentLaserPointerMID->SetScalarParameterValue( StaticCrawlSpeedParameterName, CrawlSpeed );
-
-	static FName StaticHandTrimColorParameter( "TrimGlowColor" );
-	Hand.HandMeshMID->SetVectorParameterValue( StaticHandTrimColorParameter, NewColor );
-
-	Hand.HoverPointLightComponent->SetLightColor( NewColor );
-}
-
-
-void FVREditorMode::ApplyVelocityDamping( FVector& Velocity, const bool bVelocitySensitive )
-{
-	const float InertialMovementZeroEpsilon = 0.01f;	// @todo vreditor tweak
-	if( !Velocity.IsNearlyZero( InertialMovementZeroEpsilon ) )
-	{
-		// Apply damping
-		if( bVelocitySensitive )
-		{
-			const float DampenMultiplierAtLowSpeeds = 0.94f;	// @todo vreditor tweak
-			const float DampenMultiplierAtHighSpeeds = 0.99f;	// @todo vreditor tweak
-			const float SpeedForMinimalDamping = 2.5f * GetWorldScaleFactor();	// cm/frame	// @todo vreditor tweak
-			const float SpeedBasedDampeningScalar = FMath::Clamp( Velocity.Size(), 0.0f, SpeedForMinimalDamping ) / SpeedForMinimalDamping;	// @todo vreditor: Probably needs a curve applied to this to compensate for our framerate insensitivity
-			Velocity = Velocity * FMath::Lerp( DampenMultiplierAtLowSpeeds, DampenMultiplierAtHighSpeeds, SpeedBasedDampeningScalar );	// @todo vreditor: Frame rate sensitive damping.  Make use of delta time!
-		}
-		else
-		{
-			Velocity = Velocity * 0.95f;
-		}
-	}
-
-	if( Velocity.IsNearlyZero( InertialMovementZeroEpsilon ) )
-	{
-		Velocity = FVector::ZeroVector;
-	}
-}
-
-
-void FVREditorMode::ApplyButtonPressColors( FVRAction VRAction, EInputEvent Event )
-{
-	SpawnAvatarMeshActor();
-	
-	FVirtualHand& Hand = GetVirtualHand( VRAction.HandIndex );
-
-	const float PressStrength = 10.0f;
-
-	// Trigger
-	if ( VRAction.ActionType == EVRActionType::SelectAndMove || VRAction.ActionType == EVRActionType::SelectAndMove_LightlyPressed )
-	{
-		static FName StaticTriggerParameter( "B1" );
-		SetMotionControllerButtonPressedVisuals( Hand, Event, StaticTriggerParameter, PressStrength );
-	}
-
-	// Shoulder button
-	if ( VRAction.ActionType == EVRActionType::WorldMovement )
-	{
-		static FName StaticShoulderParameter( "B2" );
-		SetMotionControllerButtonPressedVisuals( Hand, Event, StaticShoulderParameter, PressStrength );
-	}
-
-	// Trackpad
-	if( VRAction.ActionType == EVRActionType::ConfirmRadialSelection )
-	{
-		static FName StaticTrackpadParameter( "B3" );
-		SetMotionControllerButtonPressedVisuals( Hand, Event, StaticTrackpadParameter, PressStrength );
-	}
-
-	// Modifier
-	if( VRAction.ActionType == EVRActionType::Modifier )
-	{
-		static FName StaticModifierParameter( "B4" );
-		SetMotionControllerButtonPressedVisuals( Hand, Event, StaticModifierParameter, PressStrength );
-	}
-}
-
-void FVREditorMode::SetMotionControllerButtonPressedVisuals( FVirtualHand& Hand, EInputEvent Event, const FName& ParameterName, const float PressStrength )
-{
-	if ( Event == IE_Pressed )
-	{
-		Hand.HandMeshMID->SetScalarParameterValue( ParameterName, PressStrength);
-	}
-	else if ( Event == IE_Released )
-	{
-		Hand.HandMeshMID->SetScalarParameterValue( ParameterName, 0.0f);
-	}
+	FlashlightComponent->SetVisibility(bIsFlashlightOn);
 }
 
 void FVREditorMode::CycleTransformGizmoHandleType()
 {
-	EGizmoHandleTypes NewGizmoType = (EGizmoHandleTypes)( (uint8)CurrentGizmoType + 1 );
+	EGizmoHandleTypes NewGizmoType = (EGizmoHandleTypes)( (uint8)WorldInteraction->GetCurrentGizmoType() + 1 );
 	
 	if( NewGizmoType > EGizmoHandleTypes::Scale )
 	{
@@ -2804,7 +1244,7 @@ void FVREditorMode::CycleTransformGizmoHandleType()
 	// Set coordinate system to local if the next gizmo will be for non-uniform scaling 
 	if ( NewGizmoType == EGizmoHandleTypes::Scale )
 	{
-		const ECoordSystem CurrentCoordSystem = GetTransformGizmoCoordinateSpace();
+		const ECoordSystem CurrentCoordSystem = WorldInteraction->GetTransformGizmoCoordinateSpace();
 		if ( CurrentCoordSystem == COORD_World )
 		{
 			GLevelEditorModeTools().SetCoordSystem( COORD_Local );
@@ -2816,41 +1256,15 @@ void FVREditorMode::CycleTransformGizmoHandleType()
 			bWasInWorldSpaceBeforeScaleMode = false;
 		}
 	} 
-	else if ( CurrentGizmoType == EGizmoHandleTypes::Scale && bWasInWorldSpaceBeforeScaleMode )
+	else if ( WorldInteraction->GetCurrentGizmoType() == EGizmoHandleTypes::Scale && bWasInWorldSpaceBeforeScaleMode )
 	{
 		// Set the coordinate system to world space if the coordinate system was world before scaling
-		SetTransformGizmoCoordinateSpace( COORD_World );
+		WorldInteraction->SetTransformGizmoCoordinateSpace( COORD_World );
 	}
 	
-	CurrentGizmoType = NewGizmoType;
-}
-
-
-EGizmoHandleTypes FVREditorMode::GetCurrentGizmoType() const
-{
-	return CurrentGizmoType;
-}
-
-void FVREditorMode::CycleTransformGizmoCoordinateSpace()
-{
-	const bool bGetRawValue = true;
-	const ECoordSystem CurrentCoordSystem = GLevelEditorModeTools().GetCoordSystem( bGetRawValue );
-	SetTransformGizmoCoordinateSpace( CurrentCoordSystem == COORD_World ? COORD_Local : COORD_World );
-}
-
-void FVREditorMode::SetTransformGizmoCoordinateSpace( const ECoordSystem NewCoordSystem )
-{
-	GLevelEditorModeTools().SetCoordSystem( NewCoordSystem );
-}
-
-ECoordSystem FVREditorMode::GetTransformGizmoCoordinateSpace() const
-{
-	const bool bGetRawValue = true;
-	const ECoordSystem CurrentCoordSystem = GLevelEditorModeTools().GetCoordSystem( bGetRawValue );
-	return CurrentCoordSystem;
+	WorldInteraction->SetGizmoHandleType( NewGizmoType );
 }
 
-
 EHMDDeviceType::Type FVREditorMode::GetHMDDeviceType() const
 {
 	return GEngine->HMDDevice.IsValid() ? GEngine->HMDDevice->GetHMDDeviceType() : EHMDDeviceType::DT_SteamVR;
@@ -2861,14 +1275,13 @@ FLinearColor FVREditorMode::GetColor( const EColors Color ) const
 	return Colors[ (int32)Color ];
 }
 
-
-bool FVREditorMode::IsHandAimingTowardsCapsule( const int32 HandIndex, const FTransform& CapsuleTransform, FVector CapsuleStart, const FVector CapsuleEnd, const float CapsuleRadius, const float MinDistanceToCapsule, const FVector CapsuleFrontDirection, const float MinDotForAimingAtCapsule ) const
+bool FVREditorMode::IsHandAimingTowardsCapsule( UViewportInteractor* Interactor, const FTransform& CapsuleTransform, FVector CapsuleStart, const FVector CapsuleEnd, const float CapsuleRadius, const float MinDistanceToCapsule, const FVector CapsuleFrontDirection, const float MinDotForAimingAtCapsule ) const
 {
 	bool bIsAimingTowards = false;
 	const float WorldScaleFactor = GetWorldScaleFactor();
 
 	FVector LaserPointerStart, LaserPointerEnd;
-	if( GetLaserPointer( HandIndex, /* Out */ LaserPointerStart, /* Out */ LaserPointerEnd ) )
+	if( Interactor->GetLaserPointer( /* Out */ LaserPointerStart, /* Out */ LaserPointerEnd ) )
 	{
 		const FVector LaserPointerStartInCapsuleSpace = CapsuleTransform.InverseTransformPosition( LaserPointerStart );
 		const FVector LaserPointerEndInCapsuleSpace = CapsuleTransform.InverseTransformPosition( LaserPointerEnd );
@@ -2923,4 +1336,17 @@ bool FVREditorMode::IsHandAimingTowardsCapsule( const int32 HandIndex, const FTr
 	return bIsAimingTowards;
 }
 
+UVREditorInteractor* FVREditorMode::GetHandInteractor( const EControllerHand ControllerHand ) const 
+{
+	UVREditorInteractor* ResultInteractor = ControllerHand == EControllerHand::Left ? LeftHandInteractor : RightHandInteractor;
+	check( ResultInteractor != nullptr );
+	return ResultInteractor;
+}
+
+void FVREditorMode::StopOldHapticEffects()
+{
+	LeftHandInteractor->StopOldHapticEffects();
+	RightHandInteractor->StopOldHapticEffects();
+}
+
 #undef LOCTEXT_NAMESPACE
\ No newline at end of file
