VR Editor:

diff --git a/Engine/Source/Editor/VREditor/VREditorMode.h b/Engine/Source/Editor/VREditor/VREditorMode.h
--- a/Engine/Source/Editor/VREditor/VREditorMode.h
+++ b/Engine/Source/Editor/VREditor/VREditorMode.h
@@ -4,10 +4,8 @@
 
 #include "IVREditorMode.h"
 #include "IVREditorModule.h"
-#include "VirtualHand.h"
 #include "HeadMountedDisplayTypes.h"	// For EHMDDeviceType::Type
 #include "Kismet/HeadMountedDisplayFunctionLibrary.h"	// For EHMDTrackingOrigin::Type
-#include "VREditorInputProcessor.h"
 
 // Forward declare the GizmoHandleTypes that is defined in VREditorTransformGizmo
 enum class EGizmoHandleTypes : uint8;
@@ -65,27 +63,20 @@ public:
 	virtual void AddReferencedObjects( FReferenceCollector& Collector ) override;
 	virtual void Render( const FSceneView* SceneView, FViewport* Viewport, FPrimitiveDrawInterface* PDI ) override;	
 
-	bool HandleInputKey(FEditorViewportClient* ViewportClient, FViewport* Viewport, FKey Key, EInputEvent Event);
-	bool HandleInputAxis(FEditorViewportClient* ViewportClient, FViewport* Viewport, int32 ControllerId, FKey Key, float Delta, float DeltaTime);
-
 	// IVREditorMode interface
-	virtual bool GetLaserPointer( const int32 HandIndex, FVector& LaserPointerStart, FVector& LaserPointerEnd, const bool bEvenIfUIIsInFront = false, const float LaserLengthOverride = 0.0f) const override;
+	virtual AActor* GetAvatarMeshActor() override;
+	virtual UWorld* GetWorld() const override;
+	virtual FTransform GetRoomTransform() const override;
+	virtual void SetRoomTransform( const FTransform& NewRoomTransform ) override;
+	virtual FTransform GetRoomSpaceHeadTransform() const override;
 	virtual FTransform GetHeadTransform() const override;
-	virtual bool IsInputCaptured( const FVRAction VRAction ) const;
-	virtual AActor* GetAvatarMeshActor() override
-	{
-		return AvatarMeshActor;
-	}
-	virtual FOnVRAction& OnVRAction() override
-	{
-		return OnVRActionEvent;
-	}
-	virtual FOnVRHoverUpdate& OnVRHoverUpdate() override
+	virtual const class UViewportWorldInteraction& GetWorldInteraction() const override;
+	virtual class UViewportWorldInteraction& GetWorldInteraction() override;
+	virtual bool IsFullyInitialized() const override;
+	virtual FOnVRTickHandle& OnTickHandle() override
 	{
-		return OnVRHoverUpdateEvent;
+		return TickHandle;
 	}
-	virtual FHandInfo GetHandInfo( const int32 HandIndex ) const override;
-	virtual void SetAllowTriggerLightPressLocking( const int32 HandIndex, const bool bAllow ) override;
 
 	/** Returns the Unreal controller ID for the motion controllers we're using */
 	int32 GetMotionControllerID() const
@@ -114,53 +105,6 @@ public:
 		return 1.0f - ( ( Alpha * ( ( OvershootAmount + 1 ) * Alpha + OvershootAmount ) + 1 ) - 1.0f );
 	}
 
-	/** Gets the maximum length of a laser pointer */
-	float GetLaserPointerMaxLength() const;
-
-	/** Gets a virtual hand */
-	FVirtualHand& GetVirtualHand( const int32 VirtualHandIndex )
-	{
-		check( VirtualHandIndex >= 0 && VirtualHandIndex < VREditorConstants::NumVirtualHands );
-		return VirtualHands[ VirtualHandIndex ];
-	}
-
-	/** Gets a virtual hand (const) */
-	const FVirtualHand& GetVirtualHand( const int32 VirtualHandIndex ) const
-	{
-		check( VirtualHandIndex >= 0 && VirtualHandIndex < VREditorConstants::NumVirtualHands );
-		return VirtualHands[ VirtualHandIndex ];
-	}
-
-	/** Given an index of a hand, gets the index of the opposite hand */
-	int32 GetOtherHandIndex( const int32 HandIndex ) const
-	{
-		check( HandIndex >= 0 && HandIndex < VREditorConstants::NumVirtualHands );
-		const int32 OtherHandIndex = ( HandIndex == VREditorConstants::LeftHandIndex ) ? VREditorConstants::RightHandIndex : VREditorConstants::LeftHandIndex;
-		return OtherHandIndex;
-	}
-
-	/** Gets the other hand, given a hand index */
-	FVirtualHand& GetOtherHand( const int32 HandIndex )
-	{
-		return VirtualHands[ GetOtherHandIndex( HandIndex ) ];
-	}
-
-	/** Gets the other hand, given a hand index (const) */
-	const FVirtualHand& GetOtherHand( const int32 HandIndex ) const
-	{
-		return VirtualHands[ GetOtherHandIndex( HandIndex ) ];
-	}
-
-	/** Gets the world space transform of the calibrated VR room origin.  When using a seated VR device, this will feel like the
-	    camera's world transform (before any HMD positional or rotation adjustments are applied.) */
-	FTransform GetRoomTransform() const;
-
-	/** Sets a new transform for the room, in world space.  This is basically setting the editor's camera transform for the viewport */
-	void SetRoomTransform( const FTransform& NewRoomTransform );
-
-	/** Gets the transform of the user's HMD in room space */
-	FTransform GetRoomSpaceHeadTransform() const;
-
 	/** Spawns a transient actor that we can use in the editor world (templated for convenience) */
 	template<class T>
 	T* SpawnTransientSceneActor( const FString& ActorName, const bool bWithSceneComponent ) const
@@ -174,42 +118,14 @@ public:
 	/** Destroys a transient actor we created earlier and nulls out the pointer */
 	void DestroyTransientActor( AActor* Actor ) const;
 
-	/**
-	 * Traces along the laser pointer vector and returns what it first hits in the world
-	 *
-	 * @param HandIndex	Index of the hand that uses has the laser pointer to trace
-	 * @param OptionalListOfIgnoredActors Actors to exclude from hit testing
-	 * @param bIgnoreGizmos True if no gizmo results should be returned, otherwise they are preferred (x-ray)
-	 * @param bEvenIfUIIsInFront If true, ignores any UI that might be blocking the ray
-	 * @param LaserLengthOverride If zero the default laser length (VREdMode::GetLaserLength) is used
-	 *
-	 * @return What the laster pointer hit
-	 */
-	FHitResult GetHitResultFromLaserPointer( int32 HandIndex, TArray<AActor*>* OptionalListOfIgnoredActors = nullptr, const bool bIgnoreGizmos = false, const bool bEvenIfUIIsInFront = false, const float LaserLengthOverride = 0.0f );
-
-	/** Triggers a force feedback effect on either hand */
-	void PlayHapticEffect( const float LeftStrength, const float RightStrength );
-
-	/** Gets access to the world interaction system (const) */
-	const class FVREditorWorldInteraction& GetWorldInteraction() const
-	{
-		return *WorldInteraction;
-	}
-
-	/** Gets access to the world interaction system */
-	class FVREditorWorldInteraction& GetWorldInteraction()
-	{
-		return *WorldInteraction;
-	}
-
 	/** Gets access to the VR UI system (const) */
-	const class FVREditorUISystem& GetUISystem() const
+	const class UVREditorUISystem& GetUISystem() const
 	{
 		return *UISystem;
 	}
 
 	/** Gets access to the VR UI system */
-	class FVREditorUISystem& GetUISystem()
+	class UVREditorUISystem& GetUISystem()
 	{
 		return *UISystem;
 	}
@@ -223,10 +139,7 @@ public:
 	class SLevelViewport& GetLevelViewportPossessedForVR();
 
 	/** Gets the world scale factor, which can be multiplied by a scale vector to convert to room space */
-	float GetWorldScaleFactor() const
-	{
-		return GetWorld()->GetWorldSettings()->WorldToMeters / 100.0f;
-	}
+	float GetWorldScaleFactor() const;
 
 	/** Sets up our avatar mesh, if not already spawned */
 	void SpawnAvatarMeshActor();
@@ -234,49 +147,11 @@ public:
 	/** Gets our VR post process component, creating it if it doesn't exist yet */
 	UPostProcessComponent* GetPostProcessComponent();
 
-	/** Gets the snap grid mesh actor */
-	AActor* GetSnapGridActor()
-	{
-		SpawnAvatarMeshActor();
-		return SnapGridActor;
-	}
-
-	/** Gets the snap grid mesh MID */
-	class UMaterialInstanceDynamic* GetSnapGridMID()
-	{
-		SpawnAvatarMeshActor();
-		return SnapGridMID;
-	}
-
 	/** Called internally when the user changes maps, enters/exits PIE or SIE, or switched between PIE/SIE */
 	void CleanUpActorsBeforeMapChangeOrSimulate();
 
-	/** Invokes the editor's undo system to undo the last change */
-	void Undo();
-	
-	/** Redoes the last change */
-	void Redo();
-
-	/** Copies selected objects to the clipboard */
-	void Copy();
-
-	/** Pastes the clipboard contents into the scene */
-	void Paste();
-
-	/** Duplicates the selected objects */
-	void Duplicate();
-
-	/** Snaps the selected objects to the ground */
-	void SnapSelectedActorsToGround();
-
 	/** Spawns a flashlight on the specified hand */
-	void ToggleFlashlight( const int32 HandIndex );
-
-	/** Sets the visuals of the LaserPointer */
-	void SetLaserVisuals( const int32 HandIndex, const FLinearColor& NewColor, const float CrawlFade = 0.0f, const float CrawlSpeed = 0.0f );
-
-	/** Given a world space velocity vector, applies inertial damping to it to slow it down */
-	void ApplyVelocityDamping( FVector& Velocity, const bool bVelocitySensitive );
+	void ToggleFlashlight( class UVREditorInteractor* Interactor );
 
 	/** Will update the TransformGizmo Actor with the next Gizmo type  */
 	void CycleTransformGizmoHandleType();
@@ -284,90 +159,34 @@ public:
 	/** Gets the current Gizmo handle type */
 	EGizmoHandleTypes GetCurrentGizmoType() const;
 
-	/** Switch which transform gizmo coordinate space we're using. */
-	void CycleTransformGizmoCoordinateSpace();
-
-	/** Sets which transform gizmo coordinate space is used */
-	void SetTransformGizmoCoordinateSpace( const ECoordSystem NewCoordSystem );
-
-	/** Returns which transform gizmo coordinate space we're using, world or local */
-	ECoordSystem GetTransformGizmoCoordinateSpace() const;
-		
 	/** @return Returns the type of HMD we're dealing with */
 	EHMDDeviceType::Type GetHMDDeviceType() const;
 
 	/** @return Checks to see if the specified hand is aiming roughly toward the specified capsule */
-	bool IsHandAimingTowardsCapsule( const int32 HandIndex, const FTransform& CapsuleTransform, const FVector CapsuleStart, const FVector CapsuleEnd, const float CapsuleRadius, const float MinDistanceToCapsule, const FVector CapsuleFrontDirection, const float MinDotForAimingAtCapsule ) const;
+	bool IsHandAimingTowardsCapsule( class UViewportInteractor* Interactor, const FTransform& CapsuleTransform, const FVector CapsuleStart, const FVector CapsuleEnd, const float CapsuleRadius, const float MinDistanceToCapsule, const FVector CapsuleFrontDirection, const float MinDotForAimingAtCapsule ) const;
 	
-	/**
-	 * Creates a hand transform and forward vector for a laser pointer for a given hand
-	 *
-	 * @param HandIndex			Index of the hand to use
-	 * @param OutHandTransform	The created hand transform
-	 * @param OutForwardVector	The forward vector of the hand
-	 *
-	 * @return	True if we have motion controller data for this hand and could return a valid result
-	 */
-	bool GetHandTransformAndForwardVector( int32 HandIndex, FTransform& OutHandTransform, FVector& OutForwardVector ) const;
-
-	/** Returns the maximum user scale */
-	float GetMaxScale();
-
-	/** Returns the minimum user scale */
-	float GetMinScale();
-
-	/** Sets GNewWorldToMetersScale */
-	void SetWorldToMetersScale( const float NewWorldToMetersScale );
+	/** Gets the hand interactor  */
+	class UVREditorInteractor* GetHandInteractor( const EControllerHand ControllerHand ) const;
+
+	/** Stops the haptic feedback for the left and right hand interactors for the oculus */
+	void StopOldHapticEffects();
 
 private:
 
 	/** Called when someone closes a standalone VR Editor window */
 	void OnVREditorWindowClosed( const TSharedRef<SWindow>& ClosedWindow );
 
-	/**
-	 * Translates a raw key into a vr editing action
-	 *
-	 * @param InKey FKey to translate
-	 * @return The action mapped to InKey
-	 */
-	int32 GetHandIndexFromKey( const FKey& InKey ) const;
-
-	/** Conditionally polls input from controllers, if we haven't done that already this frame */
-	void PollInputIfNeeded();
-
-	/** Grabs latest data from the controllers and updates our virtual hands.  You should call PollInputIfNeeded() instead. */
-	void PollInputFromMotionControllers();
-	
-	/** Stops playing any haptic effects that have been going for a while.  Called every frame. */
-	void StopOldHapticEffects();
-
-	/** Pops up some help text labels for the controller in the specified hand, or hides it, if requested */
-	void ShowHelpForHand( const int32 HandIndex, const bool bShowIt );
-
-	/** Called every frame to update the position of any floating help labels */
-	void UpdateHelpLabels();
-
-	/** Given a mesh and a key name, tries to find a socket on the mesh that matches a supported key */
-	class UStaticMeshSocket* FindMeshSocketForKey( UStaticMesh* StaticMesh, const FKey Key );
-
 	/** FEditorDelegates callbacks */
 	void OnMapChange( uint32 MapChangeFlags );
 	void OnBeginPIE( const bool bIsSimulatingInEditor );
 	void OnEndPIE( const bool bIsSimulatingInEditor );
 	void OnSwitchBetweenPIEAndSIE( const bool bIsSimulatingInEditor );
 
-	/** Changes the color of the buttons on the handmesh */
-	void ApplyButtonPressColors( FVRAction VRAction, EInputEvent Event );
-
-	/** Set the visuals for a button on the motion controller */
-	void SetMotionControllerButtonPressedVisuals( FVirtualHand& Hand, EInputEvent Event, const FName& ParameterName, const float PressStrength );
-
 protected:
 
 	/** Static ID name for our editor mode */
 	static FEditorModeID VREditorModeID;
 
-
 	//
 	// Startup/Shutdown
 	//
@@ -461,21 +280,6 @@ protected:
 	/** Post process material for "greying out" the world while in world movement mode */
 	class UMaterialInstanceDynamic* WorldMovementPostProcessMaterial;
 
-
-	//
-	// Snap grid
-	//
-
-	/** Actor for the snap grid */
-	class AActor* SnapGridActor;
-
-	/** The plane mesh we use to draw a snap grid under selected objects */
-	class UStaticMeshComponent* SnapGridMeshComponent;
-
-	/** MID for the snap grid, so we can update snap values on the fly */
-	class UMaterialInstanceDynamic* SnapGridMID;
-
-
 	//
 	// World scaling progress bar
 	//
@@ -504,45 +308,42 @@ protected:
 	/** The Unreal controller ID for the motion controllers we're using */
 	int32 MotionControllerID;
 
-	/** The last frame that input was polled.  We keep track of this so that we can make sure to poll for latest input right before
-	    its first needed each frame */
-	uint32 LastFrameNumberInputWasPolled;
 
-	/** Mapping of raw keys to actions*/
-	TMap<FKey,FVRAction> KeyToActionMap;
-
-	/** List of virtual hands being controlled */
-	FVirtualHand VirtualHands[VREditorConstants::NumVirtualHands];
-
-	/** Event that's triggered when the user presses a button on their motion controller device.  Set bWasHandled to true if you reacted to the event. */
-	FOnVRAction OnVRActionEvent;
-
-	/** Event that fires every frame to update hover state based on what's under the cursor.  Set bWasHandled to true if you detect something to hover. */
-	FOnVRHoverUpdate OnVRHoverUpdateEvent;
+	//
+	// Subsystems registered
+	//
 
+	FOnVRTickHandle TickHandle;
 
 	//
-	// UI
+	// Subsystems
 	//
 
 	/** VR UI system */
-	TUniquePtr<class FVREditorUISystem> UISystem;
+	class UVREditorUISystem* UISystem;
 
+	/** Teleporter system */
+	class UVREditorTeleporter* TeleporterSystem;
 
 	//
 	// World interaction
 	//
 
 	/** World interaction manager */
-	TUniquePtr<class FVREditorWorldInteraction> WorldInteraction;
+	class UVREditorWorldInteraction* WorldInteraction;
 
 	/** The current Gizmo type that is used for the TransformGizmo Actor */
 	EGizmoHandleTypes CurrentGizmoType;
 
-	// Slate Input Processor
-	TSharedPtr<FVREditorInputProcessor> InputProcessor;
-
+	//
+	// Interactors
+	//
 
+	/** The right motion controller */
+	class UVREditorInteractor* LeftHandInteractor; //@todo vreditor: Hardcoded interactors
+	
+	/** The right motion controller */
+	class UVREditorInteractor* RightHandInteractor; 
 
 	//
 	// Colors
diff --git a/Engine/Source/Editor/VREditor/VREditorQuickMenu.cpp b/Engine/Source/Editor/VREditor/VREditorQuickMenu.cpp
--- a/Engine/Source/Editor/VREditor/VREditorQuickMenu.cpp
+++ b/Engine/Source/Editor/VREditor/VREditorQuickMenu.cpp
@@ -32,25 +32,25 @@ void UVREditorQuickMenu::OnDeleteButtonClicked()
 
 void UVREditorQuickMenu::OnUndoButtonClicked()
 {
-	GetOwner()->GetOwner().GetOwner().Undo();
+	GetOwner()->GetOwner().GetOwner().GetWorldInteraction().Undo();
 }
 
 
 void UVREditorQuickMenu::OnRedoButtonClicked()
 {
-	GetOwner()->GetOwner().GetOwner().Redo();
+	GetOwner()->GetOwner().GetOwner().GetWorldInteraction().Redo();
 }
 
 
 void UVREditorQuickMenu::OnCopyButtonClicked()
 {
-	GetOwner()->GetOwner().GetOwner().Copy();
+	GetOwner()->GetOwner().GetOwner().GetWorldInteraction().Copy();
 }
 
 
 void UVREditorQuickMenu::OnPasteButtonClicked()
 {
-	GetOwner()->GetOwner().GetOwner().Paste();
+	GetOwner()->GetOwner().GetOwner().GetWorldInteraction().Paste();
 }
 
 
@@ -170,22 +170,21 @@ bool UVREditorQuickMenu::IsScaleSnapEnabled() const
 
 FText UVREditorQuickMenu::OnGizmoCoordinateSystemButtonClicked()
 {
-	GetOwner()->GetOwner().GetOwner().CycleTransformGizmoCoordinateSpace();
-
+	GetOwner()->GetOwner().GetOwner().GetWorldInteraction().CycleTransformGizmoCoordinateSpace(); 
 	return GetGizmoCoordinateSystemText();
 }
 
 
 FText UVREditorQuickMenu::GetGizmoCoordinateSystemText() const
 {
-	const ECoordSystem CurrentCoordSystem = GetOwner()->GetOwner().GetOwner().GetTransformGizmoCoordinateSpace();
+	const ECoordSystem CurrentCoordSystem = GetOwner()->GetOwner().GetOwner().GetWorldInteraction().GetTransformGizmoCoordinateSpace(); //@todo VREditor
 	return CurrentCoordSystem == COORD_World ? LOCTEXT( "WorldCoordinateSystem", "World Space" ) : LOCTEXT( "LocalCoordinateSystem", "Local Space" );
 }
 
 
 FText UVREditorQuickMenu::OnGizmoModeButtonClicked()
 {
-	const EGizmoHandleTypes CurrentGizmoMode = GetOwner()->GetOwner().GetOwner().GetCurrentGizmoType();
+	const EGizmoHandleTypes CurrentGizmoMode = GetOwner()->GetOwner().GetOwner().GetWorldInteraction().GetCurrentGizmoType();
 	GetOwner()->GetOwner().GetOwner().CycleTransformGizmoHandleType();
 	return GetGizmoModeText();
 }
@@ -193,7 +192,7 @@ FText UVREditorQuickMenu::OnGizmoModeButtonClicked()
 
 FText UVREditorQuickMenu::GetGizmoModeText() const
 {
-	const EGizmoHandleTypes CurrentGizmoType = GetOwner()->GetOwner().GetOwner().GetCurrentGizmoType();
+	const EGizmoHandleTypes CurrentGizmoType = GetOwner()->GetOwner().GetOwner().GetWorldInteraction().GetCurrentGizmoType();
 	FText GizmoTypeText;
 	switch( CurrentGizmoType )
 	{
@@ -262,57 +261,53 @@ bool UVREditorQuickMenu::IsSimulatingEnabled() const
 
 bool UVREditorQuickMenu::OnContentBrowserButtonClicked( const bool bIsChecked )
 {
-	GetOwner()->GetOwner().TogglePanelVisibility( FVREditorUISystem::EEditorUIPanel::ContentBrowser );
-
+	GetOwner()->GetOwner().TogglePanelVisibility( UVREditorUISystem::EEditorUIPanel::ContentBrowser );
 	return IsContentBrowserVisible();
 }
 
 
 bool UVREditorQuickMenu::IsContentBrowserVisible() const
 {
-	return ( GetOwner()->GetOwner().IsShowingEditorUIPanel( FVREditorUISystem::EEditorUIPanel::ContentBrowser ) );
+	return ( GetOwner()->GetOwner().IsShowingEditorUIPanel( UVREditorUISystem::EEditorUIPanel::ContentBrowser ) );
 }
 
 
 bool UVREditorQuickMenu::OnWorldOutlinerButtonClicked( const bool bIsChecked )
 {
-	GetOwner()->GetOwner().TogglePanelVisibility( FVREditorUISystem::EEditorUIPanel::WorldOutliner );
-
+	GetOwner()->GetOwner().TogglePanelVisibility( UVREditorUISystem::EEditorUIPanel::WorldOutliner );
 	return IsWorldOutlinerVisible();
 }
 
 
 bool UVREditorQuickMenu::IsWorldOutlinerVisible() const
 {
-	return ( GetOwner()->GetOwner().IsShowingEditorUIPanel( FVREditorUISystem::EEditorUIPanel::WorldOutliner ) );
+	return ( GetOwner()->GetOwner().IsShowingEditorUIPanel( UVREditorUISystem::EEditorUIPanel::WorldOutliner ) );
 }
 
 
 bool UVREditorQuickMenu::OnActorDetailsButtonClicked( const bool bIsChecked )
 {
-	GetOwner()->GetOwner().TogglePanelVisibility( FVREditorUISystem::EEditorUIPanel::ActorDetails );
-
+	GetOwner()->GetOwner().TogglePanelVisibility( UVREditorUISystem::EEditorUIPanel::ActorDetails );
 	return IsActorDetailsVisible();
 }
 
 
 bool UVREditorQuickMenu::IsActorDetailsVisible() const
 {
-	return ( GetOwner()->GetOwner().IsShowingEditorUIPanel( FVREditorUISystem::EEditorUIPanel::ActorDetails ) );
+	return ( GetOwner()->GetOwner().IsShowingEditorUIPanel( UVREditorUISystem::EEditorUIPanel::ActorDetails ) );
 }
 
 
 bool UVREditorQuickMenu::OnModesButtonClicked( const bool bIsChecked )
 {
-	GetOwner()->GetOwner().TogglePanelVisibility( FVREditorUISystem::EEditorUIPanel::Modes );
-
+	GetOwner()->GetOwner().TogglePanelVisibility( UVREditorUISystem::EEditorUIPanel::Modes );
 	return IsModesVisible();
 }
 
 
 bool UVREditorQuickMenu::IsModesVisible() const
 {
-	return ( GetOwner()->GetOwner().IsShowingEditorUIPanel( FVREditorUISystem::EEditorUIPanel::Modes ) );
+	return ( GetOwner()->GetOwner().IsShowingEditorUIPanel( UVREditorUISystem::EEditorUIPanel::Modes ) );
 }
 
 
@@ -338,21 +333,20 @@ bool UVREditorQuickMenu::IsGameModeEnabled() const
 
 bool UVREditorQuickMenu::OnTutorialButtonClicked( const bool bIsChecked )
 {
-	GetOwner()->GetOwner().TogglePanelVisibility( FVREditorUISystem::EEditorUIPanel::Tutorial );
-
+	GetOwner()->GetOwner().TogglePanelVisibility( UVREditorUISystem::EEditorUIPanel::Tutorial );
 	return IsTutorialVisible();
 }
 
 bool UVREditorQuickMenu::IsTutorialVisible() const
 {
-	return ( GetOwner()->GetOwner().IsShowingEditorUIPanel( FVREditorUISystem::EEditorUIPanel::Tutorial ) );
+	return ( GetOwner()->GetOwner().IsShowingEditorUIPanel( UVREditorUISystem::EEditorUIPanel::Tutorial ) );
 }
 
-bool UVREditorQuickMenu::OnLightButtonClicked(const bool bIsChecked)
+bool UVREditorQuickMenu::OnLightButtonClicked( const bool bIsChecked ) 
 {
-	const int32 HandIndexWithoutQuickMenu = GetOwner()->GetDockedTo() == AVREditorFloatingUI::EDockedTo::LeftArm ? VREditorConstants::RightHandIndex : VREditorConstants::LeftHandIndex;
-	GetOwner()->GetOwner().GetOwner().ToggleFlashlight(HandIndexWithoutQuickMenu);
-
+	EControllerHand ControllerHand = GetOwner()->GetDockedTo() == AVREditorFloatingUI::EDockedTo::LeftArm ? EControllerHand::Right : EControllerHand::Left;
+	UVREditorInteractor* Interactor = GetOwner()->GetOwner().GetOwner().GetHandInteractor( ControllerHand );
+	GetOwner()->GetOwner().GetOwner().ToggleFlashlight( Interactor );
 	return IsLightVisible();
 }
 
@@ -387,25 +381,24 @@ void UVREditorQuickMenu::OnScreenshotButtonClicked()
 
 bool UVREditorQuickMenu::OnAssetEditorButtonClicked( const bool bIsChecked )
 {
-	GetOwner()->GetOwner().TogglePanelVisibility( FVREditorUISystem::EEditorUIPanel::AssetEditor );
+	GetOwner()->GetOwner().TogglePanelVisibility( UVREditorUISystem::EEditorUIPanel::AssetEditor );
 	return IsAssetEditorVisible();
 }
 
 bool UVREditorQuickMenu::IsAssetEditorVisible() const
 {
-	return ( GetOwner()->GetOwner().IsShowingEditorUIPanel( FVREditorUISystem::EEditorUIPanel::AssetEditor ) );
+	return ( GetOwner()->GetOwner().IsShowingEditorUIPanel( UVREditorUISystem::EEditorUIPanel::AssetEditor ) );
 }
 
-bool UVREditorQuickMenu::OnWorldSettingsButtonClicked(const bool bIsChecked)
+bool UVREditorQuickMenu::OnWorldSettingsButtonClicked( const bool bIsChecked )
 {
-	GetOwner()->GetOwner().TogglePanelVisibility(FVREditorUISystem::EEditorUIPanel::WorldSettings);
-
+	GetOwner()->GetOwner().TogglePanelVisibility( UVREditorUISystem::EEditorUIPanel::WorldSettings);
 	return IsWorldSettingsVisible();
 }
 
 bool UVREditorQuickMenu::IsWorldSettingsVisible() const
 {
-	return (GetOwner()->GetOwner().IsShowingEditorUIPanel(FVREditorUISystem::EEditorUIPanel::WorldSettings));
+	return ( GetOwner()->GetOwner().IsShowingEditorUIPanel( UVREditorUISystem::EEditorUIPanel::WorldSettings ) );
 }
 
 #undef LOCTEXT_NAMESPACE
\ No newline at end of file
diff --git a/Engine/Source/Editor/VREditor/VREditorRadialMenu.cpp b/Engine/Source/Editor/VREditor/VREditorRadialMenu.cpp
--- a/Engine/Source/Editor/VREditor/VREditorRadialMenu.cpp
+++ b/Engine/Source/Editor/VREditor/VREditorRadialMenu.cpp
@@ -8,7 +8,7 @@
 #include "VREditorMode.h"
 #include "VREditorWorldInteraction.h"
 #include "VREditorTransformGizmo.h"
-#include "VirtualHand.h"
+#include "VREditorInteractor.h"
 
 #include "Editor/LevelEditor/Public/LevelEditor.h"
 
@@ -28,14 +28,14 @@ PreviousItem( nullptr )
 }
 
 
-void UVREditorRadialMenu::Update( const FVirtualHand& Hand )
+void UVREditorRadialMenu::Update( const UVREditorInteractor* Interactor )
 {
-	TrackpadPosition = Hand.TrackpadPosition;
+	TrackpadPosition = Interactor->GetTrackpadPosition();
 
-	if (!GetOwner()->bHidden)
+	if ( !GetOwner()->bHidden )
 	{
 		// Check if position is from the center
-		if (!IsInMenuRadius())
+		if ( !IsInMenuRadius() )
 		{
 			const float AnglePerItem = 360 / MenuItems.Num();
 			float Angle = FRotator::NormalizeAxis( FMath::RadiansToDegrees( FMath::Atan2( TrackpadPosition.X, TrackpadPosition.Y ) ) );
@@ -131,7 +131,7 @@ void UVREditorRadialMenu::ResetItem()
 	}
 }
 
-void UVREditorRadialMenu::SelectCurrentItem( const FVirtualHand& Hand )
+void UVREditorRadialMenu::SelectCurrentItem()
 {
 	if (IsInMenuRadius())
 	{
@@ -153,7 +153,7 @@ void UVREditorRadialMenu::UpdateGizmoTypeLabel()
 {
 	FString Result;
 
-	switch (GetOwner()->GetOwner().GetOwner().GetCurrentGizmoType())
+	switch (GetOwner()->GetOwner().GetOwner().GetWorldInteraction().GetCurrentGizmoType())
 	{
 	case EGizmoHandleTypes::All:
 		Result = LOCTEXT( "AllGizmoType", "Universal Gizmo" ).ToString();
@@ -182,7 +182,7 @@ void UVREditorRadialMenu::OnGizmoCycle()
 
 void UVREditorRadialMenu::OnDuplicateButtonClicked()
 {
-	GetOwner()->GetOwner().GetOwner().Duplicate();
+	GetOwner()->GetOwner().GetOwner().GetWorldInteraction().Duplicate();
 }
 
 
@@ -194,30 +194,34 @@ void UVREditorRadialMenu::OnDeleteButtonClicked()
 
 void UVREditorRadialMenu::OnUndoButtonClicked()
 {
-	GetOwner()->GetOwner().GetOwner().Undo();
+	GetOwner()->GetOwner().GetOwner().GetWorldInteraction().Undo();
 }
 
 
 void UVREditorRadialMenu::OnRedoButtonClicked()
 {
-	GetOwner()->GetOwner().GetOwner().Redo();
+	GetOwner()->GetOwner().GetOwner().GetWorldInteraction().Redo();
 }
 
 
 void UVREditorRadialMenu::OnCopyButtonClicked()
 {
-	GetOwner()->GetOwner().GetOwner().Copy();
+	GetOwner()->GetOwner().GetOwner().GetWorldInteraction().Copy();
 }
 
 
 void UVREditorRadialMenu::OnPasteButtonClicked()
 {
-	GetOwner()->GetOwner().GetOwner().Paste();
+	GetOwner()->GetOwner().GetOwner().GetWorldInteraction().Paste();
 }
 
 void UVREditorRadialMenu::OnSnapActorsToGroundClicked()
 {
-	GetOwner()->GetOwner().GetOwner().SnapSelectedActorsToGround();
+	UVREditorWorldInteraction* WorldInteraction = Cast<UVREditorWorldInteraction>( &( GetOwner()->GetOwner().GetOwner().GetWorldInteraction() ) );
+	if ( WorldInteraction )
+	{
+		WorldInteraction->SnapSelectedActorsToGround();
+	}
 }
 
 
diff --git a/Engine/Source/Editor/VREditor/VREditorRadialMenu.h b/Engine/Source/Editor/VREditor/VREditorRadialMenu.h
--- a/Engine/Source/Editor/VREditor/VREditorRadialMenu.h
+++ b/Engine/Source/Editor/VREditor/VREditorRadialMenu.h
@@ -21,7 +21,7 @@ public:
 	UVREditorRadialMenu( const FObjectInitializer& ObjectInitializer );
 
 	/** Update the state of the widget with the trackpad location */
-	void Update( const struct FVirtualHand& Hand );
+	void Update( const class UVREditorInteractor* Interactor );
 
 	/** To bind the buttons from UMG */
 	UFUNCTION(BlueprintCallable, Category = "RadialItems")
@@ -32,7 +32,7 @@ public:
 	void ResetItem();
 
 	/** Press the current item */
-	void SelectCurrentItem( const FVirtualHand& Hand );
+	void SelectCurrentItem();
 
 	/** If the users has his trackpad or analog stick inside the radius where buttons will not be selected */
 	bool IsInMenuRadius() const;
@@ -47,6 +47,7 @@ public:
 
 private:
 	/** All the buttons in the radial menu */
+	UPROPERTY()
 	TArray<UVREditorRadialMenuItem*> MenuItems;
 
 	/** The current widget currently selected */
diff --git a/Engine/Source/Editor/VREditor/VREditorTransformGizmo.cpp b/Engine/Source/Editor/VREditor/VREditorTransformGizmo.cpp
--- a/Engine/Source/Editor/VREditor/VREditorTransformGizmo.cpp
+++ b/Engine/Source/Editor/VREditor/VREditorTransformGizmo.cpp
@@ -6,8 +6,8 @@
 #include "Gizmo/VREditorTranslationGizmoHandle.h"
 #include "Gizmo/VREditorRotationGizmoHandle.h"
 #include "Gizmo/VREditorPlaneTranslationGizmoHandle.h"
-#include "Gizmo/VREditorStretchGizmoHandle.h"
-#include "Gizmo/VREditorUniformScaleGizmoHandle.h"
+#include "VIStretchGizmoHandle.h"
+#include "VIUniformScaleGizmoHandle.h"
 #include "VREditorMode.h"
 
 namespace VREd
@@ -39,13 +39,13 @@ ATransformGizmo::ATransformGizmo()
 	RotationGizmoHandleGroup->SetupAttachment( SceneComponent );
 	AllHandleGroups.Add( RotationGizmoHandleGroup );
 
-	StretchGizmoHandleGroup = CreateDefaultSubobject<UVREditorStretchGizmoHandleGroup>( TEXT( "StretchHandles" ), true );
+	StretchGizmoHandleGroup = CreateDefaultSubobject<UStretchGizmoHandleGroup>( TEXT( "StretchHandles" ), true );
 	StretchGizmoHandleGroup->SetTranslucentGizmoMaterial( TranslucentGizmoMaterial );
 	StretchGizmoHandleGroup->SetGizmoMaterial( GizmoMaterial );
 	StretchGizmoHandleGroup->SetupAttachment( SceneComponent );
 	AllHandleGroups.Add( StretchGizmoHandleGroup );
 
-	UniformScaleGizmoHandleGroup = CreateDefaultSubobject<UVREditorUniformScaleGizmoHandleGroup>( TEXT( "UniformScaleHandles" ), true );
+	UniformScaleGizmoHandleGroup = CreateDefaultSubobject<UUniformScaleGizmoHandleGroup>( TEXT( "UniformScaleHandles" ), true );
 	UniformScaleGizmoHandleGroup->SetTranslucentGizmoMaterial( TranslucentGizmoMaterial );
 	UniformScaleGizmoHandleGroup->SetGizmoMaterial( GizmoMaterial );
 	UniformScaleGizmoHandleGroup->SetUsePivotPointAsLocation( false );
@@ -96,8 +96,8 @@ ATransformGizmo::ATransformGizmo()
 
 			MeasurementText->SetWorldSize( 8.0f );	// @todo vreditor tweak: Size of text in world units
 
-													// Assign our custom text rendering material.  It should be Unlit and have Depth Test disabled
-													// so that it draws on top of overlapping geometry.
+			// Assign our custom text rendering material.  It should be Unlit and have Depth Test disabled
+			// so that it draws on top of overlapping geometry.
 			MeasurementText->SetTextMaterial( GizmoTextMaterial );
 
 			// Color the text based on the axis it represents
@@ -163,9 +163,9 @@ void ATransformGizmo::UpdateGizmo( const EGizmoHandleTypes GizmoType, const ECoo
 
 	// Update all the handles
 	bool bIsHoveringOrDraggingScaleGizmo = false;
-	for (UVREditorGizmoHandleGroup* HandleGroup : AllHandleGroups)
+	for ( UGizmoHandleGroup* HandleGroup : AllHandleGroups )
 	{
-		if (HandleGroup != nullptr)
+		if ( HandleGroup != nullptr )
 		{
 			bool bIsHoveringOrDraggingThisHandleGroup = false;
 			HandleGroup->UpdateGizmoHandleGroup( LocalToWorld, LocalBounds, ViewLocation, bAllHandlesVisible, DraggingHandle, 
@@ -249,7 +249,7 @@ void ATransformGizmo::UpdateGizmo( const EGizmoHandleTypes GizmoType, const ECoo
 				Measurement.MeasurementText->SetRelativeLocation( GizmoSpaceEdgeCenter + GizmoSpaceUpVector * UpOffsetAmount );
 
 				const int32 TextFacingAxisIndex = AxisIndex == 0 ? 1 : ( AxisIndex == 1 ? 0 : 0 );
-				FVector TextFacingDirection = UVREditorGizmoHandleGroup::GetAxisVector( TextFacingAxisIndex, ETransformGizmoHandleDirection::Positive );
+				FVector TextFacingDirection = UGizmoHandleGroup::GetAxisVector( TextFacingAxisIndex, ETransformGizmoHandleDirection::Positive );
 
 				// Make sure to face the camera
 				const FVector WorldSpaceEdgeCenterToViewDirection = ( LocalToWorld.TransformPosition( GizmoSpaceEdgeCenter ) - ViewLocation ).GetSafeNormal();
@@ -268,8 +268,6 @@ void ATransformGizmo::UpdateGizmo( const EGizmoHandleTypes GizmoType, const ECoo
 
 	// Update visibility
 	{
-		UpdateHandleVisibility( GizmoType, GizmoCoordinateSpace, bAllHandlesVisible, DraggingHandle );
-
 		const bool bShouldShowMeasurements =
 			GetShowMeasurementText() ||
 			bIsHoveringOrDraggingScaleGizmo;
diff --git a/Engine/Source/Editor/VREditor/VREditorTransformGizmo.h b/Engine/Source/Editor/VREditor/VREditorTransformGizmo.h
--- a/Engine/Source/Editor/VREditor/VREditorTransformGizmo.h
+++ b/Engine/Source/Editor/VREditor/VREditorTransformGizmo.h
@@ -2,23 +2,9 @@
 
 #pragma once
 
-#include "Gizmo/VREditorBaseTransformGizmo.h"
+#include "VIBaseTransformGizmo.h"
 #include "VREditorTransformGizmo.generated.h"
 
-/**
- * Displays measurements along the bounds of selected objects
- */
-USTRUCT()
-struct FTransformGizmoMeasurement
-{
-	GENERATED_BODY()
-
-	/** The text that displays the actual measurement and units */
-	UPROPERTY()
-	class UTextRenderComponent* MeasurementText;
-};
-
-
 /**
  * A transform gizmo that allows you to interact with selected objects by moving, scaling and rotating.
  */
@@ -44,18 +30,23 @@ private:
 	UPROPERTY()
 	FTransformGizmoMeasurement Measurements[ 3 ];
 
+	/** Bounding box translation handles */
 	UPROPERTY()
 	class UVREditorTranslationGizmoHandleGroup* TranslationGizmoHandleGroup;
 
+	/** Bounding box plane translation handles */
 	UPROPERTY()
 	class UVREditorPlaneTranslationGizmoHandleGroup* PlaneTranslationGizmoHandleGroup;
 
+	/** Bounding box rotation handles */
 	UPROPERTY()
 	class UVREditorRotationGizmoHandleGroup* RotationGizmoHandleGroup;
 
+	/** Bounding box stretch handles */
 	UPROPERTY()
-	class UVREditorStretchGizmoHandleGroup* StretchGizmoHandleGroup;
+	class UStretchGizmoHandleGroup* StretchGizmoHandleGroup;
 
+	/** Bounding box uniform scale handle */
 	UPROPERTY()
-	class UVREditorUniformScaleGizmoHandleGroup* UniformScaleGizmoHandleGroup;
+	class UUniformScaleGizmoHandleGroup* UniformScaleGizmoHandleGroup;
 };
\ No newline at end of file
diff --git a/Engine/Source/Editor/VREditor/VREditorUISystem.cpp b/Engine/Source/Editor/VREditor/VREditorUISystem.cpp
--- a/Engine/Source/Editor/VREditor/VREditorUISystem.cpp
+++ b/Engine/Source/Editor/VREditor/VREditorUISystem.cpp
@@ -8,9 +8,10 @@
 #include "VREditorQuickMenu.h"
 #include "VREditorRadialMenu.h"
 #include "VREditorRadialMenuItem.h"
-#include "VirtualHand.h"
 #include "IHeadMountedDisplay.h"
-
+#include "VREditorWorldInteraction.h"
+#include "ViewportInteractor.h"
+#include "VREditorInteractor.h"
 // UI
 #include "WidgetComponent.h"
 #include "VREditorWidgetComponent.h"
@@ -80,31 +81,28 @@ namespace VREd
 }
 
 
-FVREditorUISystem::FVREditorUISystem( FVREditorMode& InitOwner )
-	: Owner( InitOwner ),
-	  FloatingUIs(),
-	  QuickMenuUI( nullptr ),
-	  QuickRadialMenu( nullptr ),
-	  RadialMenuHideDelayTime( 0.0f ),
-	  QuickMenuWidgetClass( nullptr ),
-	  QuickRadialWidgetClass( nullptr ),
-	  TutorialWidgetClass( nullptr ),
-	  DraggingUIHandIndex( INDEX_NONE ),
-	  DraggingUIOffsetTransform( FTransform::Identity ),
-	  DraggingUI( nullptr ),
-	  bRefocusViewport(false),
-	  StartDragUISound( nullptr ),
-	  StopDragUISound( nullptr ),
-	  HideUISound( nullptr ),
-	  ShowUISound( nullptr ),
-	  bDraggedDockFromHandPassedThreshold( false ),
-	  LastDraggingHoverLocation( FVector::ZeroVector ),
-	  bSetDefaultLayout( true )
+UVREditorUISystem::UVREditorUISystem( const FObjectInitializer& Initializer ) : 
+	Super( Initializer ),
+	VRMode( nullptr ),
+	FloatingUIs(),
+	QuickMenuUI( nullptr ),
+	QuickRadialMenu( nullptr ),
+	RadialMenuHideDelayTime( 0.0f ),
+	QuickMenuWidgetClass( nullptr ),
+	QuickRadialWidgetClass( nullptr ),
+	TutorialWidgetClass( nullptr ),
+	InteractorDraggingUI( nullptr ),
+	DraggingUIOffsetTransform( FTransform::Identity ),
+	DraggingUI( nullptr ),
+	bRefocusViewport(false),
+	StartDragUISound( nullptr ),
+	StopDragUISound( nullptr ),
+	HideUISound( nullptr ),
+	ShowUISound( nullptr ),
+	bDraggedDockFromHandPassedThreshold( false ),
+	LastDraggingHoverLocation( FVector::ZeroVector ),
+	bSetDefaultLayout( true )
 {
-	// Register to find out about VR events
-	Owner.OnVRAction().AddRaw( this, &FVREditorUISystem::OnVRAction );
-	Owner.OnVRHoverUpdate().AddRaw( this, &FVREditorUISystem::OnVRHoverUpdate );
-
 	EditorUIPanels.SetNumZeroed( (int32)EEditorUIPanel::TotalCount );
 	
 	// Set default layout transform in correct order of enum
@@ -134,17 +132,26 @@ FVREditorUISystem::FVREditorUISystem( FVREditorMode& InitOwner )
 
 	ShowUISound = LoadObject<USoundCue>( nullptr, TEXT( "/Engine/VREditor/Sounds/VR_open_Cue" ) );
 	check( ShowUISound != nullptr );
+}
 
+void UVREditorUISystem::Init()
+{
+	// Register to find out about VR events
+	GetOwner().GetWorldInteraction().OnViewportInteractionInputAction().AddUObject( this, &UVREditorUISystem::OnVRAction );
+	GetOwner().GetWorldInteraction().OnViewportInteractionHoverUpdate().AddUObject( this, &UVREditorUISystem::OnVRHoverUpdate );
+	
 	// Create all of our UI panels
 	CreateUIs();
 }
 
 
-
-FVREditorUISystem::~FVREditorUISystem()
+void UVREditorUISystem::Shutdown()
 {
-	Owner.OnVRAction().RemoveAll( this );
-	Owner.OnVRHoverUpdate().RemoveAll( this );
+	if ( VRMode != nullptr && &VRMode->GetWorldInteraction() != nullptr )
+	{
+		GetOwner().GetWorldInteraction().OnViewportInteractionInputAction().RemoveAll( this );
+		GetOwner().GetWorldInteraction().OnViewportInteractionHoverUpdate().RemoveAll( this );
+	}
 
 	CleanUpActorsBeforeMapChangeOrSimulate();
 
@@ -157,96 +164,81 @@ FVREditorUISystem::~FVREditorUISystem()
 	StopDragUISound = nullptr;
 	HideUISound = nullptr;
 	ShowUISound = nullptr;
+	VRMode = nullptr;
 }
 
-
-void FVREditorUISystem::AddReferencedObjects( FReferenceCollector& Collector )
-{
-	for( AVREditorFloatingUI* FloatingUIPtr : FloatingUIs )
-	{
-		Collector.AddReferencedObject( FloatingUIPtr );
-	}
-	Collector.AddReferencedObject( QuickMenuWidgetClass );
-	Collector.AddReferencedObject( QuickRadialWidgetClass );
-	Collector.AddReferencedObject( TutorialWidgetClass );
-	Collector.AddReferencedObject( StartDragUISound );
-	Collector.AddReferencedObject( StopDragUISound );
-	Collector.AddReferencedObject( HideUISound );
-	Collector.AddReferencedObject( ShowUISound );
-}
-
-
-void FVREditorUISystem::OnVRAction( FEditorViewportClient& ViewportClient, const FVRAction VRAction, const EInputEvent Event, bool& bIsInputCaptured, bool& bWasHandled )
+void UVREditorUISystem::OnVRAction( FEditorViewportClient& ViewportClient, UViewportInteractor* Interactor,
+	const FViewportActionKeyInput& Action, const EInputEvent Event, bool& bOutIsInCaptured, bool& bWasHandled )
 {
-	if( !bWasHandled )
+	UVREditorInteractor* VREditorInteractor = Cast<UVREditorInteractor>( Interactor );
+	if ( VREditorInteractor )
 	{
-
-		if( VRAction.ActionType == EVRActionType::ConfirmRadialSelection )
+		if ( !bWasHandled )
 		{
-			FVirtualHand& Hand = Owner.GetVirtualHand( VRAction.HandIndex );
-			if( Event == IE_Pressed )
+			if ( Action.ActionType == VRActionTypes::ConfirmRadialSelection )
 			{
-				UVREditorRadialMenu* RadialMenu = QuickRadialMenu->GetUserWidget<UVREditorRadialMenu>();
-				if( IsShowingRadialMenu( VRAction.HandIndex ) )
+				if ( Event == IE_Pressed )
 				{
-					RadialMenu->Update( Owner.GetVirtualHand( VRAction.HandIndex ) );
+					UVREditorRadialMenu* RadialMenu = QuickRadialMenu->GetUserWidget<UVREditorRadialMenu>();
+					if ( IsShowingRadialMenu( VREditorInteractor ) )
+					{
+						RadialMenu->Update( VREditorInteractor );
+					}
+
+					RadialMenu->SelectCurrentItem();
 				}
-				
-				RadialMenu->SelectCurrentItem( Hand );
-			}
 
-			bWasHandled = true;
-		}
-		else if( VRAction.ActionType == EVRActionType::SelectAndMove_LightlyPressed )
-		{
-			FVirtualHand& Hand = Owner.GetVirtualHand( VRAction.HandIndex );
-			FVector LaserPointerStart, LaserPointerEnd;
-			if( Owner.GetLaserPointer( VRAction.HandIndex, LaserPointerStart, LaserPointerEnd ) )
+				bWasHandled = true;
+			}
+			else if ( Action.ActionType == ViewportWorldActionTypes::SelectAndMove_LightlyPressed )
 			{
-				FHitResult HitResult = Owner.GetHitResultFromLaserPointer( VRAction.HandIndex );
-				if( HitResult.Actor.IsValid()  )
+				FVector LaserPointerStart, LaserPointerEnd;
+				if ( VREditorInteractor->GetLaserPointer( LaserPointerStart, LaserPointerEnd ) )
 				{
-					// Only allow clicks to our own widget components
-					UWidgetComponent* WidgetComponent = Cast<UWidgetComponent>( HitResult.GetComponent() );
-					if( WidgetComponent != nullptr && IsWidgetAnEditorUIWidget( WidgetComponent ) )
+					FHitResult HitResult = VREditorInteractor->GetHitResultFromLaserPointer();
+					if ( HitResult.Actor.IsValid() )
 					{
-						// Always mark the event as handled so that the editor doesn't try to select the widget component
-						bWasHandled = true;
-
-						if( Event != IE_Repeat )
+						// Only allow clicks to our own widget components
+						UWidgetComponent* WidgetComponent = Cast<UWidgetComponent>( HitResult.GetComponent() );
+						if ( WidgetComponent != nullptr && IsWidgetAnEditorUIWidget( WidgetComponent ) )
 						{
-							// If the Modifier button is held down, treat this like a right click instead of a left click
-							const bool bIsRightClicking =
-								( Event == IE_Pressed && Hand.bIsModifierPressed ) ||
-								( Event == IE_Released && Hand.bIsRightClickingOnUI );
+							// Always mark the event as handled so that the editor doesn't try to select the widget component
+							bWasHandled = true;
 
+							if ( Event != IE_Repeat )
 							{
+								// If the Modifier button is held down, treat this like a right click instead of a left click
+								const bool bIsRightClicking =
+									( Event == IE_Pressed && VREditorInteractor->IsModifierPressed() ) ||
+									( Event == IE_Released && VREditorInteractor->IsRightClickingOnUI() );
+
 								FVector2D LastLocalHitLocation = WidgetComponent->GetLastLocalHitLocation();
 
 								FVector2D LocalHitLocation;
 								WidgetComponent->GetLocalHitLocation( HitResult.ImpactPoint, LocalHitLocation );
 
 								// If we weren't already hovering over this widget, then we'll reset the last hit location
-								if( WidgetComponent != Hand.HoveringOverWidgetComponent )
+								if ( WidgetComponent != VREditorInteractor->GetHoveringOverWidgetComponent() )
 								{
 									LastLocalHitLocation = LocalHitLocation;
 
-									if ( UVREditorWidgetComponent* VRWidgetComponent = Cast<UVREditorWidgetComponent>(Hand.HoveringOverWidgetComponent) )
+									if ( UVREditorWidgetComponent* VRWidgetComponent = Cast<UVREditorWidgetComponent>( VREditorInteractor->GetHoveringOverWidgetComponent() ) )
 									{
-										VRWidgetComponent->SetIsHovering(false);
+										VRWidgetComponent->SetIsHovering( false );
 									}
 								}
 
 								FWidgetPath WidgetPathUnderFinger = FWidgetPath( WidgetComponent->GetHitWidgetPath( HitResult.ImpactPoint, /*bIgnoreEnabledStatus*/ false ) );
-								if( WidgetPathUnderFinger.IsValid() )
+								if ( WidgetPathUnderFinger.IsValid() )
 								{
 									TSet<FKey> PressedButtons;
-									if( Event == IE_Pressed )
+									if ( Event == IE_Pressed )
 									{
 										PressedButtons.Add( bIsRightClicking ? EKeys::RightMouseButton : EKeys::LeftMouseButton );
 									}
+
 									FPointerEvent PointerEvent(
-										1 + VRAction.HandIndex,
+										1 + (uint8)VREditorInteractor->GetControllerSide(),
 										LocalHitLocation,
 										LastLocalHitLocation,
 										PressedButtons,
@@ -254,18 +246,18 @@ void FVREditorUISystem::OnVRAction( FEditorViewportClient& ViewportClient, const
 										0.0f,	// Wheel delta
 										FModifierKeysState() );
 
-									Hand.HoveringOverWidgetComponent = WidgetComponent;
+									VREditorInteractor->SetHoveringOverWidgetComponent( WidgetComponent );
 
-									if ( UVREditorWidgetComponent* VRWidgetComponent = Cast<UVREditorWidgetComponent>(WidgetComponent) )
+									if ( UVREditorWidgetComponent* VRWidgetComponent = Cast<UVREditorWidgetComponent>( VREditorInteractor->GetHoveringOverWidgetComponent() ) )
 									{
-										VRWidgetComponent->SetIsHovering(true);
+										VRWidgetComponent->SetIsHovering( true );
 									}
 
 									FReply Reply = FReply::Unhandled();
-									if( Event == IE_Pressed )
+									if ( Event == IE_Pressed )
 									{
 										const double CurrentTime = FPlatformTime::Seconds();
-										if (CurrentTime - Hand.LastClickPressTime <= VREd::DoubleClickTime->GetFloat())
+										if (CurrentTime - VREditorInteractor->GetLastClickReleaseTime() <= VREd::DoubleClickTime->GetFloat())
 										{
 											// Trigger a double click event!
 											Reply = FSlateApplication::Get().RoutePointerDoubleClickEvent(WidgetPathUnderFinger, PointerEvent);
@@ -275,19 +267,15 @@ void FVREditorUISystem::OnVRAction( FEditorViewportClient& ViewportClient, const
 											Reply = FSlateApplication::Get().RoutePointerDownEvent(WidgetPathUnderFinger, PointerEvent);
 										}
 										
-										Hand.bIsClickingOnUI = true;
-										Hand.bIsRightClickingOnUI = bIsRightClicking;
-										bIsInputCaptured = true;
-
-										Hand.LastClickPressTime = CurrentTime;
+										VREditorInteractor->SetIsClickingOnUI( true );
+										VREditorInteractor->SetIsRightClickingOnUI( bIsRightClicking );
+										VREditorInteractor->SetLastClickReleaseTime( CurrentTime );
+										bOutIsInCaptured = true;
 
 										// Play a haptic effect on press
-										const float Strength = VREd::UIPressHapticFeedbackStrength->GetFloat();
-										Owner.PlayHapticEffect(
-											VRAction.HandIndex == VREditorConstants::LeftHandIndex ? Strength : 0.0f,
-											VRAction.HandIndex == VREditorConstants::RightHandIndex ? Strength : 0.0f );
+										VREditorInteractor->PlayHapticEffect( VREd::UIPressHapticFeedbackStrength->GetFloat() );
 									}
-									else if( Event == IE_Released )
+									else if ( Event == IE_Released )
 									{
 										Reply = FSlateApplication::Get().RoutePointerUpEvent( WidgetPathUnderFinger, PointerEvent );
 									}
@@ -296,85 +284,56 @@ void FVREditorUISystem::OnVRAction( FEditorViewportClient& ViewportClient, const
 						}
 					}
 				}
-			}
 
-			if( Event == IE_Released )
-			{
-				bool bWasRightClicking = false;
-				if( Hand.bIsClickingOnUI )
+				if ( Event == IE_Released )
 				{
-					if( Hand.bIsRightClickingOnUI )
+					bool bWasRightClicking = false;
+					if ( VREditorInteractor->IsClickingOnUI() )
 					{
-						bWasRightClicking = true;
+						if ( VREditorInteractor->IsRightClickingOnUI() )
+						{
+							bWasRightClicking = true;
+						}
+						VREditorInteractor->SetIsClickingOnUI( false );
+						VREditorInteractor->SetIsRightClickingOnUI( false );
+						bOutIsInCaptured = false;
 					}
-					Hand.bIsClickingOnUI = false;
-					Hand.bIsRightClickingOnUI = false;
-					bIsInputCaptured = false;
-				}
 
-				if ( !bWasHandled )
-				{
-					TSet<FKey> PressedButtons;
-					FPointerEvent PointerEvent(
-						1 + VRAction.HandIndex,
-						FVector2D::ZeroVector,
-						FVector2D::ZeroVector,
-						PressedButtons,
-						bWasRightClicking ? EKeys::RightMouseButton : EKeys::LeftMouseButton,
-						0.0f,	// Wheel delta
-						FModifierKeysState() );
-
-					FWidgetPath EmptyWidgetPath;
-
-					Hand.bIsClickingOnUI = false;
-					Hand.bIsRightClickingOnUI = false;
-					FReply Reply = FSlateApplication::Get().RoutePointerUpEvent( EmptyWidgetPath, PointerEvent );
-				}
-			}
-		}
-	}
+					if ( !bWasHandled )
+					{
+						TSet<FKey> PressedButtons;
+						FPointerEvent PointerEvent(
+							1 + (uint8)VREditorInteractor->GetControllerSide(),
+							FVector2D::ZeroVector,
+							FVector2D::ZeroVector,
+							PressedButtons,
+							bWasRightClicking ? EKeys::RightMouseButton : EKeys::LeftMouseButton,
+							0.0f,	// Wheel delta
+							FModifierKeysState() );
 
-	// Stop dragging the dock if we are dragging a dock
-	if (VRAction.ActionType == EVRActionType::SelectAndMove && Event == IE_Released)
-	{
-		// Put the Dock back on the hand it came from or leave it where it is in the room
-		if (DraggingUI != nullptr && DraggingUIHandIndex == VRAction.HandIndex)
-		{
-			FVirtualHand& OtherHand = Owner.GetOtherHand( VRAction.HandIndex );
-			const int32 OtherHandIndex = Owner.GetOtherHandIndex( DraggingUIHandIndex );
+						FWidgetPath EmptyWidgetPath;
 
-			bool bOnHand = false;
-			const float Distance = FVector::Dist( LastDraggingHoverLocation, OtherHand.Transform.GetLocation() ) / Owner.GetWorldScaleFactor();
-			if (Distance > VREd::MinDockDragDistance->GetFloat())
-			{
-				DraggingUI->SetDockedTo( AVREditorFloatingUI::EDockedTo::Room );
-			}
-			else
-			{
-				bOnHand = true;
-				const FVector EditorUIRelativeOffset( DraggingUI->GetSize().Y * 0.5f + VREd::EditorUIOffsetFromHand->GetFloat(), 0.0f, 0.0f );
-				const FTransform MovingUIGoalTransform = DraggingUI->MakeUITransformLockedToHand( OtherHandIndex, false, EditorUIRelativeOffset, FRotator( 90.0f, 180.0f, 0.0f ) );
-				const AVREditorFloatingUI::EDockedTo NewDockedTo = OtherHandIndex == VREditorConstants::LeftHandIndex ? AVREditorFloatingUI::EDockedTo::LeftHand : AVREditorFloatingUI::EDockedTo::RightHand;
-				DraggingUI->MoveTo( MovingUIGoalTransform, VREd::DockSnapTime->GetFloat(), NewDockedTo );
+						VREditorInteractor->SetIsClickingOnUI( false );
+						VREditorInteractor->SetIsRightClickingOnUI( false );
+						FReply Reply = FSlateApplication::Get().RoutePointerUpEvent( EmptyWidgetPath, PointerEvent );
+					}
+				}
 			}
-
-			ShowEditorUIPanel( DraggingUI, OtherHandIndex, true, bOnHand );
-			StopDraggingDockUI();
 		}
 	}
 }
 
 
-void FVREditorUISystem::OnVRHoverUpdate( FEditorViewportClient& ViewportClient, const int32 HandIndex, FVector& HoverImpactPoint, bool& bIsHoveringOverUI, bool& bWasHandled )
+void UVREditorUISystem::OnVRHoverUpdate( FEditorViewportClient& ViewportClient, UViewportInteractor* Interactor, FVector& HoverImpactPoint, bool& bWasHandled )
 {
-	FVirtualHand& Hand = Owner.GetVirtualHand( HandIndex );
+	UVREditorInteractor* VREditorInteractor = Cast<UVREditorInteractor>( Interactor );
 
-	if( !bWasHandled && Hand.DraggingMode != EVREditorDraggingMode::DockableWindow )
+	if( !bWasHandled && Interactor->GetDraggingMode() != EViewportInteractionDraggingMode::Interactable )
 	{
 		FVector LaserPointerStart, LaserPointerEnd;
-		if( Owner.GetLaserPointer( HandIndex, LaserPointerStart, LaserPointerEnd ) )
+		if( Interactor->GetLaserPointer( LaserPointerStart, LaserPointerEnd ) )
 		{
-			FHitResult HitResult = Owner.GetHitResultFromLaserPointer( HandIndex );
+			FHitResult HitResult = Interactor->GetHitResultFromLaserPointer();
 			if( HitResult.Actor.IsValid() )
 			{
 				// Only allow clicks to our own widget components
@@ -387,11 +346,11 @@ void FVREditorUISystem::OnVRHoverUpdate( FEditorViewportClient& ViewportClient,
 					WidgetComponent->GetLocalHitLocation( HitResult.ImpactPoint, LocalHitLocation );
 
 					// If we weren't already hovering over this widget, then we'll reset the last hit location
-					if( WidgetComponent != Hand.HoveringOverWidgetComponent )
+					if( WidgetComponent != VREditorInteractor->GetHoveringOverWidgetComponent() )
 					{
 						LastLocalHitLocation = LocalHitLocation;
 
-						if ( UVREditorWidgetComponent* VRWidgetComponent = Cast<UVREditorWidgetComponent>(Hand.HoveringOverWidgetComponent) )
+						if ( UVREditorWidgetComponent* VRWidgetComponent = Cast<UVREditorWidgetComponent>( VREditorInteractor->GetHoveringOverWidgetComponent() ) )
 						{
 							VRWidgetComponent->SetIsHovering(false);
 						}
@@ -401,12 +360,9 @@ void FVREditorUISystem::OnVRHoverUpdate( FEditorViewportClient& ViewportClient,
 					FWidgetPath WidgetPathUnderFinger = FWidgetPath( WidgetComponent->GetHitWidgetPath( HitResult.ImpactPoint, /*bIgnoreEnabledStatus*/ false ) );
 					if ( WidgetPathUnderFinger.IsValid() )
 					{
-						Hand.bIsHovering = true;
-						Hand.HoverLocation = HitResult.ImpactPoint;
-
 						TSet<FKey> PressedButtons;
 						FPointerEvent PointerEvent(
-							1 + HandIndex,
+							1 + (uint8)VREditorInteractor->GetControllerSide(),
 							LocalHitLocation,
 							LastLocalHitLocation,
 							LocalHitLocation - LastLocalHitLocation,
@@ -417,8 +373,8 @@ void FVREditorUISystem::OnVRHoverUpdate( FEditorViewportClient& ViewportClient,
 
 						bWasHandled = true;
 						HoverImpactPoint = HitResult.ImpactPoint;
-						Hand.HoveringOverWidgetComponent = WidgetComponent;
-						bIsHoveringOverUI = true;
+						VREditorInteractor->SetHoveringOverWidgetComponent( WidgetComponent );
+						VREditorInteractor->SetIsHoveringOverUI( true );
 
 						if ( UVREditorWidgetComponent* VRWidgetComponent = Cast<UVREditorWidgetComponent>(WidgetComponent) )
 						{
@@ -426,27 +382,27 @@ void FVREditorUISystem::OnVRHoverUpdate( FEditorViewportClient& ViewportClient,
 						}
 
 						// Route the mouse scrolling
-						if ( Hand.bIsTrackpadPositionValid[1] )
+						if ( VREditorInteractor->IsTrackpadPositionValid( 1 ) )
 						{
-							const bool bIsAbsolute = ( Owner.GetHMDDeviceType() == EHMDDeviceType::DT_SteamVR );
+							const bool bIsAbsolute = ( GetOwner().GetHMDDeviceType() == EHMDDeviceType::DT_SteamVR );
 
 							float ScrollDelta = 0.0f;
-							if( Hand.bIsTouchingTrackpad || !bIsAbsolute )
+							if( VREditorInteractor->IsTouchingTrackpad() || !bIsAbsolute )
 							{
 								if( bIsAbsolute )
 								{
 									const float ScrollSpeed = VREd::UIAbsoluteScrollSpeed->GetFloat();
-									ScrollDelta = ( Hand.TrackpadPosition.Y - Hand.LastTrackpadPosition.Y ) * ScrollSpeed;
+									ScrollDelta = ( VREditorInteractor->GetTrackpadPosition().Y - VREditorInteractor->GetLastTrackpadPosition().Y ) * ScrollSpeed;
 								}
 								else
 								{
 									const float ScrollSpeed = VREd::UIRelativeScrollSpeed->GetFloat();
-									ScrollDelta = Hand.TrackpadPosition.Y * ScrollSpeed;
+									ScrollDelta = VREditorInteractor->GetTrackpadPosition().Y * ScrollSpeed;
 								}
 							}
 
 							// If using a trackpad (Vive), invert scroll direction so that it feels more like scrolling on a mobile device
-							if( Owner.GetHMDDeviceType() == EHMDDeviceType::DT_SteamVR )
+							if( GetOwner().GetHMDDeviceType() == EHMDDeviceType::DT_SteamVR )
 							{
 								ScrollDelta *= -1.0f;
 							}
@@ -454,7 +410,7 @@ void FVREditorUISystem::OnVRHoverUpdate( FEditorViewportClient& ViewportClient,
 							if( !FMath::IsNearlyZero( ScrollDelta ) )
 							{
 								FPointerEvent MouseWheelEvent(
-									1 + HandIndex,
+									1 + (uint8)VREditorInteractor->GetControllerSide(),
 									LocalHitLocation,
 									LastLocalHitLocation,
 									PressedButtons,
@@ -464,30 +420,30 @@ void FVREditorUISystem::OnVRHoverUpdate( FEditorViewportClient& ViewportClient,
 
 								FSlateApplication::Get().RouteMouseWheelOrGestureEvent(WidgetPathUnderFinger, MouseWheelEvent, nullptr);
 
-								Hand.UIScrollVelocity = 0.0f;
+								VREditorInteractor->SetUIScrollVelocity( 0.0f );
 
 								// Don't apply inertia unless the user dragged a decent amount this frame
 								if( bIsAbsolute && FMath::Abs( ScrollDelta ) >= VREd::MinUIScrollDeltaForInertia->GetFloat() )
 								{
 									// Don't apply inertia if our data is sort of old
 									const FTimespan CurrentTime = FTimespan::FromSeconds( FPlatformTime::Seconds() );
-									if( CurrentTime - Hand.LastTrackpadPositionUpdateTime < FTimespan::FromSeconds( 1.0f / 30.0f ) )
+									if( CurrentTime - VREditorInteractor->GetLastTrackpadPositionUpdateTime() < FTimespan::FromSeconds( 1.0f / 30.0f ) )
 									{
 										//GWarn->Logf( TEXT( "INPUT: UIScrollVelocity=%0.2f" ), ScrollDelta );
-										Hand.UIScrollVelocity = ScrollDelta;
+										VREditorInteractor->SetUIScrollVelocity( ScrollDelta );
 									}
 								}
 							}
 						}
 						else
 						{
-							if( !FMath::IsNearlyZero( Hand.UIScrollVelocity ) )
+							if( !FMath::IsNearlyZero( VREditorInteractor->GetUIScrollVelocity() ) )
 							{
 								// Apply UI scrolling inertia
-								const float ScrollDelta = Hand.UIScrollVelocity;
+								const float ScrollDelta = VREditorInteractor->GetUIScrollVelocity();
 								{
 									FPointerEvent MouseWheelEvent(
-										1 + HandIndex,
+										1 + (uint8)VREditorInteractor->GetControllerSide(),
 										LocalHitLocation,
 										LastLocalHitLocation,
 										PressedButtons,
@@ -499,16 +455,14 @@ void FVREditorUISystem::OnVRHoverUpdate( FEditorViewportClient& ViewportClient,
 								}
 
 								// Apply damping
-								FVector ScrollVelocityVector( Hand.UIScrollVelocity, 0.0f, 0.0f );
+								FVector ScrollVelocityVector( VREditorInteractor->GetUIScrollVelocity(), 0.0f, 0.0f );
 								const bool bVelocitySensitive = false;
-								Owner.ApplyVelocityDamping( ScrollVelocityVector, bVelocitySensitive );
-								Hand.UIScrollVelocity = ScrollVelocityVector.X;
-
-								//GWarn->Logf( TEXT( "INERTIA: UIScrollVelocity==%0.2f  (DAMPING: UIScrollVelocity==%0.2f)" ), ScrollDelta, Hand.UIScrollVelocity );
+								GetOwner().GetWorldInteraction().ApplyVelocityDamping( ScrollVelocityVector, bVelocitySensitive );
+								VREditorInteractor->SetUIScrollVelocity( ScrollVelocityVector.X );
 							}
 							else
 							{
-								Hand.UIScrollVelocity = 0.0f;
+								VREditorInteractor->SetUIScrollVelocity( 0.0f );
 							}
 						}
 					}
@@ -518,30 +472,30 @@ void FVREditorUISystem::OnVRHoverUpdate( FEditorViewportClient& ViewportClient,
 	}
 
 	// If nothing was hovered, make sure we tell Slate about that
-	if( !bWasHandled && Hand.HoveringOverWidgetComponent != nullptr )
+	if( !bWasHandled && VREditorInteractor->GetHoveringOverWidgetComponent() != nullptr )
 	{
-		if ( UVREditorWidgetComponent* VRWidgetComponent = Cast<UVREditorWidgetComponent>(Hand.HoveringOverWidgetComponent) )
+		if ( UVREditorWidgetComponent* VRWidgetComponent = Cast<UVREditorWidgetComponent>( VREditorInteractor->GetHoveringOverWidgetComponent() ) )
 		{
-			VRWidgetComponent->SetIsHovering(false);
+			VRWidgetComponent->SetIsHovering( false );
 		}
 
-		const FVector2D LastLocalHitLocation = Hand.HoveringOverWidgetComponent->GetLastLocalHitLocation();
-		Hand.HoveringOverWidgetComponent = nullptr;
+		const FVector2D LastLocalHitLocation = VREditorInteractor->GetHoveringOverWidgetComponent()->GetLastLocalHitLocation();
+		VREditorInteractor->SetHoveringOverWidgetComponent( nullptr );
 
 		TSet<FKey> PressedButtons;
 		FPointerEvent PointerEvent(
-			1 + HandIndex,
+			1 + (uint8)VREditorInteractor->GetControllerSide(),
 			LastLocalHitLocation,
 			LastLocalHitLocation,
 			FVector2D::ZeroVector,
 			PressedButtons,
 			FModifierKeysState() );
+
 		FSlateApplication::Get().RoutePointerMoveEvent( FWidgetPath(), PointerEvent, false );
 	}
 }
 
-
-void FVREditorUISystem::Tick( FEditorViewportClient* ViewportClient, const float DeltaTime )
+void UVREditorUISystem::Tick( FEditorViewportClient* ViewportClient, const float DeltaTime )
 {
 	if ( bRefocusViewport )
 	{
@@ -553,108 +507,71 @@ void FVREditorUISystem::Tick( FEditorViewportClient* ViewportClient, const float
 	// aimed at when the user does this.
 	if( QuickMenuUI != nullptr )
 	{
-		int32 HandIndexWithQuickMenu = INDEX_NONE;
+		UVREditorInteractor* HandInteractorWithQuickMenu = nullptr;
 		if( QuickMenuUI->IsUIVisible() )
 		{
-			HandIndexWithQuickMenu = QuickMenuUI->GetDockedTo() == AVREditorFloatingUI::EDockedTo::LeftArm ? VREditorConstants::LeftHandIndex : VREditorConstants::RightHandIndex;
+			HandInteractorWithQuickMenu = GetOwner().GetHandInteractor( QuickMenuUI->GetDockedTo() == AVREditorFloatingUI::EDockedTo::LeftArm ? EControllerHand::Left : EControllerHand::Right );
 		}
 
-		const float WorldScaleFactor = Owner.GetWorldScaleFactor();
+		const float WorldScaleFactor = GetOwner().GetWorldScaleFactor();
 
-		int32 HandIndexThatNeedsQuickMenu = INDEX_NONE;
-		for( int32 HandIndex = 0; HandIndex < VREditorConstants::NumVirtualHands; ++HandIndex )
+		UViewportInteractor* HandInteractorThatNeedsQuickMenu = nullptr;
+		for( UViewportInteractor* Interactor : GetOwner().GetWorldInteraction().GetInteractors() )
 		{
 			bool bShowQuickMenuOnArm = false;
 
-			const FVirtualHand& Hand = Owner.GetVirtualHand( HandIndex );
-			const FVirtualHand& OtherHand = Owner.GetOtherHand( HandIndex );
-			const int32 OtherHandIndex = Owner.GetOtherHandIndex( HandIndex );		
+			UViewportInteractor* OtherInteractor = Interactor->GetOtherInteractor();
 
-			// @todo vreditor tweak: Weird to hard code this here.  Probably should be an accessor on the hand itself, and based on the actual device type
-			const FTransform UICapsuleTransform = OtherHand.Transform;
+			// @todo vreditor tweak: Weird to hard code this here. Probably should be an accessor on the hand itself, and based on the actual device type
+			const FTransform UICapsuleTransform = OtherInteractor->GetTransform();
 			const FVector UICapsuleStart = FVector( -9.0f, 0.0f, 0.0f ) * WorldScaleFactor;
 			const FVector UICapsuleEnd = FVector( -18.0f, 0.0f, 0.0f ) * WorldScaleFactor;
 			const float UICapsuleLocalRadius = 6.0f * WorldScaleFactor;
-			const float MinDistanceToUICapsule = 8.0f * WorldScaleFactor;	// @todo vreditor tweak
+			const float MinDistanceToUICapsule = 8.0f * WorldScaleFactor; // @todo vreditor tweak
 			const FVector UIForwardVector = FVector::UpVector;
 			const float MinDotForAimingAtOtherHand = 0.25f;	// @todo vreditor tweak
 
-			if( Owner.IsHandAimingTowardsCapsule( HandIndex, UICapsuleTransform, UICapsuleStart, UICapsuleEnd, UICapsuleLocalRadius, MinDistanceToUICapsule, UIForwardVector, MinDotForAimingAtOtherHand ) )
+			if( GetOwner().IsHandAimingTowardsCapsule( Interactor, UICapsuleTransform, UICapsuleStart, UICapsuleEnd, UICapsuleLocalRadius, MinDistanceToUICapsule, UIForwardVector, MinDotForAimingAtOtherHand ) )
 			{
 				bShowQuickMenuOnArm = true;
 			}
 
 			if( bShowQuickMenuOnArm )
 			{
-				HandIndexThatNeedsQuickMenu = OtherHandIndex;
-			}
-		}
-
-		if( QuickMenuUI->IsUIVisible() )
-		{
-			// If we don't need a quick menu, or if a different hand needs to spawn it, then kill the existing menu
-			if( HandIndexThatNeedsQuickMenu == INDEX_NONE || HandIndexThatNeedsQuickMenu != HandIndexWithQuickMenu  )
-			{
-				// Despawn
-				Owner.GetVirtualHand( HandIndexWithQuickMenu ).bHasUIOnForearm = false;
-				QuickMenuUI->ShowUI( false );
+				HandInteractorThatNeedsQuickMenu = OtherInteractor;
 			}
 		}
-
-		if( HandIndexThatNeedsQuickMenu != INDEX_NONE && !QuickMenuUI->IsUIVisible() )
+		
+		// If we don't need a quick menu, or if a different hand needs to spawn it, then kill the existing menu
+		if( QuickMenuUI->IsUIVisible() && HandInteractorWithQuickMenu != nullptr && 
+			( HandInteractorThatNeedsQuickMenu == nullptr || HandInteractorThatNeedsQuickMenu != HandInteractorWithQuickMenu )  )
 		{
-			const AVREditorFloatingUI::EDockedTo DockTo = ( HandIndexThatNeedsQuickMenu == VREditorConstants::LeftHandIndex ) ? AVREditorFloatingUI::EDockedTo::LeftArm : AVREditorFloatingUI::EDockedTo::RightArm;
-			QuickMenuUI->SetDockedTo( DockTo );
-			QuickMenuUI->ShowUI( true );
-			Owner.GetVirtualHand( HandIndexThatNeedsQuickMenu ).bHasUIOnForearm = true;
+			// Despawn
+			HandInteractorWithQuickMenu->SetHasUIOnForearm( false );
+			QuickMenuUI->ShowUI( false );
 		}
-	}
 
-	// If the user is moving the analog stick, try to spawn the radial menu for that hand
-	if( Owner.GetHMDDeviceType() == EHMDDeviceType::DT_OculusRift )
-	{
-		for( int32 HandIndex = 0; HandIndex < VREditorConstants::NumVirtualHands; ++HandIndex )
+		if( HandInteractorThatNeedsQuickMenu != nullptr && !QuickMenuUI->IsUIVisible() )
 		{
-			const FVirtualHand& Hand = Owner.GetVirtualHand( HandIndex );
-			const FVirtualHand& OtherHand = Owner.GetOtherHand( HandIndex );
-			bool bShouldShowRadialMenu = false;
-
-			const float MinJoystickOffsetBeforeRadialMenu = 0.15f;	// @todo vreditor twea
-
-			const bool bOtherHandAlreadyHasRadialMenu =
-				IsShowingRadialMenu( Owner.GetOtherHandIndex( HandIndex ) );
-
-			if( !Hand.bHasUIInFront &&
-				Hand.bIsTrackpadPositionValid[ 0 ] &&
-				Hand.bIsTrackpadPositionValid[ 1 ] &&
-				Hand.TrackpadPosition.Size() > MinJoystickOffsetBeforeRadialMenu &&
-				GetDraggingDockUIHandIndex() != HandIndex &&
-				!bOtherHandAlreadyHasRadialMenu )
-			{
-				bShouldShowRadialMenu = true;
-			}
-
-			if( bShouldShowRadialMenu )
+			UVREditorInteractor* VREditorHandInteractorThatNeedsQuickMenu = Cast<UVREditorInteractor>( HandInteractorThatNeedsQuickMenu );
+			if ( VREditorHandInteractorThatNeedsQuickMenu )
 			{
-				TryToSpawnRadialMenu( HandIndex );
-			}
-			else
-			{
-				// Close it
-				HideRadialMenu( HandIndex );
+				const AVREditorFloatingUI::EDockedTo DockTo = ( VREditorHandInteractorThatNeedsQuickMenu->GetControllerSide() == EControllerHand::Left ) ? AVREditorFloatingUI::EDockedTo::LeftArm : AVREditorFloatingUI::EDockedTo::RightArm;
+				QuickMenuUI->SetDockedTo( DockTo );
+				QuickMenuUI->ShowUI( true );
+				VREditorHandInteractorThatNeedsQuickMenu->SetHasUIOnForearm( true );
 			}
 		}
 	}
-	else
+
+	// Close the radial menu if it was not updated for a while
+	RadialMenuHideDelayTime += DeltaTime;
+	if( RadialMenuHideDelayTime >= VREd::RadialMenuFadeDelay->GetFloat() )
 	{
-		// Close the radial menu if it was not updated for a while
-		RadialMenuHideDelayTime += DeltaTime;
-		if( RadialMenuHideDelayTime >= VREd::RadialMenuFadeDelay->GetFloat() )
-		{
-			HideRadialMenu( QuickRadialMenu->GetDockedTo() == AVREditorFloatingUI::EDockedTo::LeftHand ? VREditorConstants::LeftHandIndex : VREditorConstants::RightHandIndex );
-		}
+		UVREditorInteractor* Interactor = GetOwner().GetHandInteractor( QuickRadialMenu->GetDockedTo() == AVREditorFloatingUI::EDockedTo::LeftHand ? EControllerHand::Left : EControllerHand::Right );
+		HideRadialMenu( Interactor );
 	}
-
+	
 	// Tick all of our floating UIs
 	for( AVREditorFloatingUI* FloatingUIPtr : FloatingUIs )
 	{
@@ -665,12 +582,12 @@ void FVREditorUISystem::Tick( FEditorViewportClient* ViewportClient, const float
 	}
 }
 
-void FVREditorUISystem::Render( const FSceneView* SceneView, FViewport* Viewport, FPrimitiveDrawInterface* PDI )
+void UVREditorUISystem::Render( const FSceneView* SceneView, FViewport* Viewport, FPrimitiveDrawInterface* PDI )
 {
 	// ...
 }
 
-void FVREditorUISystem::CreateUIs()
+void UVREditorUISystem::CreateUIs()
 {
 	const FIntPoint DefaultResolution( VREd::DefaultEditorUIResolutionX->GetInt(), VREd::DefaultEditorUIResolutionY->GetInt() );
 
@@ -857,28 +774,28 @@ void FVREditorUISystem::CreateUIs()
 			const FIntPoint Resolution(VREd::AssetEditorUIResolutionX->GetInt(), VREd::AssetEditorUIResolutionY->GetInt());
 
 			const bool bWithSceneComponent = false;
-			AVREditorDockableWindow* TabManagerUI = GetOwner().SpawnTransientSceneActor< AVREditorDockableWindow >(TEXT("AssetEditor"), bWithSceneComponent);
-			TabManagerUI->SetSlateWidget(*this, SNullWidget::NullWidget, Resolution, VREd::EditorUISize->GetFloat(), AVREditorFloatingUI::EDockedTo::Nothing);
+			AVREditorDockableWindow* TabManagerUI = GetOwner().SpawnTransientSceneActor< AVREditorDockableWindow >( TEXT( "AssetEditor" ), bWithSceneComponent );
+			TabManagerUI->SetSlateWidget( *this, SNullWidget::NullWidget, Resolution, VREd::EditorUISize->GetFloat(), AVREditorFloatingUI::EDockedTo::Nothing );
 			TabManagerUI->ShowUI( false );
 
 			// @todo vreditor: Could use "Hovering" instead for better performance with many UIs, but needs to be manually refreshed in some cases
-			TabManagerUI->GetWidgetComponent()->SetDrawingPolicy(EVREditorWidgetDrawingPolicy::Always);
+			TabManagerUI->GetWidgetComponent()->SetDrawingPolicy( EVREditorWidgetDrawingPolicy::Always );
 
-			FloatingUIs.Add(TabManagerUI);
+			FloatingUIs.Add( TabManagerUI );
 
 			EditorUIPanels[ (int32)EEditorUIPanel::AssetEditor ] = TabManagerUI;
 
 			TSharedPtr<SWindow> TabManagerWindow = TabManagerUI->GetWidgetComponent()->GetSlateWindow();
 			TSharedRef<SWindow> TabManagerWindowRef = TabManagerWindow.ToSharedRef();
-			ProxyTabManager = MakeShareable(new FProxyTabmanager(TabManagerWindowRef));
+			ProxyTabManager = MakeShareable( new FProxyTabmanager( TabManagerWindowRef ) );
 
-			ProxyTabManager->OnTabOpened.Add(FOnTabEvent::FDelegate::CreateRaw(this, &FVREditorUISystem::OnProxyTabLaunched));
-			ProxyTabManager->OnAttentionDrawnToTab.Add(FOnTabEvent::FDelegate::CreateRaw(this, &FVREditorUISystem::OnAttentionDrawnToTab));
+			ProxyTabManager->OnTabOpened.AddUObject( this, &UVREditorUISystem::OnProxyTabLaunched );
+			ProxyTabManager->OnAttentionDrawnToTab.AddUObject( this, &UVREditorUISystem::OnAttentionDrawnToTab );
 
 			// We're going to start stealing tabs from the global tab manager inserting them into the world instead.
-			FGlobalTabmanager::Get()->SetProxyTabManager(ProxyTabManager);
+			FGlobalTabmanager::Get()->SetProxyTabManager( ProxyTabManager );
 
-			FAssetEditorManager::Get().OnAssetEditorOpened().AddRaw(this, &FVREditorUISystem::OnAssetEditorOpened);
+			FAssetEditorManager::Get().OnAssetEditorOpened().AddUObject( this, &UVREditorUISystem::OnAssetEditorOpened );
 		}
 		// create the world settings menu
 		{
@@ -914,7 +831,7 @@ void FVREditorUISystem::CreateUIs()
 	}
 }
 
-void FVREditorUISystem::CleanUpActorsBeforeMapChangeOrSimulate()
+void UVREditorUISystem::CleanUpActorsBeforeMapChangeOrSimulate()
 {
 	for( AVREditorFloatingUI* FloatingUIPtr : FloatingUIs )
 	{
@@ -937,7 +854,7 @@ void FVREditorUISystem::CleanUpActorsBeforeMapChangeOrSimulate()
 	FAssetEditorManager::Get().OnAssetEditorOpened().RemoveAll(this);
 }
 
-void FVREditorUISystem::OnAssetEditorOpened(UObject* Asset)
+void UVREditorUISystem::OnAssetEditorOpened(UObject* Asset)
 {
 	// We need to disable drag drop on the tabs spawned in VR mode.
 	TArray<IAssetEditorInstance*> Editors = FAssetEditorManager::Get().FindEditorsForAsset(Asset);
@@ -947,7 +864,7 @@ void FVREditorUISystem::OnAssetEditorOpened(UObject* Asset)
 	}
 }
 
-bool FVREditorUISystem::IsWidgetAnEditorUIWidget( const UActorComponent* WidgetComponent ) const
+bool UVREditorUISystem::IsWidgetAnEditorUIWidget( const UActorComponent* WidgetComponent ) const
 {
 	if( WidgetComponent != nullptr && WidgetComponent->IsA( UWidgetComponent::StaticClass() ) )
 	{
@@ -967,7 +884,7 @@ bool FVREditorUISystem::IsWidgetAnEditorUIWidget( const UActorComponent* WidgetC
 }
 
 
-bool FVREditorUISystem::IsShowingEditorUIPanel( const EEditorUIPanel EditorUIPanel ) const
+bool UVREditorUISystem::IsShowingEditorUIPanel( const EEditorUIPanel EditorUIPanel ) const
 {
 	AVREditorFloatingUI* Panel = EditorUIPanels[ (int32)EditorUIPanel ];
 	if( Panel != nullptr )
@@ -979,7 +896,7 @@ bool FVREditorUISystem::IsShowingEditorUIPanel( const EEditorUIPanel EditorUIPan
 }
 
 
-void FVREditorUISystem::ShowEditorUIPanel( const UWidgetComponent* WidgetComponent, const int32 HandIndex, const bool bShouldShow, const bool OnHand, const bool bRefreshQuickMenu, const bool bPlaySound )
+void UVREditorUISystem::ShowEditorUIPanel( const UWidgetComponent* WidgetComponent, UVREditorInteractor* Interactor, const bool bShouldShow, const bool OnHand, const bool bRefreshQuickMenu, const bool bPlaySound )
 {
 	AVREditorFloatingUI* Panel = nullptr;
 	for( AVREditorFloatingUI* CurrentPanel : EditorUIPanels )
@@ -991,17 +908,17 @@ void FVREditorUISystem::ShowEditorUIPanel( const UWidgetComponent* WidgetCompone
 		}
 	}
 
-	ShowEditorUIPanel( Panel, HandIndex, bShouldShow, OnHand, bRefreshQuickMenu, bPlaySound );
+	ShowEditorUIPanel( Panel, Interactor, bShouldShow, OnHand, bRefreshQuickMenu, bPlaySound );
 }
 
 
-void FVREditorUISystem::ShowEditorUIPanel( const EEditorUIPanel EditorUIPanel, const int32 HandIndex, const bool bShouldShow, const bool OnHand, const bool bRefreshQuickMenu, const bool bPlaySound )
+void UVREditorUISystem::ShowEditorUIPanel( const EEditorUIPanel EditorUIPanel, UVREditorInteractor* Interactor, const bool bShouldShow, const bool OnHand, const bool bRefreshQuickMenu, const bool bPlaySound )
 {
 	AVREditorFloatingUI* Panel = EditorUIPanels[ (int32)EditorUIPanel ];
-	ShowEditorUIPanel( Panel, HandIndex, bShouldShow, OnHand, bRefreshQuickMenu, bPlaySound );
+	ShowEditorUIPanel( Panel, Interactor, bShouldShow, OnHand, bRefreshQuickMenu, bPlaySound );
 }
 
-void FVREditorUISystem::ShowEditorUIPanel( AVREditorFloatingUI* Panel, const int32 HandIndex, const bool bShouldShow, const bool OnHand, const bool bRefreshQuickMenu, const bool bPlaySound )
+void UVREditorUISystem::ShowEditorUIPanel( AVREditorFloatingUI* Panel, UVREditorInteractor* Interactor, const bool bShouldShow, const bool OnHand, const bool bRefreshQuickMenu, const bool bPlaySound )
 {
 	if( Panel != nullptr )
 	{
@@ -1014,42 +931,42 @@ void FVREditorUISystem::ShowEditorUIPanel( AVREditorFloatingUI* Panel, const int
 				AVREditorFloatingUI* OtherPanel = EditorUIPanels[PanelIndex];
 				if( OtherPanel != nullptr && OtherPanel->IsUIVisible() && OtherPanel->GetDockedTo() != AVREditorFloatingUI::EDockedTo::Room )
 				{
-					const uint32 OtherPanelDockedToHandIndex = OtherPanel->GetDockedTo() == AVREditorFloatingUI::EDockedTo::LeftHand ? VREditorConstants::LeftHandIndex : VREditorConstants::RightHandIndex;
-					if( OtherPanelDockedToHandIndex == HandIndex )
+					UVREditorInteractor* OtherPanelDockToVREditorInteractor = GetOwner().GetHandInteractor( OtherPanel->GetDockedTo() == AVREditorFloatingUI::EDockedTo::LeftHand ? EControllerHand::Left : EControllerHand::Right );
+					if( OtherPanelDockToVREditorInteractor == Interactor )
 					{
 						OtherPanel->ShowUI( false );
 						OtherPanel->SetDockedTo( AVREditorFloatingUI::EDockedTo::Nothing );
-						Owner.GetVirtualHand( OtherPanelDockedToHandIndex ).bHasUIInFront = false;
+						OtherPanelDockToVREditorInteractor->SetHasUIInFront( false );
 					}
 				}
 			}
 			
-			const AVREditorFloatingUI::EDockedTo NewDockedTo = HandIndex == VREditorConstants::LeftHandIndex ? AVREditorFloatingUI::EDockedTo::LeftHand : AVREditorFloatingUI::EDockedTo::RightHand;
+			const AVREditorFloatingUI::EDockedTo NewDockedTo = Interactor->GetControllerSide() == EControllerHand::Left ? AVREditorFloatingUI::EDockedTo::LeftHand : AVREditorFloatingUI::EDockedTo::RightHand;
 			Panel->SetDockedTo( NewDockedTo );
 			
-			if (bShouldShow)
+			if ( bShouldShow )
 			{
 				Panel->SetScale( Panel->GetInitialScale() );
 			}
 
 			const FVector EditorUIRelativeOffset( Panel->GetSize().Y * 0.5f + VREd::EditorUIOffsetFromHand->GetFloat(), 0.0f, 0.0f );
 			Panel->SetRelativeOffset( EditorUIRelativeOffset );
-			Panel->SetLocalRotation( FRotator( 90.0f, 180.0f, 0.0f ) ); // Todo needs initial rotation 
+			Panel->SetLocalRotation( FRotator( 90.0f, 180.0f, 0.0f ) ); // @todo vreditor: needs initial rotation
 		}
 
 		Panel->ShowUI( bShouldShow );
 
-		if (Panel->GetDockedTo() == AVREditorFloatingUI::EDockedTo::LeftHand || Panel->GetDockedTo() == AVREditorFloatingUI::EDockedTo::RightHand)
+		if ( Panel->GetDockedTo() == AVREditorFloatingUI::EDockedTo::LeftHand || Panel->GetDockedTo() == AVREditorFloatingUI::EDockedTo::RightHand )
 		{
-			Owner.GetVirtualHand( HandIndex ).bHasUIInFront = bShouldShow;
+			Interactor->SetHasUIInFront( bShouldShow );
 		}
 
 		if( bPlaySound )
 		{
-			UGameplayStatics::PlaySound2D( Owner.GetWorld(), bShouldShow ? ShowUISound : HideUISound );
+			UGameplayStatics::PlaySound2D( GetOwner().GetWorld(), bShouldShow ? ShowUISound : HideUISound );
 		}
 
-		if (bRefreshQuickMenu && QuickMenuUI)
+		if ( bRefreshQuickMenu && QuickMenuUI )
 		{
 			QuickMenuUI->GetUserWidget<UVREditorQuickMenu>()->RefreshUI();
 		}
@@ -1057,51 +974,52 @@ void FVREditorUISystem::ShowEditorUIPanel( AVREditorFloatingUI* Panel, const int
 }
 
 
-bool FVREditorUISystem::IsShowingRadialMenu( const int32 HandIndex ) const
+bool UVREditorUISystem::IsShowingRadialMenu( UVREditorInteractor* Interactor ) const
 {
-	int32 DockedToHandIndex = INDEX_NONE;
-	DockedToHandIndex = QuickRadialMenu->GetDockedTo() == AVREditorFloatingUI::EDockedTo::LeftHand ? VREditorConstants::LeftHandIndex : VREditorConstants::RightHandIndex;
-	return DockedToHandIndex == HandIndex && !QuickRadialMenu->bHidden;
+	const EControllerHand DockedToHand = QuickRadialMenu->GetDockedTo() == AVREditorFloatingUI::EDockedTo::LeftHand ? EControllerHand::Left : EControllerHand::Right;
+	const UVREditorInteractor* DockedToInteractor = GetOwner().GetHandInteractor( DockedToHand );
+	return DockedToInteractor == Interactor && !QuickRadialMenu->bHidden;
 }
 
 
-void FVREditorUISystem::UpdateRadialMenu( const int32 HandIndex )
+void UVREditorUISystem::UpdateRadialMenu( UVREditorInteractor* Interactor )
 {
-	if(QuickRadialMenu->bHidden)
+	if( QuickRadialMenu->bHidden )
 	{
 		QuickRadialMenu->ShowUI( true, false );
 	}
 
-	if(!QuickRadialMenu->bHidden)
+	if( !QuickRadialMenu->bHidden )
 	{
 		RadialMenuHideDelayTime = 0.0f;
-		QuickRadialMenu->GetUserWidget<UVREditorRadialMenu>()->Update( Owner.GetVirtualHand( HandIndex ) );
+		QuickRadialMenu->GetUserWidget<UVREditorRadialMenu>()->Update( Interactor );
 	}
 }
 
 
-void FVREditorUISystem::TryToSpawnRadialMenu( const int32 HandIndex )
+void UVREditorUISystem::TryToSpawnRadialMenu( UVREditorInteractor* Interactor )
 {
-	FVirtualHand& Hand = Owner.GetVirtualHand( HandIndex );
-
-	int32 DockedToHandIndex = INDEX_NONE;
+	UVREditorInteractor* DockedToInteractor = nullptr;
 	if( !QuickRadialMenu->bHidden )
 	{
-		DockedToHandIndex = QuickRadialMenu->GetDockedTo() == AVREditorFloatingUI::EDockedTo::LeftHand ? VREditorConstants::LeftHandIndex : VREditorConstants::RightHandIndex;
+		const EControllerHand ControllerHand = QuickRadialMenu->GetDockedTo() == AVREditorFloatingUI::EDockedTo::LeftHand ? EControllerHand::Left : EControllerHand::Right;
+		DockedToInteractor = GetOwner().GetHandInteractor( ControllerHand );
 	}
 
+	EViewportInteractionDraggingMode DraggingMode = Interactor->GetDraggingMode();
+
 	bool bNeedsSpawn =
-		( QuickRadialMenu->bHidden || DockedToHandIndex != HandIndex ) &&
-		Hand.DraggingMode != EVREditorDraggingMode::ActorsAtLaserImpact &&	// Don't show radial menu if the hand is busy dragging something around
-		Hand.DraggingMode != EVREditorDraggingMode::ActorsFreely &&
-		Hand.DraggingMode != EVREditorDraggingMode::ActorsWithGizmo &&
-		Hand.DraggingMode != EVREditorDraggingMode::AssistingDrag &&
-		DraggingUIHandIndex != HandIndex &&
-		!Hand.bIsHoveringOverUI;	// Don't show radial menu when aiming at a UI  (too much clutter)
+		( QuickRadialMenu->bHidden || DockedToInteractor != Interactor ) &&
+		DraggingMode != EViewportInteractionDraggingMode::ActorsAtLaserImpact &&	// Don't show radial menu if the hand is busy dragging something around
+		DraggingMode != EViewportInteractionDraggingMode::ActorsFreely &&
+		DraggingMode != EViewportInteractionDraggingMode::ActorsWithGizmo &&
+		DraggingMode != EViewportInteractionDraggingMode::AssistingDrag &&
+		( InteractorDraggingUI == nullptr || InteractorDraggingUI != Interactor ) &&
+		!Interactor->IsHoveringOverUI();	// Don't show radial menu when aiming at a UI  (too much clutter)
 
 	UVREditorRadialMenu* RadialMenu = QuickRadialMenu->GetUserWidget<UVREditorRadialMenu>();
 	// We need to update the trackpad position in the radialmenu before checking if it can be used
-	RadialMenu->Update( Hand );
+	RadialMenu->Update( Interactor );
 	if( RadialMenu && RadialMenu->IsInMenuRadius() )
 	{
 		bNeedsSpawn = false;
@@ -1109,65 +1027,60 @@ void FVREditorUISystem::TryToSpawnRadialMenu( const int32 HandIndex )
 
 	if( bNeedsSpawn )
 	{
-		DockedToHandIndex = HandIndex;
-
-		const AVREditorFloatingUI::EDockedTo DockedTo = DockedToHandIndex == VREditorConstants::LeftHandIndex ? AVREditorFloatingUI::EDockedTo::LeftHand : AVREditorFloatingUI::EDockedTo::RightHand;
+		const AVREditorFloatingUI::EDockedTo DockedTo = Interactor->GetControllerSide() == EControllerHand::Left ? AVREditorFloatingUI::EDockedTo::LeftHand : AVREditorFloatingUI::EDockedTo::RightHand;
 		QuickRadialMenu->SetDockedTo( DockedTo );
 		QuickRadialMenu->ShowUI( true );
 	}
 }
 
 
-void FVREditorUISystem::HideRadialMenu( const int32 HandIndex )
+void UVREditorUISystem::HideRadialMenu( UVREditorInteractor* Interactor )
 {
-	if( IsShowingRadialMenu( HandIndex ) )
+	// Only hide the radial menu if the passed interactor is actually the interactor with the radial menu
+	if( IsShowingRadialMenu( Interactor ) )
 	{
-		UVREditorRadialMenu* RadialMenu = QuickRadialMenu->GetUserWidget<UVREditorRadialMenu>();
 		QuickRadialMenu->ShowUI( false, true, VREd::RadialMenuFadeDelay->GetFloat() );
 	}
 }
 
 
-FTransform FVREditorUISystem::MakeDockableUITransformOnLaser( AVREditorDockableWindow* InitDraggingDockUI, const int32 HandIndex, const float DockSelectDistance ) const
+FTransform UVREditorUISystem::MakeDockableUITransformOnLaser( AVREditorDockableWindow* InitDraggingDockUI, UVREditorInteractor* Interactor, const float DockSelectDistance ) const
 {
-	const FVirtualHand& Hand = Owner.GetVirtualHand( HandIndex );
-
 	FTransform HandTransform;
 	FVector HandForward;
-	Owner.GetHandTransformAndForwardVector( HandIndex, HandTransform, HandForward );
+	Interactor->GetTransformAndForwardVector( HandTransform, HandForward );
+	const FTransform InteractorTransform = Interactor->GetTransform();
+
+	const FVector NewLocation = InteractorTransform.GetLocation() + ( HandForward.GetSafeNormal() * DockSelectDistance );
+	FRotator NewRotation = ( InteractorTransform.GetLocation() - NewLocation ).ToOrientationRotator();
+	NewRotation.Roll = -InteractorTransform.GetRotation().Rotator().Roll;
 
-	const FVector NewLocation = Hand.Transform.GetLocation() + ( HandForward.GetSafeNormal() * DockSelectDistance );
-	
-	FRotator NewRotation = ( Hand.Transform.GetLocation() - NewLocation ).ToOrientationRotator();
-	NewRotation.Roll = -Hand.Transform.GetRotation().Rotator().Roll;
-	
 	const FTransform LaserImpactToWorld( NewRotation, NewLocation );
 	return LaserImpactToWorld;
 }
 
 
-FTransform FVREditorUISystem::MakeDockableUITransform( AVREditorDockableWindow* InitDraggingDockUI, const int32 HandIndex, const float DockSelectDistance )
+FTransform UVREditorUISystem::MakeDockableUITransform( AVREditorDockableWindow* InitDraggingDockUI, UVREditorInteractor* Interactor, const float DockSelectDistance )
 {
-	const FTransform UIOnLaserToWorld = MakeDockableUITransformOnLaser( DraggingUI, HandIndex, DockSelectDistance );
+	const FTransform UIOnLaserToWorld = MakeDockableUITransformOnLaser( DraggingUI, Interactor, DockSelectDistance );
 	const FTransform UIToUIOnLaser = DraggingUIOffsetTransform;
 	
 	const FTransform UIToWorld = UIToUIOnLaser * UIOnLaserToWorld;
 
-	if( DraggingUI != nullptr &&  DraggingUI->GetDockedTo() == AVREditorFloatingUI::EDockedTo::Dragging )
+	if( DraggingUI != nullptr &&  DraggingUI->GetDockedTo() == AVREditorFloatingUI::EDockedTo::Dragging && InteractorDraggingUI )
 	{
+		UVREditorInteractor* OtherInteractor = Cast<UVREditorInteractor>( InteractorDraggingUI->GetOtherInteractor() );
 		
-		FVirtualHand& Hand = Owner.GetVirtualHand( DraggingUIHandIndex );
-		FVirtualHand& OtherHand = Owner.GetOtherHand( DraggingUIHandIndex );
-		
-		const float WorldScaleFactor = Owner.GetWorldScaleFactor();
-		const FVector DraggingUILocation = Hand.HoverLocation;
-		const FVector OtherHandLocation =  OtherHand.Transform.GetLocation();
+		const float WorldScaleFactor = GetOwner().GetWorldScaleFactor();
+		const FVector DraggingUILocation = InteractorDraggingUI->GetHoverLocation();
+		const FVector OtherHandLocation =  OtherInteractor->GetTransform().GetLocation();
 		LastDraggingHoverLocation = UIOnLaserToWorld.GetLocation();
 
 		const float Distance = FVector::Dist( UIOnLaserToWorld.GetLocation(), OtherHandLocation ) / WorldScaleFactor;
 
 		// if dragged passed the threshold since starting dragging
-		if ( Distance > VREd::MinDockDragDistance->GetFloat() && !bDraggedDockFromHandPassedThreshold && ( DraggingUI->GetPreviouslyDockedTo() == AVREditorFloatingUI::EDockedTo::LeftHand || DraggingUI->GetPreviouslyDockedTo() == AVREditorFloatingUI::EDockedTo::RightHand ) )
+		if ( Distance > VREd::MinDockDragDistance->GetFloat() && !bDraggedDockFromHandPassedThreshold && 
+			( DraggingUI->GetPreviouslyDockedTo() == AVREditorFloatingUI::EDockedTo::LeftHand || DraggingUI->GetPreviouslyDockedTo() == AVREditorFloatingUI::EDockedTo::RightHand ) )
 		{
 			bDraggedDockFromHandPassedThreshold = true;
 		}
@@ -1175,10 +1088,10 @@ FTransform FVREditorUISystem::MakeDockableUITransform( AVREditorDockableWindow*
 		// Snapping to a hand when in range
 		if ( Distance <= VREd::MinDockDragDistance->GetFloat() && ( bDraggedDockFromHandPassedThreshold || DraggingUI->GetPreviouslyDockedTo() == AVREditorFloatingUI::EDockedTo::Room ) )
 		{
-			const int32 OtherHandIndex = Owner.GetOtherHandIndex( HandIndex );
+			UVREditorInteractor* OtherInteractor = Cast<UVREditorInteractor>( Interactor->GetOtherInteractor() );
 			const FVector EditorUIRelativeOffset( DraggingUI->GetSize().Y * 0.5f + VREd::EditorUIOffsetFromHand->GetFloat(), 0.0f, 0.0f );
-			const FTransform MovingUIGoalTransform = DraggingUI->MakeUITransformLockedToHand( OtherHandIndex, false, EditorUIRelativeOffset, FRotator( 90.0f, 180.0f, 0.0f ) );
-			const AVREditorFloatingUI::EDockedTo NewDockedTo = OtherHandIndex == VREditorConstants::LeftHandIndex ? AVREditorFloatingUI::EDockedTo::LeftHand : AVREditorFloatingUI::EDockedTo::RightHand;
+			const FTransform MovingUIGoalTransform = DraggingUI->MakeUITransformLockedToHand( OtherInteractor, false, EditorUIRelativeOffset, FRotator( 90.0f, 180.0f, 0.0f ) );
+			const AVREditorFloatingUI::EDockedTo NewDockedTo = OtherInteractor->GetControllerSide() == EControllerHand::Left ? AVREditorFloatingUI::EDockedTo::LeftHand : AVREditorFloatingUI::EDockedTo::RightHand;
 			DraggingUI->MoveTo( MovingUIGoalTransform, VREd::DockDragSpeed->GetFloat(), NewDockedTo );
 		}
 		else
@@ -1191,22 +1104,32 @@ FTransform FVREditorUISystem::MakeDockableUITransform( AVREditorDockableWindow*
 }
 
 
-AVREditorFloatingUI* FVREditorUISystem::StartDraggingDockUI( AVREditorDockableWindow* InitDraggingDockUI, const int32 HandIndex, const float DockSelectDistance )
+void UVREditorUISystem::StartDraggingDockUI( AVREditorDockableWindow* InitDraggingDockUI, UVREditorInteractor* Interactor, const float DockSelectDistance )
 {
 	AVREditorFloatingUI::EDockedTo DockTo = InitDraggingDockUI->GetDockedTo();
 	if( DockTo == AVREditorFloatingUI::EDockedTo::LeftHand || DockTo == AVREditorFloatingUI::EDockedTo::RightHand )
 	{
 		bDraggedDockFromHandPassedThreshold = false;
-		Owner.GetOtherHand( HandIndex ).bHasUIInFront = false;
 	}
 
-	DraggingUIHandIndex = HandIndex;
+	// Set the laser of the interactor the dockable window was pulled from
+	UViewportInteractor* OtherInteractor = Interactor->GetOtherInteractor();
+	if ( OtherInteractor )
+	{
+		UVREditorInteractor* OtherVREditorInteractor = Cast<UVREditorInteractor>( OtherInteractor );
+		if ( OtherVREditorInteractor )
+		{
+			OtherVREditorInteractor->SetHasUIInFront( false );
+		}
+	}
+
+	InteractorDraggingUI = Interactor;
 
 	FTransform UIToWorld = InitDraggingDockUI->GetActorTransform();
 	UIToWorld.SetScale3D( FVector( 1.0f ) );
 	const FTransform WorldToUI = UIToWorld.Inverse();
 
-	const FTransform UIOnLaserToWorld = MakeDockableUITransformOnLaser( InitDraggingDockUI, HandIndex, DockSelectDistance );
+	const FTransform UIOnLaserToWorld = MakeDockableUITransformOnLaser( InitDraggingDockUI, Interactor, DockSelectDistance );
 	const FTransform UIOnLaserToUI = UIOnLaserToWorld * WorldToUI;
 	const FTransform UIToUIOnLaser = UIOnLaserToUI.Inverse();
 	DraggingUIOffsetTransform = UIToUIOnLaser;
@@ -1214,38 +1137,63 @@ AVREditorFloatingUI* FVREditorUISystem::StartDraggingDockUI( AVREditorDockableWi
 	DraggingUI = InitDraggingDockUI;
 	DraggingUI->SetDockedTo( AVREditorFloatingUI::EDockedTo::Dragging );
 
-	UGameplayStatics::PlaySound2D( Owner.GetWorld(), StartDragUISound );
-
-	return DraggingUI;
+	UGameplayStatics::PlaySound2D( GetOwner().GetWorld(), StartDragUISound );
 }
 
-AVREditorDockableWindow* FVREditorUISystem::GetDraggingDockUI() const
+AVREditorDockableWindow* UVREditorUISystem::GetDraggingDockUI() const
 {
 	return DraggingUI;
 }
 
-int32 FVREditorUISystem::GetDraggingDockUIHandIndex() const
+void UVREditorUISystem::StopDraggingDockUI( UVREditorInteractor* VREditorInteractor )
 {
-	return DraggingUIHandIndex;
-}
+	if ( IsInteractorDraggingDockUI( VREditorInteractor ) )
+	{
+		// Put the Dock back on the hand it came from or leave it where it is in the room
+		UViewportInteractor* OtherInteractor = VREditorInteractor->GetOtherInteractor();
+		if ( OtherInteractor )
+		{
+			UVREditorInteractor* OtherVREditorInteractor = Cast<UVREditorInteractor>( OtherInteractor );
+			if ( OtherVREditorInteractor )
+			{
+				bool bOnHand = false;
+				const float Distance = FVector::Dist( LastDraggingHoverLocation, OtherVREditorInteractor->GetTransform().GetLocation() ) / GetOwner().GetWorldScaleFactor();
+				if ( Distance > VREd::MinDockDragDistance->GetFloat() )
+				{
+					DraggingUI->SetDockedTo( AVREditorFloatingUI::EDockedTo::Room );
+				}
+				else
+				{
+					bOnHand = true;
+					const FVector EditorUIRelativeOffset( DraggingUI->GetSize().Y * 0.5f + VREd::EditorUIOffsetFromHand->GetFloat(), 0.0f, 0.0f );
+					const FTransform MovingUIGoalTransform = DraggingUI->MakeUITransformLockedToHand( OtherVREditorInteractor, false, EditorUIRelativeOffset, FRotator( 90.0f, 180.0f, 0.0f ) );
+					const AVREditorFloatingUI::EDockedTo NewDockedTo = OtherVREditorInteractor->GetControllerSide() == EControllerHand::Left ? AVREditorFloatingUI::EDockedTo::LeftHand : AVREditorFloatingUI::EDockedTo::RightHand;
+					DraggingUI->MoveTo( MovingUIGoalTransform, VREd::DockSnapTime->GetFloat(), NewDockedTo );
+				}
 
-void FVREditorUISystem::StopDraggingDockUI()
-{
-	DraggingUI = nullptr;
-	FVirtualHand& Hand = Owner.GetVirtualHand( DraggingUIHandIndex );
-	Hand.DraggingMode = EVREditorDraggingMode::Nothing;	
-	DraggingUIHandIndex = INDEX_NONE;
+				ShowEditorUIPanel( DraggingUI, OtherVREditorInteractor, true, bOnHand );
 
-	UGameplayStatics::PlaySound2D( Owner.GetWorld(), StopDragUISound );
+				DraggingUI = nullptr;
+				InteractorDraggingUI->SetDraggingMode( EViewportInteractionDraggingMode::Nothing );
+				InteractorDraggingUI = nullptr;
 
+				UGameplayStatics::PlaySound2D( GetOwner().GetWorld(), StopDragUISound );
+			}
+		}
+	}
 }
 
-bool FVREditorUISystem::IsDraggingDockUI()
+bool UVREditorUISystem::IsDraggingDockUI()
 {
 	return DraggingUI != nullptr && DraggingUI->GetDockedTo() == AVREditorFloatingUI::EDockedTo::Dragging;
 }
 
-void FVREditorUISystem::TogglePanelsVisibility()
+bool UVREditorUISystem::IsInteractorDraggingDockUI( const UVREditorInteractor* Interactor ) const
+{
+	return InteractorDraggingUI && InteractorDraggingUI == Interactor;
+}
+
+void UVREditorUISystem::TogglePanelsVisibility()
 {
 	bool bAnyPanelsVisible = false;
 	AVREditorFloatingUI* PanelOnHand = nullptr;
@@ -1269,7 +1217,7 @@ void FVREditorUISystem::TogglePanelsVisibility()
 	// Hide if there is any UI visible
 	const bool bShowUI = !bAnyPanelsVisible;
 	
-	UGameplayStatics::PlaySound2D( Owner.GetWorld(), bShowUI ? ShowUISound : HideUISound );
+	UGameplayStatics::PlaySound2D( GetOwner().GetWorld(), bShowUI ? ShowUISound : HideUISound );
 
 	if ( bSetDefaultLayout)
 	{
@@ -1278,20 +1226,20 @@ void FVREditorUISystem::TogglePanelsVisibility()
 	}
 	else
 	{
-		for (AVREditorFloatingUI* Panel : EditorUIPanels)
+		for ( AVREditorFloatingUI* Panel : EditorUIPanels )
 		{
-			if (Panel != nullptr && Panel->IsUIVisible() != bShowUI)
+			if ( Panel != nullptr && Panel->IsUIVisible() != bShowUI )
 			{
-				if (bShowUI)
+				if ( bShowUI )
 				{
 					bool SetNewVisibility = true;
 					const AVREditorFloatingUI::EDockedTo DockedTo = Panel->GetDockedTo();
-					if (DockedTo == AVREditorFloatingUI::EDockedTo::LeftHand || DockedTo == AVREditorFloatingUI::EDockedTo::RightHand)
+					if ( DockedTo == AVREditorFloatingUI::EDockedTo::LeftHand || DockedTo == AVREditorFloatingUI::EDockedTo::RightHand )
 					{
-						if (!PanelOnHand)
+						if ( !PanelOnHand )
 						{
-							const uint32 HandIndex = Panel->GetDockedTo() == AVREditorFloatingUI::EDockedTo::LeftHand ? VREditorConstants::LeftHandIndex : VREditorConstants::RightHandIndex;
-							Owner.GetVirtualHand( HandIndex ).bHasUIInFront = true;
+							UVREditorInteractor* VREditorInteractor = GetOwner().GetHandInteractor( Panel->GetDockedTo() == AVREditorFloatingUI::EDockedTo::LeftHand ? EControllerHand::Left : EControllerHand::Right );
+							VREditorInteractor->SetHasUIInFront( true );
 							PanelOnHand = Panel;
 						}
 						else
@@ -1299,12 +1247,12 @@ void FVREditorUISystem::TogglePanelsVisibility()
 							SetNewVisibility = false;
 						}
 					}
-					else if (DockedTo == AVREditorFloatingUI::EDockedTo::Nothing)
+					else if ( DockedTo == AVREditorFloatingUI::EDockedTo::Nothing )
 					{
 						SetNewVisibility = false;
 					}
 
-					if (SetNewVisibility)
+					if ( SetNewVisibility )
 					{
 						Panel->ShowUI( true );
 					}
@@ -1319,8 +1267,8 @@ void FVREditorUISystem::TogglePanelsVisibility()
 
 	if( !bShowUI && PanelOnHand )
 	{
-		const uint32 HandIndex = PanelOnHand->GetDockedTo() == AVREditorFloatingUI::EDockedTo::LeftHand ? VREditorConstants::LeftHandIndex : VREditorConstants::RightHandIndex;
-		Owner.GetVirtualHand( HandIndex ).bHasUIInFront = false;
+		UVREditorInteractor* VREditorInteractor = GetOwner().GetHandInteractor( PanelOnHand->GetDockedTo() == AVREditorFloatingUI::EDockedTo::LeftHand ? EControllerHand::Left : EControllerHand::Right );
+		VREditorInteractor->SetHasUIInFront( false );
 	}
 
 	if (QuickMenuUI)
@@ -1329,22 +1277,22 @@ void FVREditorUISystem::TogglePanelsVisibility()
 	}
 }
 
-float FVREditorUISystem::GetMaxDockWindowSize() const
+float UVREditorUISystem::GetMaxDockWindowSize() const
 {
 	return VREd::MaxDockWindowSize->GetFloat();
 }
 
-float FVREditorUISystem::GetMinDockWindowSize() const
+float UVREditorUISystem::GetMinDockWindowSize() const
 {
 	return VREd::MinDockWindowSize->GetFloat();
 }
 
-void FVREditorUISystem::OnProxyTabLaunched(TSharedPtr<SDockTab> NewTab)
+void UVREditorUISystem::OnProxyTabLaunched(TSharedPtr<SDockTab> NewTab)
 {
 	ShowAssetEditor();
 }
 
-void FVREditorUISystem::OnAttentionDrawnToTab(TSharedPtr<SDockTab> NewTab)
+void UVREditorUISystem::OnAttentionDrawnToTab(TSharedPtr<SDockTab> NewTab)
 {
 	// @todo vreditor: When clicking on Modes icons in the Modes panel, the LevelEditor tab is invoked which causes an empty asset editor to pop-up
 	static FName LevelEditorTabID( "LevelEditor" );
@@ -1354,7 +1302,7 @@ void FVREditorUISystem::OnAttentionDrawnToTab(TSharedPtr<SDockTab> NewTab)
 	}
 }
 
-void FVREditorUISystem::ShowAssetEditor()
+void UVREditorUISystem::ShowAssetEditor()
 {
 	bRefocusViewport = true;
 
@@ -1363,26 +1311,28 @@ void FVREditorUISystem::ShowAssetEditor()
 	if ( !IsShowingEditorUIPanel(EEditorUIPanel::AssetEditor) )
 	{
 		// Always spawn on a hand.  But which hand?  Well, we'll choose the hand that isn't actively clicking on something using a laser.
-		const int32 HandIndex = Owner.GetVirtualHand(VREditorConstants::LeftHandIndex).bIsClickingOnUI ? VREditorConstants::RightHandIndex : VREditorConstants::LeftHandIndex;	// Hand that did not clicked with a laser
+		UVREditorInteractor* VREditorInteractor = GetOwner().GetHandInteractor( EControllerHand::Left ); // Hand that did not clicked with a laser
+		if ( VREditorInteractor->IsClickingOnUI() )
+		{
+			VREditorInteractor = GetOwner().GetHandInteractor( EControllerHand::Right );
+		}
+		
 		const bool bShouldShow = true;
 		const bool bShowOnHand = true;
-		ShowEditorUIPanel(EEditorUIPanel::AssetEditor, HandIndex, bShouldShow, bShowOnHand);
+		ShowEditorUIPanel( EEditorUIPanel::AssetEditor, VREditorInteractor, bShouldShow, bShowOnHand );
 
 		// Play haptic effect so user knows to look at their hand that now has UI on it!
-		const float Strength = VREd::UIAssetEditorSummonedOnHandHapticFeedbackStrength->GetFloat();
-		Owner.PlayHapticEffect(
-			HandIndex == VREditorConstants::LeftHandIndex ? Strength : 0.0f,
-			HandIndex == VREditorConstants::RightHandIndex ? Strength : 0.0f);
+		VREditorInteractor->PlayHapticEffect( VREd::UIAssetEditorSummonedOnHandHapticFeedbackStrength->GetFloat() );
 	}
 }
 
-void FVREditorUISystem::TogglePanelVisibility( const EEditorUIPanel EditorUIPanel )
+void UVREditorUISystem::TogglePanelVisibility( const EEditorUIPanel EditorUIPanel )
 {
 	AVREditorFloatingUI* Panel = EditorUIPanels[(int32)EditorUIPanel];
 	if (Panel != nullptr)
 	{
 		const bool bIsShowing = Panel->IsUIVisible();
-		const int32 HandIndexWithQuickMenu = QuickMenuUI->GetDockedTo() == AVREditorFloatingUI::EDockedTo::LeftArm ? VREditorConstants::LeftHandIndex : VREditorConstants::RightHandIndex;
+		UVREditorInteractor* InteractorWithQuickMenu = GetOwner().GetHandInteractor( QuickMenuUI->GetDockedTo() == AVREditorFloatingUI::EDockedTo::LeftArm ? EControllerHand::Left : EControllerHand::Right );
 		if( Panel->GetDockedTo() == AVREditorFloatingUI::EDockedTo::Room )
 		{
 			Panel->ShowUI( false );
@@ -1390,31 +1340,32 @@ void FVREditorUISystem::TogglePanelVisibility( const EEditorUIPanel EditorUIPane
 		}
 		else
 		{
-			ShowEditorUIPanel( Panel, HandIndexWithQuickMenu, !bIsShowing, true );
+			ShowEditorUIPanel( Panel, InteractorWithQuickMenu, !bIsShowing, true );
 		}
 	}
 }
 
-void FVREditorUISystem::SetDefaultWindowLayout()
+void UVREditorUISystem::SetDefaultWindowLayout()
 {
-	for (int32 PanelIndex = 0; PanelIndex < DefaultWindowTransforms.Num(); ++PanelIndex)
+	for ( int32 PanelIndex = 0; PanelIndex < DefaultWindowTransforms.Num(); ++PanelIndex )
 	{
 		AVREditorFloatingUI* Panel = EditorUIPanels[PanelIndex];
-		if (Panel)
+		if ( Panel )
 		{
 			const AVREditorFloatingUI::EDockedTo DockedTo = Panel->GetDockedTo();
 			if (DockedTo == AVREditorFloatingUI::EDockedTo::LeftHand || DockedTo == AVREditorFloatingUI::EDockedTo::RightHand)
 			{
-				const int32 HandIndex = DockedTo == AVREditorFloatingUI::EDockedTo::LeftHand ? VREditorConstants::LeftHandIndex : VREditorConstants::RightHandIndex;
-				Owner.GetVirtualHand( HandIndex ).bHasUIInFront = false;
+				const EControllerHand ControllerHand = DockedTo == AVREditorFloatingUI::EDockedTo::LeftHand ? EControllerHand::Left : EControllerHand::Right;
+				UVREditorInteractor* Interactor = GetOwner().GetHandInteractor( ControllerHand );
+				Interactor->SetHasUIInFront( false );
 			}
 
 			Panel->SetDockedTo( AVREditorFloatingUI::EDockedTo::Room );
 			Panel->ShowUI( true );
 
 			// Make sure the UIs are centered around the direction your head is looking (yaw only!)
-			const FVector RoomSpaceHeadLocation = Owner.GetRoomSpaceHeadTransform().GetLocation() / Owner.GetWorldScaleFactor();
-			FRotator RoomSpaceHeadYawRotator = Owner.GetRoomSpaceHeadTransform().GetRotation().Rotator();
+			const FVector RoomSpaceHeadLocation = GetOwner().GetRoomSpaceHeadTransform().GetLocation() / GetOwner().GetWorldScaleFactor();
+			FRotator RoomSpaceHeadYawRotator = GetOwner().GetRoomSpaceHeadTransform().GetRotation().Rotator();
 			RoomSpaceHeadYawRotator.Pitch = 0.0f;
 			RoomSpaceHeadYawRotator.Roll = 0.0f;
 
