VR Editor:

diff --git a/Engine/Source/Editor/VREditor/VREditorUISystem.h b/Engine/Source/Editor/VREditor/VREditorUISystem.h
--- a/Engine/Source/Editor/VREditor/VREditorUISystem.h
+++ b/Engine/Source/Editor/VREditor/VREditorUISystem.h
@@ -2,11 +2,19 @@
 
 #pragma once
 
+#include "VREditorUISystem.generated.h"
+
+class UVREditorInteractor;
+class UViewportInteractor;
+
 /**
  * VR Editor user interface manager
  */
-class FVREditorUISystem : public FGCObject
+UCLASS()
+class UVREditorUISystem : public UObject
 {
+	GENERATED_UCLASS_BODY()
+
 public:
 
 	enum class EEditorUIPanel
@@ -23,14 +31,11 @@ public:
 		TotalCount,
 	};
 
-	/** Default constructor for FVREditorUISystem, which sets up safe defaults */
-	FVREditorUISystem( class FVREditorMode& InitOwner );
+	/** Shuts down the UISystem whenever the mode is exited */
+	void Shutdown();
 
-	/** Clean up before destruction */
-	virtual ~FVREditorUISystem();
-
-	// FGCObject overrides
-	virtual void AddReferencedObjects( FReferenceCollector& Collector ) override;
+	/** Initializes default values for the UISystem and creates the UI panels */
+	void Init();
 
 	/** Called by VREditorMode to update us every frame */
 	void Tick( class FEditorViewportClient* ViewportClient, const float DeltaTime );
@@ -38,16 +43,22 @@ public:
 	/** Called by VREditorMode to draw debug graphics every frame */
 	void Render( const class FSceneView* SceneView, class FViewport* Viewport, class FPrimitiveDrawInterface* PDI );
 
+	/** Sets the owner of this system */
+	void SetOwner( class FVREditorMode* NewOwner )
+	{
+		VRMode = NewOwner;
+	}
+
 	/** Gets the owner of this system */
 	class FVREditorMode& GetOwner()
 	{
-		return Owner;
+		return *VRMode;
 	}
 
 	/** Gets the owner of this system (const) */
 	const class FVREditorMode& GetOwner() const
 	{
-		return Owner;
+		return *VRMode;
 	}
 
 	/** Called by our owner right before a map is loaded or switching between Simulate and normal Editor, so we can clean
@@ -61,43 +72,42 @@ public:
 	bool IsShowingEditorUIPanel( const EEditorUIPanel EditorUIPanel ) const;
 
 	/** Sets whether the specified editor UI panel should be visible.  Any other UI floating off this hand will be dismissed when showing it. */
-	void ShowEditorUIPanel( const class UWidgetComponent* WidgetComponent, const int32 HandIndex, const bool bShouldShow, const bool OnHand = false, const bool bRefreshQuickMenu = true, const bool bPlaySound = true );
-	void ShowEditorUIPanel( const EEditorUIPanel EditorUIPanel, const int32 HandIndex, const bool bShouldShow, const bool OnHand = false, const bool bRefreshQuickMenu = true, const bool bPlaySound = true );
-	void ShowEditorUIPanel( class AVREditorFloatingUI* Panel, const int32 HandIndex, const bool bShouldShow, const bool OnHand = false, const bool bRefreshQuickMenu = true, const bool bPlaySound = true );
+	void ShowEditorUIPanel( const class UWidgetComponent* WidgetComponent, UVREditorInteractor* Interactor, const bool bShouldShow, const bool OnHand = false, const bool bRefreshQuickMenu = true, const bool bPlaySound = true );
+	void ShowEditorUIPanel( const EEditorUIPanel EditorUIPanel, UVREditorInteractor* Interactor, const bool bShouldShow, const bool OnHand = false, const bool bRefreshQuickMenu = true, const bool bPlaySound = true );
+	void ShowEditorUIPanel( class AVREditorFloatingUI* Panel, UVREditorInteractor* Interactor, const bool bShouldShow, const bool OnHand = false, const bool bRefreshQuickMenu = true, const bool bPlaySound = true );
 
 	/** Returns true if the radial menu is visible on this hand */
-	bool IsShowingRadialMenu( const int32 HandIndex ) const;
+	bool IsShowingRadialMenu( UVREditorInteractor* Interactor ) const;
 
 	/** Tries to spawn the radial menu (if the specified hand isn't doing anything else) */
-	void TryToSpawnRadialMenu( const int32 HandIndex );
+	void TryToSpawnRadialMenu( UVREditorInteractor* Interactor );
 
 	/** Updates the radial menu with the current hand trackpad location */
-	void UpdateRadialMenu( const int32 HandIndex );
+	void UpdateRadialMenu( UVREditorInteractor* Interactor );
 
 	/** Hides the radial menu if the specified hand is showing it */
-	void HideRadialMenu( const int32 HandIndex );
+	void HideRadialMenu( UVREditorInteractor* Interactor );
 
 	/** Start dragging a dock window on the hand */
-	class AVREditorFloatingUI* StartDraggingDockUI( class AVREditorDockableWindow* InitDraggingDockUI, const int32 HandIndex, const float DockSelectDistance );
+	void StartDraggingDockUI( class AVREditorDockableWindow* InitDraggingDockUI, UVREditorInteractor* Interactor, const float DockSelectDistance );
 
 	/** Makes up a transform for a dockable UI when dragging it with a laser at the specified distance from the laser origin */
-	FTransform MakeDockableUITransformOnLaser( AVREditorDockableWindow* InitDraggingDockUI, const int32 HandIndex, const float DockSelectDistance ) const;
+	FTransform MakeDockableUITransformOnLaser( AVREditorDockableWindow* InitDraggingDockUI, UVREditorInteractor* Interactor, const float DockSelectDistance ) const;
 
 	/** Makes up a transform for a dockable UI when dragging it that includes the original offset from the laser's impact point */
-	FTransform MakeDockableUITransform( AVREditorDockableWindow* InitDraggingDockUI, const int32 HandIndex, const float DockSelectDistance );
+	FTransform MakeDockableUITransform( AVREditorDockableWindow* InitDraggingDockUI, UVREditorInteractor* Interactor, const float DockSelectDistance );
 	
 	/** Returns the current dragged dock window, nullptr if none */
 	class AVREditorDockableWindow* GetDraggingDockUI() const;
 
-	/** Gets the hand index of the hand that is dragging a dock UI */
-	int32 GetDraggingDockUIHandIndex() const;
-
 	/** Resets all values to stop dragging the current dock window */
-	void StopDraggingDockUI();
+	void StopDraggingDockUI( UVREditorInteractor* VREditorInteractor );
 
 	/** Are we currently dragging a dock window */
 	bool IsDraggingDockUI();	
 
+	bool IsInteractorDraggingDockUI( const UVREditorInteractor* Interactor ) const;
+
 	/** Hides and unhides all the editor UI panels */
 	void TogglePanelsVisibility();
 
@@ -113,10 +123,11 @@ public:
 protected:
 
 	/** Called when the user presses a button on their motion controller device */
-	void OnVRAction( FEditorViewportClient& ViewportClient, const FVRAction VRAction, const EInputEvent Event, bool& bIsInputCaptured, bool& bOutWasHandled );
+	void OnVRAction( class FEditorViewportClient& ViewportClient, UViewportInteractor* Interactor,
+		const struct FViewportActionKeyInput& Action, const EInputEvent Event, bool& bOutIsInCaptured, bool& bWasHandled );
 
 	/** Called every frame to update hover state */
-	void OnVRHoverUpdate( FEditorViewportClient& ViewportClient, const int32 HandIndex, FVector& HoverImpactPoint, bool& bIsHoveringOverUI, bool& bOutWasHandled );
+	void OnVRHoverUpdate( class FEditorViewportClient& ViewportClient, UViewportInteractor* Interactor, FVector& HoverImpactPoint, bool& bWasHandled );
 
 	/** Testing Slate UI */
 	void CreateUIs();
@@ -142,18 +153,22 @@ protected:
 protected:
 
 	/** Owning object */
-	class FVREditorMode& Owner;
+	class FVREditorMode* VRMode;
 
 	/** All of the floating UIs.  These may or may not be visible (spawned) */
+	UPROPERTY()
 	TArray< class AVREditorFloatingUI* > FloatingUIs;
 
 	/** Our Quick Menu UI */
+	UPROPERTY()
 	AVREditorFloatingUI* QuickMenuUI;
 
 	/** Editor UI panels */
+	UPROPERTY()
 	TArray<AVREditorFloatingUI*> EditorUIPanels;
 
 	/** The Radial Menu UI */
+	UPROPERTY()
 	AVREditorFloatingUI* QuickRadialMenu;
 
 	/** The time since the radial menu was updated */
@@ -172,13 +187,14 @@ protected:
 	// Dragging UI
 	//
 
-	/** Hand that is dragging the UI */
-	int32 DraggingUIHandIndex;
+	/** Interactor that is dragging the UI */
+	class UVREditorInteractor* InteractorDraggingUI;
 
 	/** Offset transform from room-relative transform to the object where we picked it up at */
 	FTransform DraggingUIOffsetTransform;
 
 	/** The current UI that is being dragged */
+	UPROPERTY()
 	class AVREditorDockableWindow* DraggingUI;
 
 	//
@@ -197,15 +213,19 @@ protected:
 	//
 
 	/** Start dragging UI sound */
+	UPROPERTY()
 	class USoundCue* StartDragUISound;
 
 	/** Stop dragging UI sound */
+	UPROPERTY()
 	class USoundCue* StopDragUISound;
 
 	/** Hide UI sound */
+	UPROPERTY()
 	class USoundCue* HideUISound;
 
 	/** Show UI sound */
+	UPROPERTY()
 	class USoundCue* ShowUISound;
 
 	/** If the current dragged dock passed a certain distance if dragged from a hand */
