VR Editor:

diff --git a/Engine/Source/Editor/VREditor/VREditorWorldInteraction.cpp b/Engine/Source/Editor/VREditor/VREditorWorldInteraction.cpp
--- a/Engine/Source/Editor/VREditor/VREditorWorldInteraction.cpp
+++ b/Engine/Source/Editor/VREditor/VREditorWorldInteraction.cpp
@@ -3,2840 +3,156 @@
 #include "VREditorModule.h"
 #include "VREditorWorldInteraction.h"
 #include "VREditorMode.h"
-#include "VREditorTransformGizmo.h"
-#include "Gizmo/VREditorPivotTransformGizmo.h"
 #include "VREditorUISystem.h"
 #include "VREditorFloatingUI.h"
 #include "VREditorDockableWindow.h"
+#include "VREditorInteractor.h"
+#include "VIBaseTransformGizmo.h"
+#include "SnappingUtils.h"
+#include "ScopedTransaction.h"
 
-#include "SnappingUtils.h"
-#include "ScopedTransaction.h"
-
-// For actor placement
-#include "ObjectTools.h"
-#include "AssetSelection.h"
-#include "IPlacementModeModule.h"
-#include "Kismet/GameplayStatics.h"
-
-#define LOCTEXT_NAMESPACE "VREditor"
-
-
-namespace VREd
-{
-	static FAutoConsoleVariable ShowTransformGizmo( TEXT( "VREd.ShowTransformGizmo" ), 1, TEXT( "Whether the transform gizmo should be shown for selected objects" ) );
-	static FAutoConsoleVariable SmoothSnap( TEXT( "VREd.SmoothSnap" ), 1, TEXT( "When enabled with grid snap, transformed objects will smoothly blend to their new location (instead of teleporting instantly)" ) );
-	static FAutoConsoleVariable SmoothSnapSpeed( TEXT( "VREd.SmoothSnapSpeed" ), 30.0f, TEXT( "How quickly objects should interpolate to their new position when grid snapping is enabled" ) );
-	static FAutoConsoleVariable ElasticSnap( TEXT( "VREd.ElasticSnap" ), 1, TEXT( "When enabled with grid snap, you can 'pull' objects slightly away from their snapped position" ) );
-	static FAutoConsoleVariable ElasticSnapStrength( TEXT( "VREd.ElasticSnapStrength" ), 0.3f, TEXT( "How much objects should 'reach' toward their unsnapped position when elastic snapping is enabled with grid snap" ) );
-	static FAutoConsoleVariable MinVelocityForInertia( TEXT( "VREd.MinVelocityForInertia" ), 1.0f, TEXT( "Minimum velocity (in cm/frame in unscaled room space) before inertia will kick in when releasing objects (or the world)" ) );
-	static FAutoConsoleVariable InertiaVelocityBoost( TEXT( "VREd.InertiaVelocityBoost" ), 3000.0f, TEXT( "How much to scale object velocity when releasing dragged simulating objects in Simulate mode" ) );
-	static FAutoConsoleVariable SweepPhysicsWhileSimulating( TEXT( "VREd.SweepPhysicsWhileSimulating" ), 1, TEXT( "If enabled, simulated objects won't be able to penetrate other objects while being dragged in Simulate mode" ) );
-	static FAutoConsoleVariable GizmoRotationSensitivity( TEXT( "VREd.GizmoRotationSensitivity" ), 0.25f, TEXT( "How much to rotate as the user drags on a rotation gizmo handle" ) );
-	static FAutoConsoleVariable AllowVerticalWorldMovement( TEXT( "VREd.AllowVerticalWorldMovement" ), 1, TEXT( "Whether you can move your tracking space away from the origin or not" ) );
-	static FAutoConsoleVariable ScaleWorldFromFloor( TEXT( "VREd.ScaleWorldFromFloor" ), 0, TEXT( "Whether the world should scale relative to your tracking space floor instead of the center of your hand locations" ) );
-	static FAutoConsoleVariable TeleportLerpTime( TEXT( "VREd.TeleportLerpTime" ), 0.1f, TEXT( "The lerp time to teleport" ) );
-	static FAutoConsoleVariable TeleportOffset( TEXT( "VREd.TeleportOffset" ), 100.0f, TEXT( "The offset from the hitresult towards the controller" ) );
-	static FAutoConsoleVariable SizeOfActorsOverContentBrowserThumbnail( TEXT( "VREd.SizeOfActorsOverContentBrowserThumbnail" ), 6.0f, TEXT( "How large objects should be when rendered 'thumbnail size' over the Content Browser" ) );
-	static FAutoConsoleVariable OverrideSizeOfPlacedActors( TEXT( "VREd.OverrideSizeOfPlacedActors" ), 0.0f, TEXT( "If set to a value greater than zero, sets the size in cm of a placed actor (relative to your world's scale)" ) );
-	static FAutoConsoleVariable PlacementInterpolationDuration( TEXT( "VREd.PlacementInterpolationDuration" ), 0.6f, TEXT( "How long we should interpolate newly-placed objects to their target location." ) );
-	static FAutoConsoleVariable DragAtLaserImpactInterpolationDuration( TEXT( "VREd.DragAtLaserImpactInterpolationDuration" ), 0.1f, TEXT( "How long we should interpolate objects between positions when dragging under the laser's impact point" ) );
-	static FAutoConsoleVariable DragAtLaserImpactInterpolationThreshold( TEXT( "VREd.DragAtLaserImpactInterpolationThreshold" ), 5.0f, TEXT( "Minimum distance jumped between frames before we'll force interpolation mode to activated" ) );
-	static FAutoConsoleVariable TeleportLaserPointerLength( TEXT( "VREd.TeleportLaserPointerLength" ), 500000.0f, TEXT( "Distance of the LaserPointer for teleporting" ) );
-	static FAutoConsoleVariable ScaleSensitivity( TEXT( "VREd.ScaleSensitivity" ), 0.01f, TEXT( "Sensitivity for scaling" ) );
-	static FAutoConsoleVariable GridHapticFeedbackStrength( TEXT( "VREd.GridHapticFeedbackStrength" ), 0.4f, TEXT( "Default strength for haptic feedback when moving across grid points" ) );
-	static FAutoConsoleVariable SelectionHapticFeedbackStrength( TEXT( "VREd.SelectionHapticFeedbackStrength" ), 0.5f, TEXT( "Default strength for haptic feedback when selecting objects" ) );
-	static FAutoConsoleVariable DragHapticFeedbackStrength( TEXT( "VREd.DragHapticFeedbackStrength" ), 1.0f, TEXT( "Default strength for haptic feedback when starting to drag objects" ) );
-	static FAutoConsoleVariable ForceGizmoPivotToCenterOfSelectedActorsBounds( TEXT( "VREd.ForceGizmoPivotToCenterOfSelectedActorsBounds" ), 0, TEXT( "When enabled, the gizmo's pivot will always be centered on the selected actors.  Otherwise, we use the pivot of the last selected actor." ) );
-	static FAutoConsoleVariable TrackpadAbsoluteDragSpeed( TEXT( "VREd.TrackpadAbsoluteDragSpeed" ), 40.0f, TEXT( "How fast objects move toward or away when you drag on the touchpad while carrying them" ) );
-	static FAutoConsoleVariable TrackpadRelativeDragSpeed( TEXT( "VREd.TrackpadRelativeDragSpeed" ), 8.0f, TEXT( "How fast objects move toward or away when you hold a direction on an analog stick while carrying them" ) );
-	static FAutoConsoleVariable GizmoHandleHoverScale( TEXT( "VREd.GizmoHandleHoverScale" ), 1.5f, TEXT( "How much to scale up transform gizmo handles when hovered over" ) );
-	static FAutoConsoleVariable GizmoHandleHoverAnimationDuration( TEXT( "VREd.GizmoHandleHoverAnimationDuration" ), 0.1f, TEXT( "How quickly to animate gizmo handle hover state" ) );
-	static FAutoConsoleVariable SnapGridSize( TEXT( "VREd.SnapGridSize" ), 3.0f, TEXT( "How big the snap grid should be, in multiples of the gizmo bounding box size" ) );
-	static FAutoConsoleVariable SnapGridLineWidth( TEXT( "VREd.SnapGridLineWidth" ), 3.0f, TEXT( "Width of the grid lines on the snap grid" ) );
-	static FAutoConsoleVariable HoverHapticFeedbackStrength( TEXT( "VREd.HoverHapticFeedbackStrength" ), 0.1f, TEXT( "Default strength for haptic feedback when hovering" ) );
-	static FAutoConsoleVariable HoverHapticFeedbackTime( TEXT( "VREd.HoverHapticFeedbackTime" ), 0.2f, TEXT( "The minimum time between haptic feedback for hovering" ) );
-	static FAutoConsoleVariable AllowWorldRotationPitchAndRoll( TEXT( "VREd.AllowWorldRotationPitchAndRoll" ), 0, TEXT( "When enabled, you'll not only be able to yaw, but also pitch and roll the world when rotating by gripping with two hands" ) );
-	static FAutoConsoleVariable PivotPointTransformGizmo( TEXT( "VREd.PivotPointTransformGizmo" ), 1, TEXT( "If the pivot point transform gizmo is used instead of the bounding box gizmo" ) );
-}
-
-
-/**
- * Represents an object that we're actively interacting with, such as a selected actor
- */
-class FVREditorTransformable
-{
-	
-public:
-
-	/** Sets up safe defaults */
-	FVREditorTransformable();
-
-	/** The actual object */
-	FWeakObjectPtr Object;
-
-	/** The object's position when we started the action */
-	FTransform StartTransform;
-
-	/** The last position we were at.  Used for smooth interpolation when snapping is enabled */
-	FTransform LastTransform;
-
-	/** The target position we want to be at.  Used for smooth interpolation when snapping is enabled */
-	FTransform TargetTransform;
-
-	/** The target position we'd be at if snapping wasn't enabled.  This is used for 'elastic' smoothing when
-	    snapping is turned on.  Basically it allows us to 'reach' toward the desired position */
-	FTransform UnsnappedTargetTransform;
-
-	/** A "snapshot" transform, used for certain interpolation effects */
-	FTransform InterpolationSnapshotTransform;
-};
-
-
-FVREditorTransformable::FVREditorTransformable()
-	: Object(),
-	  StartTransform( FTransform::Identity ),
-	  LastTransform( FTransform::Identity ),
-	  TargetTransform( FTransform::Identity ),
-	  UnsnappedTargetTransform( FTransform::Identity ),
-	  InterpolationSnapshotTransform( FTransform::Identity )
-{
-}
-
-
-
-// @todo vreditor: Hacky inline implementation of a double vector.  Move elsewhere and polish or get rid.
-struct DVector
-{
-	double X;
-	double Y;
-	double Z;
-
-	DVector()
-	{
-	}
-
-	DVector( const DVector& V )
-		: X( V.X ), Y( V.Y ), Z( V.Z )
-	{
-	}
-
-	DVector( const FVector& V )
-		: X( V.X ), Y( V.Y ), Z( V.Z )
-	{
-	}
-
-	DVector( double InX, double InY, double InZ )
-		: X( InX ), Y( InY ), Z( InZ )
-	{
-	}
-
-	DVector operator+( const DVector& V ) const
-	{
-		return DVector( X + V.X, Y + V.Y, Z + V.Z );
-	}
-
-	DVector operator-( const DVector& V ) const
-	{
-		return DVector( X - V.X, Y - V.Y, Z - V.Z );
-	}
-
-	DVector operator*( double Scale ) const
-	{
-		return DVector( X * Scale, Y * Scale, Z * Scale );
-	}
-
-	double operator|( const DVector& V ) const
-	{
-		return X*V.X + Y*V.Y + Z*V.Z;
-	}
-
-	FVector ToFVector() const
-	{
-		return FVector( (float)X, (float)Y, (float)Z );
-	}
-};
-
-
-// @todo vreditor: Move elsewhere or get rid
-static void SegmentDistToSegmentDouble( DVector A1, DVector B1, DVector A2, DVector B2, DVector& OutP1, DVector& OutP2 )
-{
-	const double Epsilon = 1.e-10;
-
-	// Segments
-	const	DVector	S1 = B1 - A1;
-	const	DVector	S2 = B2 - A2;
-	const	DVector	S3 = A1 - A2;
-
-	const	double	Dot11 = S1 | S1;	// always >= 0
-	const	double	Dot22 = S2 | S2;	// always >= 0
-	const	double	Dot12 = S1 | S2;
-	const	double	Dot13 = S1 | S3;
-	const	double	Dot23 = S2 | S3;
-
-	// Numerator
-	double	N1, N2;
-
-	// Denominator
-	const	double	D = Dot11*Dot22 - Dot12*Dot12;	// always >= 0
-	double	D1 = D;		// T1 = N1 / D1, default D1 = D >= 0
-	double	D2 = D;		// T2 = N2 / D2, default D2 = D >= 0
-
-						// compute the line parameters of the two closest points
-	if( D < Epsilon )
-	{
-		// the lines are almost parallel
-		N1 = 0.0;	// force using point A on segment S1
-		D1 = 1.0;	// to prevent possible division by 0 later
-		N2 = Dot23;
-		D2 = Dot22;
-	}
-	else
-	{
-		// get the closest points on the infinite lines
-		N1 = ( Dot12*Dot23 - Dot22*Dot13 );
-		N2 = ( Dot11*Dot23 - Dot12*Dot13 );
-
-		if( N1 < 0.0 )
-		{
-			// t1 < 0.0 => the s==0 edge is visible
-			N1 = 0.0;
-			N2 = Dot23;
-			D2 = Dot22;
-		}
-		else if( N1 > D1 )
-		{
-			// t1 > 1 => the t1==1 edge is visible
-			N1 = D1;
-			N2 = Dot23 + Dot12;
-			D2 = Dot22;
-		}
-	}
-
-	if( N2 < 0.0 )
-	{
-		// t2 < 0 => the t2==0 edge is visible
-		N2 = 0.0;
-
-		// recompute t1 for this edge
-		if( -Dot13 < 0.0 )
-		{
-			N1 = 0.0;
-		}
-		else if( -Dot13 > Dot11 )
-		{
-			N1 = D1;
-		}
-		else
-		{
-			N1 = -Dot13;
-			D1 = Dot11;
-		}
-	}
-	else if( N2 > D2 )
-	{
-		// t2 > 1 => the t2=1 edge is visible
-		N2 = D2;
-
-		// recompute t1 for this edge
-		if( ( -Dot13 + Dot12 ) < 0.0 )
-		{
-			N1 = 0.0;
-		}
-		else if( ( -Dot13 + Dot12 ) > Dot11 )
-		{
-			N1 = D1;
-		}
-		else
-		{
-			N1 = ( -Dot13 + Dot12 );
-			D1 = Dot11;
-		}
-	}
-
-	// finally do the division to get the points' location
-	const double T1 = ( FMath::Abs( N1 ) < Epsilon ? 0.0 : N1 / D1 );
-	const double T2 = ( FMath::Abs( N2 ) < Epsilon ? 0.0 : N2 / D2 );
-
-	// return the closest points
-	OutP1 = A1 + S1 * T1;
-	OutP2 = A2 + S2 * T2;
-}
-
-
-FVREditorWorldInteraction::FVREditorWorldInteraction( FVREditorMode& InitOwner )
-	: Owner( InitOwner ),
-	  bIsInterpolatingTransformablesFromSnapshotTransform( false ),
-	  TransformablesInterpolationStartTime( FTimespan::Zero() ),
-	  TransformablesInterpolationDuration( 1.0f ),
-	  PlacingMaterialOrTextureAsset( nullptr ),
-	  LastWorldToMetersScale( 100.0f ),
-	  TransformGizmoActor( nullptr ),
-	  GizmoLocalBounds( FBox( 0 ) ),
-	  bDraggedSinceLastSelection( false ),
-	  LastDragGizmoStartTransform( FTransform::Identity ),
-	  StartDragAngleOnRotation( ),
-	  bIsTeleporting( false ),
-	  TeleportLerpAlpha( 0 ),
-	  TeleportStartLocation( FVector::ZeroVector ),
-	  TeleportGoalLocation( FVector::ZeroVector ),
-	  bJustTeleported( false ),
-	  TeleportSound( nullptr ),
-	  DropMaterialOrMaterialSound( nullptr ),
-	  DockSelectDistance( 100 ),
-	  FloatingUIAssetDraggedFrom( nullptr )
-{
-	// Register to find out about VR input events
-	Owner.OnVRAction().AddRaw( this, &FVREditorWorldInteraction::OnVRAction );
-	Owner.OnVRHoverUpdate().AddRaw( this, &FVREditorWorldInteraction::OnVRHoverUpdate );
-
-	// Find out about selection changes
-	USelection::SelectionChangedEvent.AddRaw( this, &FVREditorWorldInteraction::OnActorSelectionChanged );
-
-	// Find out when the user drags stuff out of a content browser
-	FEditorDelegates::OnAssetDragStarted.AddRaw( this, &FVREditorWorldInteraction::OnAssetDragStartedFromContentBrowser );
-
-	// Load sounds
-	TeleportSound = LoadObject<USoundCue>( nullptr, TEXT( "/Engine/VREditor/Sounds/VR_teleport_Cue" ) );
-	check( TeleportSound != nullptr );
-
-	DropMaterialOrMaterialSound = LoadObject<USoundCue>( nullptr, TEXT( "/Engine/VREditor/Sounds/VR_grab_Cue" ) );
-	check( DropMaterialOrMaterialSound != nullptr );
-}
-
-
-
-FVREditorWorldInteraction::~FVREditorWorldInteraction()
-{
-	Owner.DestroyTransientActor( TransformGizmoActor );
-	TransformGizmoActor = nullptr;
-	FloatingUIAssetDraggedFrom = nullptr;
-	TeleportSound = nullptr;
-	DropMaterialOrMaterialSound = nullptr;
-
-	FEditorDelegates::OnAssetDragStarted.RemoveAll( this );
-
-	USelection::SelectionChangedEvent.RemoveAll( this );
-
-	Owner.OnVRAction().RemoveAll( this );
-	Owner.OnVRHoverUpdate().RemoveAll( this );
-}
-
-
-void FVREditorWorldInteraction::AddReferencedObjects( FReferenceCollector& Collector )
-{
-	Collector.AddReferencedObject( TransformGizmoActor );
-	Collector.AddReferencedObject( PlacingMaterialOrTextureAsset );
-	Collector.AddReferencedObject( FloatingUIAssetDraggedFrom );
-	Collector.AddReferencedObject( TeleportSound );
-	Collector.AddReferencedObject( DropMaterialOrMaterialSound );
-}
-
-
-void FVREditorWorldInteraction::OnVRAction( FEditorViewportClient& ViewportClient, const FVRAction VRAction, const EInputEvent Event, bool& bIsInputCaptured, bool& bWasHandled )
-{
-	if( !bWasHandled )
-	{
-		// Delete
-		if( VRAction.ActionType == EVRActionType::Delete )
-		{
-			if( Event == IE_Pressed )
-			{
-				DeleteSelectedObjects();
-			}
-			bWasHandled = true;
-		}
-
-		// Selection/Movement
-		if( (VRAction.ActionType == EVRActionType::SelectAndMove ||
-			VRAction.ActionType == EVRActionType::SelectAndMove_LightlyPressed) &&
-			!bWasHandled )
-		{
-			FVirtualHand& Hand = Owner.GetVirtualHand( VRAction.HandIndex );
-			FVirtualHand& OtherHand = Owner.GetOtherHand( VRAction.HandIndex );
-
-			if( Event == IE_Pressed )
-			{	
-				if (VRAction.ActionType == EVRActionType::SelectAndMove_LightlyPressed)
-				{
-					// Check to teleport before checking on actor selection, otherwise objects will be selected when teleporting
-					if (Hand.DraggingMode == EVREditorDraggingMode::World && OtherHand.DraggingMode != EVREditorDraggingMode::AssistingDrag )
-					{
-						StartTeleport( VRAction.HandIndex );
-						bWasHandled = true;
-						bJustTeleported = true;
-					}
-				}
-				else if( bJustTeleported )
-				{
-					bWasHandled = true;
-					bJustTeleported = false;
-				}
-
-				// Nothing should happen when the hand has a UI in front, otherwise it will deselect everything
-				if( Hand.bHasUIInFront )
-				{
-					bWasHandled = true;
-				}
-
-				// No clicking while we're dragging the world.  (No laser pointers are visible, anyway.)
-				const bool bIsDraggingWorldWithTwoHands =
-					( Hand.DraggingMode == EVREditorDraggingMode::World && OtherHand.DraggingMode == EVREditorDraggingMode::AssistingDrag ) ||
-					( OtherHand.DraggingMode == EVREditorDraggingMode::World && Hand.DraggingMode == EVREditorDraggingMode::AssistingDrag );
-				
-				FHitResult HitResult = Owner.GetHitResultFromLaserPointer( VRAction.HandIndex );
-				if( !bIsDraggingWorldWithTwoHands && !bWasHandled && HitResult.Actor.IsValid() )
-				{
-					if( IsInteractableComponent( HitResult.GetComponent() ) )
-					{
-						Hand.ClickingOnComponent = HitResult.GetComponent();
-
-						AActor* Actor = HitResult.Actor.Get();
-
-						FVector LaserPointerStart, LaserPointerEnd;
-						const bool bHaveLaserPointer = Owner.GetLaserPointer( VRAction.HandIndex, /* Out */ LaserPointerStart, /* Out */ LaserPointerEnd );
-						check( bHaveLaserPointer );
-
-						if (Actor->IsA( AVREditorDockableWindow::StaticClass() ))
-						{
-							if( VRAction.ActionType == EVRActionType::SelectAndMove && Owner.GetUISystem().GetDraggingDockUI() == nullptr )
-							{
-								AVREditorDockableWindow* DockableWindow = Cast< AVREditorDockableWindow >( Actor );
-								if (DockableWindow)
-								{
-									if( HitResult.Component == DockableWindow->GetCloseButtonMeshComponent() )
-									{
-										// Close the window
-										const bool bShouldShow = false;
-										const bool bShowOnHand = false;
-										const bool bRefreshQuickMenu = true;
-										Owner.GetUISystem().ShowEditorUIPanel( DockableWindow, Owner.GetOtherHandIndex( VRAction.HandIndex ), bShouldShow, bShowOnHand, bRefreshQuickMenu );
-									}
-									else if ( HitResult.Component == DockableWindow->GetSelectionBarMeshComponent() )
-									{
-										Hand.DraggingMode = EVREditorDraggingMode::DockableWindow;
-										DockSelectDistance = ( LaserPointerStart - HitResult.Location ).Size();
-										Owner.GetUISystem().StartDraggingDockUI( DockableWindow, VRAction.HandIndex, DockSelectDistance );
-									}
-									bWasHandled = true;
-								}
-							}
-							else if (VRAction.ActionType == EVRActionType::SelectAndMove_LightlyPressed)
-							{
-								bWasHandled = true;
-							}
-						}
-						else
-						{
-							bWasHandled = true;
-
-							// Is the other hand already dragging this stuff?
-							if( OtherHand.DraggingMode == EVREditorDraggingMode::ActorsWithGizmo ||
-								OtherHand.DraggingMode == EVREditorDraggingMode::ActorsFreely )
-							{
-								// Only if they clicked on one of the objects we're moving
-								if( IsTransformingActor( Actor ) )
-								{
-									// If we were dragging with a gizmo, we'll need to stop doing that and instead drag freely.
-									if( OtherHand.DraggingMode == EVREditorDraggingMode::ActorsWithGizmo )
-									{
-										const int32 OtherHandIndex = Owner.GetOtherHandIndex( VRAction.HandIndex );
-										StopDragging( Owner.GetOtherHandIndex( OtherHandIndex ) );
-
-										FVRAction OtherHandVRAction( EVRActionType::SelectAndMove, OtherHandIndex );
-										const bool bIsPlacingActors = false;
-										StartDraggingActors( OtherHandVRAction, HitResult.GetComponent(), OtherHand.HoverLocation, bIsPlacingActors );
-									}
-
-									Hand.DraggingMode = Hand.LastDraggingMode = EVREditorDraggingMode::AssistingDrag;
-
-									Hand.bIsFirstDragUpdate = true;
-									Hand.bWasAssistingDrag = true;
-									Hand.DragRayLength = ( HitResult.ImpactPoint - LaserPointerStart ).Size();
-									Hand.LastLaserPointerStart = LaserPointerStart;
-									Hand.LastDragToLocation = HitResult.ImpactPoint;
-									Hand.LaserPointerImpactAtDragStart = HitResult.ImpactPoint;
-									Hand.DragTranslationVelocity = FVector::ZeroVector;
-									Hand.DragRayLengthVelocity = 0.0f;
-									Hand.DraggingTransformGizmoComponent = nullptr;
-									Hand.TransformGizmoInteractionType = ETransformGizmoInteractionType::None;
-									Hand.bIsDrivingVelocityOfSimulatedTransformables = false;
-
-									Hand.GizmoStartTransform = OtherHand.GizmoStartTransform;
-									Hand.GizmoStartLocalBounds = OtherHand.GizmoStartLocalBounds;
-									Hand.GizmoSpaceFirstDragUpdateOffsetAlongAxis = FVector::ZeroVector;	// Will be determined on first update
-									Hand.GizmoSpaceDragDeltaFromStartOffset = FVector::ZeroVector;	// Set every frame while dragging
-
-									bDraggedSinceLastSelection = true;
-									LastDragGizmoStartTransform = Hand.GizmoStartTransform;
-
-									bIsInputCaptured = true;
-
-									// Play a haptic effect when objects are picked up
-									const float Strength = VREd::DragHapticFeedbackStrength->GetFloat();	// @todo vreditor: Tweak
-									Owner.PlayHapticEffect(
-										VRAction.HandIndex == VREditorConstants::LeftHandIndex ? Strength : 0.0f,
-										VRAction.HandIndex == VREditorConstants::RightHandIndex ? Strength : 0.0f );
-								}
-								else
-								{
-									// @todo vreditor: We don't currently support selection/dragging separate objects with either hand
-								}
-							}
-							else if( OtherHand.DraggingMode == EVREditorDraggingMode::ActorsWithGizmo )
-							{
-								// We don't support dragging objects with the gizmo using two hands.  Just absorb it.
-							}
-							else if( OtherHand.DraggingMode == EVREditorDraggingMode::ActorsAtLaserImpact )
-							{
-								// Doesn't work with two hands.  Just absorb it.
-							}
-							else
-							{
-								// Default to replacing our selection with whatever the user clicked on
-								enum ESelectionModification
-								{
-									Replace,
-									AddTo,
-									Toggle
-								} SelectionModification = ESelectionModification::Replace;
-
-								// If the other hand is holding down the button after selecting an object, we'll allow this hand to toggle selection
-								// of additional objects (multi select)
-								{
-									if( OtherHand.ClickingOnComponent.IsValid() &&		// Other hand is clicking on something
-										Actor != TransformGizmoActor &&					// Not clicking on a gizmo
-										OtherHand.ClickingOnComponent.Get()->GetOwner() != Actor )	// Not clicking on same actor
-									{
-										// OK, the other hand is holding down the "select and move" button.
-
-										// If the modifier key is held for the hand, we'll we'll only add to the selection set,
-										// otherwise we toggle selection on and off.
-										if( Hand.bIsModifierPressed )
-										{
-											SelectionModification = ESelectionModification::AddTo;
-										}
-										else
-										{
-											SelectionModification = ESelectionModification::Toggle;
-										}
-									}
-								}
-
-								// Only start dragging if the actor was already selected (and it's a full press), or if we clicked on a gizmo.  It's too easy to 
-								// accidentally move actors around the scene when you only want to select them otherwise.
-								bool bShouldDragSelected =
-									( ( Actor->IsSelected() && VRAction.ActionType == EVRActionType::SelectAndMove ) || Actor == TransformGizmoActor ) &&
-									( SelectionModification == ESelectionModification::Replace );
-
-								bool bSelectionChanged = false;
-								if( Actor != TransformGizmoActor )  // Don't change any actor selection state if the user clicked on a gizmo
-								{
-									const bool bWasSelected = Actor->IsSelected();
-
-									// Clicked on a normal actor.  So update selection!
-									if( SelectionModification == ESelectionModification::Replace || SelectionModification == ESelectionModification::AddTo )
-									{
-										if( !bWasSelected )
-										{
-											if( SelectionModification == ESelectionModification::Replace )
-											{
-												GEditor->SelectNone( true, true );
-											}
-
-											GEditor->SelectActor( Actor, true, true );
-											bSelectionChanged = true;
-										}
-									}
-									else if( SelectionModification == ESelectionModification::Toggle )
-									{
-										GEditor->SelectActor( Actor, !Actor->IsSelected(), true );
-										bSelectionChanged = true;
-									}
-
-									// If we've selected an actor with a light press, don't prevent the user from moving the actor if they
-									// press down all the way.  We only want locking when interacting with gizmos
-									Hand.bAllowTriggerLightPressLocking = false;
-
-									// If the user did a full press to select an actor that wasn't selected, allow them to drag it right away
-									const bool bOtherHandTryingToDrag =
-										OtherHand.ClickingOnComponent.IsValid() &&
-										OtherHand.ClickingOnComponent.Get()->GetOwner()->IsSelected() &&
-										OtherHand.ClickingOnComponent.Get()->GetOwner() == HitResult.GetComponent()->GetOwner();	// Trying to drag same actor
-									if( ( !bWasSelected || bOtherHandTryingToDrag ) &&	
-										Actor->IsSelected() &&
-										( VRAction.ActionType == EVRActionType::SelectAndMove || bOtherHandTryingToDrag ) )
-									{
-										bShouldDragSelected = true;
-									}
-								}
-
-								if( bShouldDragSelected && Hand.DraggingMode != EVREditorDraggingMode::DockableWindow )
-								{
-									const bool bIsPlacingActors = false;
-									StartDraggingActors( VRAction, HitResult.GetComponent(), HitResult.ImpactPoint, bIsPlacingActors );
-								}
-								else if( bSelectionChanged )
-								{
-									// User didn't drag but did select something, so play a haptic effect.
-									const float Strength = VREd::SelectionHapticFeedbackStrength->GetFloat();
-									Owner.PlayHapticEffect(
-										VRAction.HandIndex == VREditorConstants::LeftHandIndex ? Strength : 0.0f,
-										VRAction.HandIndex == VREditorConstants::RightHandIndex ? Strength : 0.0f );
-								}
-							}
-						}
-					}
-				}
-			}
-			else if( Event == IE_Released )
-			{
-				if( Hand.ClickingOnComponent.IsValid() )
-				{
-					bWasHandled = true;
-					Hand.ClickingOnComponent = nullptr;
-				}
-
-				const bool bIsDraggingWorld = Hand.DraggingMode == EVREditorDraggingMode::World;
-				const bool bIsDraggingWorldWithTwoHands =
-					( Hand.DraggingMode == EVREditorDraggingMode::World && OtherHand.DraggingMode == EVREditorDraggingMode::AssistingDrag ) ||
-					( OtherHand.DraggingMode == EVREditorDraggingMode::World && Hand.DraggingMode == EVREditorDraggingMode::AssistingDrag );
-
-				// Don't allow the trigger to cancel our drag on release if we're dragging the world. 
-				if( Hand.DraggingMode != EVREditorDraggingMode::Nothing &&
-					!bIsDraggingWorld && 
-					!bIsDraggingWorldWithTwoHands )
-				{
-					// If we were placing a material, go ahead and do that now
-					if( Hand.DraggingMode == EVREditorDraggingMode::Material )
-					{
-						PlaceDraggedMaterialOrTexture( VRAction.HandIndex );
-					}
-
-					// If we were placing something, bring the window back
-					if( Hand.DraggingMode == EVREditorDraggingMode::ActorsAtLaserImpact ||
-						Hand.DraggingMode == EVREditorDraggingMode::Material )
-					{
-						Owner.GetUISystem().ShowEditorUIPanel( FloatingUIAssetDraggedFrom, Owner.GetOtherHandIndex( VRAction.HandIndex ), true, false, false, false );
-						FloatingUIAssetDraggedFrom = nullptr;
-					}
-
-					StopDragging( VRAction.HandIndex );
-					bWasHandled = true;
-				}
-			}
-		}
-
-		// World Movement
-		else if( VRAction.ActionType == EVRActionType::WorldMovement )
-		{
-			FVirtualHand& Hand = Owner.GetVirtualHand( VRAction.HandIndex );
-
-			if( Event == IE_Pressed )
-			{
-				// Is our other hand already dragging the world around?
-				FVirtualHand& OtherHand = Owner.GetOtherHand( VRAction.HandIndex );
-				if( OtherHand.DraggingMode == EVREditorDraggingMode::World )
-				{
-					Hand.DraggingMode = Hand.LastDraggingMode = EVREditorDraggingMode::AssistingDrag;
-					Hand.bWasAssistingDrag = true;
-				}
-				else
-				{
-					// Start dragging the world
-					Hand.DraggingMode = Hand.LastDraggingMode = EVREditorDraggingMode::World;
-					Hand.bWasAssistingDrag = false;
-
-					// Starting a new drag, so make sure the other hand doesn't think it's assisting us
-					OtherHand.bWasAssistingDrag = false;
-
-					// Stop any interia from the other hand's previous movements -- we've grabbed the world and it needs to stick!
-					OtherHand.DragTranslationVelocity = FVector::ZeroVector;
-				}
-
-				Hand.bIsFirstDragUpdate = true;
-				Hand.DragRayLength = 0.0f;
-				Hand.LastLaserPointerStart = Hand.Transform.GetLocation();
-				Hand.LastDragToLocation = Hand.Transform.GetLocation();
-				Hand.LaserPointerImpactAtDragStart = Hand.Transform.GetLocation();
-				Hand.DragTranslationVelocity = FVector::ZeroVector;
-				Hand.DragRayLengthVelocity = 0.0f;
-				Hand.bIsDrivingVelocityOfSimulatedTransformables = false;
-
-				// We won't use gizmo features for world movement
-				Hand.DraggingTransformGizmoComponent = nullptr;
-				Hand.TransformGizmoInteractionType = ETransformGizmoInteractionType::None;
-				Hand.OptionalHandlePlacement.Reset();
-				Hand.GizmoStartTransform = FTransform::Identity;
-				Hand.GizmoStartLocalBounds = FBox( 0 );
-				Hand.GizmoSpaceFirstDragUpdateOffsetAlongAxis = FVector::ZeroVector;
-				Hand.GizmoSpaceDragDeltaFromStartOffset = FVector::ZeroVector;
-
-				bIsInputCaptured = true;
-			}
-			else if( Event == IE_Released )
-			{
-				StopDragging( VRAction.HandIndex );
-			}
-		}
-	}
-}
-
-
-void FVREditorWorldInteraction::OnVRHoverUpdate( FEditorViewportClient& ViewportClient, const int32 HandIndex, FVector& HoverImpactPoint, bool& bIsHoveringOverUI, bool& bWasHandled )
-{
-	if( !bWasHandled )
-	{
-		FVirtualHand& Hand = Owner.GetVirtualHand( HandIndex );
-		FVirtualHand& OtherHand = Owner.GetOtherHand( HandIndex );
-
-		float HapticFeedbackStrength = 0.0f;
-		float HapticFeedbackTimeBuffer = 0.0f;
-
-		UActorComponent* PreviousHoverGizmoComponent = Hand.HoveringOverTransformGizmoComponent.Get();
-		Hand.HoveringOverTransformGizmoComponent = nullptr;
-
-		UActorComponent* NewHoveredActorComponent = nullptr;
-
-		FHitResult HitResult = Owner.GetHitResultFromLaserPointer( HandIndex );
-		if( HitResult.Actor.IsValid() )
-		{
-			AActor* Actor = HitResult.Actor.Get();
-
-			if( IsInteractableComponent( HitResult.GetComponent() ) )
-			{
-				this->HoveredObjects.Add( FViewportHoverTarget( Actor ) );
-
-				HoverImpactPoint = HitResult.ImpactPoint;
-
-				if( Actor == TransformGizmoActor )
-				{
-					NewHoveredActorComponent = HitResult.GetComponent();
-					Hand.HoveringOverTransformGizmoComponent = HitResult.GetComponent();
-
-					if( Hand.HoveringOverTransformGizmoComponent != PreviousHoverGizmoComponent )
-					{
-						HapticFeedbackStrength = VREd::HoverHapticFeedbackStrength->GetFloat();
-						HapticFeedbackTimeBuffer = VREd::HoverHapticFeedbackTime->GetFloat();
-					}
-
-					Hand.HoverHapticCheckLastHoveredGizmoComponent = HitResult.GetComponent();
-				}
-				else if( Actor->IsA( AVREditorDockableWindow::StaticClass() ) )
-				{
-					NewHoveredActorComponent = HitResult.GetComponent();
-					AVREditorDockableWindow* DockableWindow = Cast< AVREditorDockableWindow >( Actor );
-					if(DockableWindow && NewHoveredActorComponent != Hand.HoveredActorComponent )
-					{
-						if( NewHoveredActorComponent != OtherHand.HoveredActorComponent )
-						{
-							DockableWindow->OnEnterHover( HitResult, HandIndex );
-						}
-						
-						HapticFeedbackStrength = VREd::HoverHapticFeedbackStrength->GetFloat();
-						HapticFeedbackTimeBuffer = VREd::HoverHapticFeedbackTime->GetFloat();
-					}
-				}
-
-				bWasHandled = true;
-			}
-		}
-
-		// Leave dockable window
-		if( Hand.HoveredActorComponent != nullptr && ( Hand.HoveredActorComponent != NewHoveredActorComponent || NewHoveredActorComponent == nullptr ) && Hand.HoveredActorComponent != OtherHand.HoveredActorComponent )
-		{
-			AActor* HoveredActor = Hand.HoveredActorComponent->GetOwner();
-			if (HoveredActor->IsA( AVREditorDockableWindow::StaticClass() ))
-			{
-				AVREditorDockableWindow* DockableWindow = Cast< AVREditorDockableWindow >( HoveredActor );
-				if (DockableWindow)
-				{
-					DockableWindow->OnLeaveHover( HandIndex, NewHoveredActorComponent );
-				}
-			}
-		}
-
-		// Play the haptic feedback
-		if( HapticFeedbackStrength > 0.0f )
-		{
-			const FTimespan CurrentTime = FTimespan::FromSeconds( FPlatformTime::Seconds() );
-			const FTimespan TimeSinceLastHover = CurrentTime - Hand.LastGizmoHandleHapticTime;
-
-			// Only play a haptic if we hovered over a different gizmo, or if we haven't already played a haptic in a little while
-			if (Hand.HoveringOverTransformGizmoComponent != Hand.HoverHapticCheckLastHoveredGizmoComponent ||
-				TimeSinceLastHover.GetTotalSeconds() > HapticFeedbackTimeBuffer )
-			{
-				Owner.PlayHapticEffect(
-					HandIndex == VREditorConstants::LeftHandIndex ? HapticFeedbackStrength : 0.0f,
-					HandIndex == VREditorConstants::RightHandIndex ? HapticFeedbackStrength : 0.0f );
-
-				Hand.LastGizmoHandleHapticTime = CurrentTime;
-			}
-		}
-
-		// Update the hovered actor component with the new component
-		Hand.HoveredActorComponent = NewHoveredActorComponent;
-	}
-}
-
-
-void FVREditorWorldInteraction::Tick( FEditorViewportClient* ViewportClient, const float DeltaTime )
-{
-	const FTimespan CurrentTime = FTimespan::FromSeconds( FPlatformTime::Seconds() );
-	const float WorldToMetersScale = Owner.GetWorld()->GetWorldSettings()->WorldToMeters;
-
-	// Update viewport with any objects that are currently hovered over
-	{
-		const bool bUseEditorSelectionHoverFeedback = GEditor != NULL && GetDefault<ULevelEditorViewportSettings>()->bEnableViewportHoverFeedback;
-		if( bUseEditorSelectionHoverFeedback )
-		{
-			if( ViewportClient->IsLevelEditorClient() )
-			{
-				FLevelEditorViewportClient* LevelEditorViewportClient = static_cast<FLevelEditorViewportClient*>( ViewportClient );
-				LevelEditorViewportClient->UpdateHoveredObjects( HoveredObjects );
-			}
-		}
-
-		// This will be filled in again during the next input update
-		HoveredObjects.Reset();
-	}
-
-	const float WorldScaleFactor = Owner.GetWorldScaleFactor();
-
-	// Move selected actors
-	int32 DraggingActorsWithHandIndex = INDEX_NONE;
-	int32 AssistingDragWithHandIndex = INDEX_NONE;
-	{
-		for( int32 HandIndex = 0; HandIndex < VREditorConstants::NumVirtualHands; ++HandIndex )
-		{
-			FVirtualHand& Hand = Owner.GetVirtualHand( HandIndex );
-
-			bool bCanSlideRayLength = false;
-
-			if( Hand.DraggingMode == EVREditorDraggingMode::ActorsWithGizmo || 
-				Hand.DraggingMode == EVREditorDraggingMode::ActorsFreely ||
-				Hand.DraggingMode == EVREditorDraggingMode::ActorsAtLaserImpact )
-			{
-				check( DraggingActorsWithHandIndex == INDEX_NONE );	// Only support dragging one thing at a time right now!
-				DraggingActorsWithHandIndex = HandIndex;
-
-				if( Hand.DraggingMode != EVREditorDraggingMode::ActorsAtLaserImpact )
-				{
-					bCanSlideRayLength = true;
-				}
-			}
-
-			if( Hand.DraggingMode == EVREditorDraggingMode::AssistingDrag )
-			{
-				check( AssistingDragWithHandIndex == INDEX_NONE );	// Only support assisting one thing at a time right now!
-				AssistingDragWithHandIndex = HandIndex;
-
-				bCanSlideRayLength = true;
-			}
-
-			if( bCanSlideRayLength )
-			{
-				// If we're dragging an object, go ahead and slide the object along the ray based on how far they slide their touch
-				// Make sure they are touching the trackpad, otherwise we get bad data
-				if( Hand.bIsTrackpadPositionValid[ 1 ] )
-				{
-					const bool bIsAbsolute = ( Owner.GetHMDDeviceType() == EHMDDeviceType::DT_SteamVR );
-					float SlideDelta = GetTrackpadSlideDelta( HandIndex ) * WorldScaleFactor;
-
-					if( !FMath::IsNearlyZero( SlideDelta ) )
-					{
-						Hand.DragRayLength += SlideDelta;
-
-						Hand.DragRayLengthVelocity = 0.0f;
-
-						// Don't apply inertia unless the user dragged a decent amount this frame
-						if( bIsAbsolute && FMath::Abs( SlideDelta ) >= VREd::MinVelocityForInertia->GetFloat() * WorldScaleFactor )
-						{
-							// Don't apply inertia if our data is sort of old
-							if( CurrentTime - Hand.LastTrackpadPositionUpdateTime <= FTimespan::FromSeconds( 1.0f / 30.0f ) )
-							{
-								Hand.DragRayLengthVelocity = SlideDelta;
-							}
-						}
-
-						// Don't go too far
-						if( Hand.DragRayLength < 0.0f )
-						{
-							Hand.DragRayLength = 0.0f;
-							Hand.DragRayLengthVelocity = 0.0f;
-						}
-					}
-				}
-				else
-				{
-					if( !FMath::IsNearlyZero( Hand.DragRayLengthVelocity ) )
-					{
-						// Apply drag ray length inertia
-						Hand.DragRayLength += Hand.DragRayLengthVelocity;
-
-						// Don't go too far!
-						if( Hand.DragRayLength < 0.0f )
-						{
-							Hand.DragRayLength = 0.0f;
-							Hand.DragRayLengthVelocity = 0.0f;
-						}
-
-						// Apply damping
-						FVector RayVelocityVector( Hand.DragRayLengthVelocity, 0.0f, 0.0f );
-						const bool bVelocitySensitive = true;
-						Owner.ApplyVelocityDamping( RayVelocityVector, bVelocitySensitive );
-						Hand.DragRayLengthVelocity = RayVelocityVector.X;
-					}
-					else
-					{
-						Hand.DragRayLengthVelocity = 0.0f;
-					}
-				}
-			}
-			else
-			{
-				if ( !Hand.bHasUIInFront && (Hand.bIsTrackpadPositionValid[0] && Hand.bIsTrackpadPositionValid[1] ) )
-				{
-					FVirtualHand& OtherHand = Owner.GetOtherHand( HandIndex );
-					if( Owner.GetHMDDeviceType() == EHMDDeviceType::DT_SteamVR && Hand.bIsTouchingTrackpad && !OtherHand.bIsTouchingTrackpad )
-					{
-						Owner.GetUISystem().TryToSpawnRadialMenu( HandIndex );
-					}
-
-					if( Owner.GetUISystem().IsShowingRadialMenu( HandIndex ) )
-					{
-						Owner.GetUISystem().UpdateRadialMenu( HandIndex );
-					}
-				}
-			}
-		}
-
-
-		for( int32 HandIndex = 0; HandIndex < VREditorConstants::NumVirtualHands; ++HandIndex )
-		{
-			FVirtualHand& Hand = Owner.GetVirtualHand( HandIndex );
-			if( Hand.DraggingMode == EVREditorDraggingMode::ActorsWithGizmo ||
-				Hand.DraggingMode == EVREditorDraggingMode::ActorsFreely ||
-				Hand.DraggingMode == EVREditorDraggingMode::ActorsAtLaserImpact ||
-				Hand.DraggingMode == EVREditorDraggingMode::World )
-			{
-				// Are we dragging with two hands?
-				FVirtualHand* OtherHand = nullptr;
-				int32 OtherHandIndex = AssistingDragWithHandIndex;
-				if( OtherHandIndex != INDEX_NONE )
-				{
-					OtherHand = &Owner.GetVirtualHand( AssistingDragWithHandIndex );
-				}
-
-				// Check to see if the other hand has any inertia to contribute
-				FVirtualHand* OtherHandThatWasAssistingDrag = nullptr;
-				{
-					FVirtualHand& OtherHandRef = Owner.GetOtherHand( HandIndex );
-
-					// If the other hand isn't doing anything, but the last thing it was doing was assisting a drag, then allow it
-					// to contribute inertia!
-					if( OtherHandRef.DraggingMode == EVREditorDraggingMode::Nothing && OtherHandRef.bWasAssistingDrag )
-					{
-						if( !OtherHandRef.DragTranslationVelocity.IsNearlyZero( KINDA_SMALL_NUMBER ) )
-						{
-							OtherHandThatWasAssistingDrag = &OtherHandRef;
-							OtherHandIndex = Owner.GetOtherHandIndex( HandIndex );
-						}
-					}
-				}
-
-
-				FVector DraggedTo = Hand.Transform.GetLocation();
-				FVector DragDelta = DraggedTo - Hand.LastTransform.GetLocation();
-				FVector DragDeltaFromStart = DraggedTo - Hand.LaserPointerImpactAtDragStart;
-
-				FVector OtherHandDraggedTo = FVector::ZeroVector;
-				FVector OtherHandDragDelta = FVector::ZeroVector;
-				FVector OtherHandDragDeltaFromStart = FVector::ZeroVector;
-				if( OtherHand != nullptr )
-				{
-					OtherHandDraggedTo = OtherHand->Transform.GetLocation();
-					OtherHandDragDelta = OtherHandDraggedTo - OtherHand->LastTransform.GetLocation();
-					OtherHandDragDeltaFromStart = DraggedTo - OtherHand->LaserPointerImpactAtDragStart;
-				}
-				else if( OtherHandThatWasAssistingDrag != nullptr )
-				{
-					OtherHandDragDelta = OtherHandThatWasAssistingDrag->DragTranslationVelocity;
-					OtherHandDraggedTo = OtherHandThatWasAssistingDrag->LastDragToLocation + OtherHandDragDelta;
-					OtherHandDragDeltaFromStart = OtherHandDraggedTo - OtherHandThatWasAssistingDrag->LaserPointerImpactAtDragStart;
-				}
-
-
-				FVector LaserPointerStart = Hand.Transform.GetLocation();
-				FVector LaserPointerDirection = Hand.Transform.GetUnitAxis( EAxis::X );
-				FVector LaserPointerEnd = Hand.Transform.GetLocation();
-
-				if( Hand.DraggingMode == EVREditorDraggingMode::ActorsWithGizmo ||
-					Hand.DraggingMode == EVREditorDraggingMode::ActorsFreely ||
-					Hand.DraggingMode == EVREditorDraggingMode::ActorsAtLaserImpact )
-				{
-					// Move objects using the laser pointer (in world space)
-					if( Owner.GetLaserPointer( HandIndex, /* Out */ LaserPointerStart, /* Out */ LaserPointerEnd ) )
-					{
-						LaserPointerDirection = ( LaserPointerEnd - LaserPointerStart ).GetSafeNormal();
-
-						if( Hand.DraggingMode == EVREditorDraggingMode::ActorsAtLaserImpact )
-						{
-							// Check to see if the laser pointer is over something we can drop on
-							FVector HitLocation = FVector::ZeroVector;
-							bool bHitSomething = FindPlacementPointUnderLaser( HandIndex, /* Out */ HitLocation );
-
-							if( bHitSomething )
-							{
-								// If the object moved reasonably far between frames, it might be because the angle we were aligning
-								// the object with during placement changed radically.  To avoid it popping, we smoothly interpolate
-								// it's position over a very short timespan
-								if( !bIsInterpolatingTransformablesFromSnapshotTransform )	// Let the last animation finish first
-								{
-									const float ScaledDragDistance = DragDelta.Size() * WorldScaleFactor;
-									if( ScaledDragDistance >= VREd::DragAtLaserImpactInterpolationThreshold->GetFloat() )
-									{
-										bIsInterpolatingTransformablesFromSnapshotTransform = true;
-										TransformablesInterpolationStartTime = CurrentTime;
-										TransformablesInterpolationDuration = VREd::DragAtLaserImpactInterpolationDuration->GetFloat();
-
-										// Snapshot everything
-										for( TUniquePtr<FVREditorTransformable>& TransformablePtr : Transformables )
-										{
-											FVREditorTransformable& Transformable = *TransformablePtr;
-											Transformable.InterpolationSnapshotTransform = Transformable.LastTransform;
-										}
-									}
-								}
-
-								DraggedTo = HitLocation;
-							}
-							else
-							{
-								DraggedTo = LaserPointerStart + LaserPointerDirection * Hand.DragRayLength;
-							}
-						}
-						else
-						{
-							DraggedTo = LaserPointerStart + LaserPointerDirection * Hand.DragRayLength;
-						}
-
-						const FVector WorldSpaceDragDelta = DraggedTo - Hand.LastDragToLocation;
-						DragDelta = WorldSpaceDragDelta;
-						Hand.DragTranslationVelocity = WorldSpaceDragDelta;
-
-						const FVector WorldSpaceDeltaFromStart = DraggedTo - Hand.LaserPointerImpactAtDragStart;
-						DragDeltaFromStart = WorldSpaceDeltaFromStart;
-
-						Hand.LastLaserPointerStart = LaserPointerStart;
-						Hand.LastDragToLocation = DraggedTo;
-
-						// Update hover location (we only do this when dragging using the laser pointer)
-						Hand.bIsHovering = true;
-						Hand.HoverLocation = FMath::ClosestPointOnLine( LaserPointerStart, LaserPointerEnd, DraggedTo );
-					}
-					else
-					{
-						// We lost our laser pointer, so cancel the drag
-						StopDragging( HandIndex );
-					}
-
-					if( OtherHand != nullptr )
-					{
-						FVector OtherHandLaserPointerStart, OtherHandLaserPointerEnd;
-						if( Owner.GetLaserPointer( OtherHandIndex, /* Out */ OtherHandLaserPointerStart, /* Out */ OtherHandLaserPointerEnd ) )
-						{
-							const FVector OtherHandLaserPointerDirection = ( OtherHandLaserPointerEnd - OtherHandLaserPointerStart ).GetSafeNormal();
-							OtherHandDraggedTo = OtherHandLaserPointerStart + OtherHandLaserPointerDirection * OtherHand->DragRayLength;
-
-							const FVector OtherHandWorldSpaceDragDelta = OtherHandDraggedTo - OtherHand->LastDragToLocation;
-							OtherHandDragDelta = OtherHandWorldSpaceDragDelta;
-							OtherHand->DragTranslationVelocity = OtherHandWorldSpaceDragDelta;
-
-							const FVector OtherHandWorldSpaceDeltaFromStart = OtherHandDraggedTo - OtherHand->LaserPointerImpactAtDragStart;
-							OtherHandDragDeltaFromStart = OtherHandWorldSpaceDeltaFromStart;
-
-							OtherHand->LastLaserPointerStart = OtherHandLaserPointerStart;
-							OtherHand->LastDragToLocation = OtherHandDraggedTo;
-
-							// Only hover if we're using the laser pointer
-							OtherHand->bIsHovering = true;
-							OtherHand->HoverLocation = OtherHandDraggedTo;
-						}
-						else
-						{
-							// We lost our laser pointer, so cancel the drag assist
-							StopDragging( OtherHandIndex );
-						}
-					}
-				}
-				else if( ensure( Hand.DraggingMode == EVREditorDraggingMode::World ) )
-				{
-					// While we're changing WorldToMetersScale every frame, our room-space hand locations will be scaled as well!  We need to
-					// inverse compensate for this scaling so that we can figure out how much the hands actually moved as if no scale happened.
-					// This only really matters when we're performing world scaling interactively.
-					const FVector RoomSpaceUnscaledHandLocation = ( Hand.RoomSpaceTransform.GetLocation() / WorldToMetersScale ) * LastWorldToMetersScale;
-					const FVector RoomSpaceUnscaledHandDelta = ( RoomSpaceUnscaledHandLocation - Hand.LastRoomSpaceTransform.GetLocation() );
-
-					// Move the world (in room space)
-					DraggedTo = Hand.LastRoomSpaceTransform.GetLocation() + RoomSpaceUnscaledHandDelta;
-
-					Hand.DragTranslationVelocity = RoomSpaceUnscaledHandDelta;
-					DragDelta = RoomSpaceUnscaledHandDelta;
-
-					Hand.LastDragToLocation = DraggedTo;
-
-					// Two handed?
-					if( OtherHand != nullptr )
-					{
-						const FVector OtherHandRoomSpaceUnscaledLocation = ( OtherHand->RoomSpaceTransform.GetLocation() / WorldToMetersScale ) * LastWorldToMetersScale;
-						const FVector OtherHandRoomSpaceUnscaledHandDelta = ( OtherHandRoomSpaceUnscaledLocation - OtherHand->LastRoomSpaceTransform.GetLocation() );
-
-						OtherHandDraggedTo = OtherHand->LastRoomSpaceTransform.GetLocation() + OtherHandRoomSpaceUnscaledHandDelta;
-
-						OtherHand->DragTranslationVelocity = OtherHandRoomSpaceUnscaledHandDelta;
-						OtherHandDragDelta = OtherHandRoomSpaceUnscaledHandDelta;
-
-						OtherHand->LastDragToLocation = OtherHandDraggedTo;
-					}
-				}
-
-				{
-					// If we're not actually in VR, the hands won't be moving at all in room space, so this check prevents inertia from
-					// working.
-					// @todo vreditor: We could do a world space distance test (times world scale factor) when forcing VR mode to get similar (but not quite the same) behavior
-					if( Owner.IsActuallyUsingVR() )
-					{
-						// Don't bother with inertia if we're not moving very fast.  This filters out tiny accidental movements.
-						const FVector RoomSpaceHandDelta = ( Hand.RoomSpaceTransform.GetLocation() - Hand.LastRoomSpaceTransform.GetLocation() );
-						if( RoomSpaceHandDelta.Size() < VREd::MinVelocityForInertia->GetFloat() * WorldScaleFactor )
-						{
-							Hand.DragTranslationVelocity = FVector::ZeroVector;
-						}
-						if( OtherHand != nullptr )
-						{
-							const FVector OtherHandRoomSpaceHandDelta = ( OtherHand->RoomSpaceTransform.GetLocation() - OtherHand->LastRoomSpaceTransform.GetLocation() );
-							if( OtherHandRoomSpaceHandDelta.Size() < VREd::MinVelocityForInertia->GetFloat() * WorldScaleFactor )
-							{
-								OtherHand->DragTranslationVelocity = FVector::ZeroVector;
-							}
-						}
-					}
-				}
-
-				const FVector OldViewLocation = ViewportClient->GetViewLocation();
-				// Dragging transform gizmo handle
-				const bool bWithTwoHands = ( OtherHand != nullptr || OtherHandThatWasAssistingDrag != nullptr );
-				const bool bIsLaserPointerValid = true;
-				FVector UnsnappedDraggedTo = FVector::ZeroVector;
-				UpdateDragging(
-					DeltaTime,
-					Hand.bIsFirstDragUpdate, 
-					Hand.DraggingMode, 
-					Hand.TransformGizmoInteractionType, 
-					bWithTwoHands,
-					*ViewportClient, 
-					Hand.OptionalHandlePlacement, 
-					DragDelta, 
-					OtherHandDragDelta,
-					DraggedTo,
-					OtherHandDraggedTo,
-					DragDeltaFromStart, 
-					OtherHandDragDeltaFromStart,
-					LaserPointerStart, 
-					LaserPointerDirection, 
-					bIsLaserPointerValid, 
-					Hand.GizmoStartTransform, 
-					Hand.GizmoStartLocalBounds,
-					Hand.DraggingTransformGizmoComponent.Get(),
-					/* In/Out */ Hand.GizmoSpaceFirstDragUpdateOffsetAlongAxis,
-					/* In/Out */ Hand.GizmoSpaceDragDeltaFromStartOffset,
-					Hand.bIsDrivingVelocityOfSimulatedTransformables,
-					/* Out */ UnsnappedDraggedTo );
-
-
-				// Make sure the hover point is right on the position that we're dragging the object to.  This is important
-				// when constraining dragged objects to a single axis or a plane
-				if( Hand.DraggingMode == EVREditorDraggingMode::ActorsWithGizmo )
-				{
-					Hand.HoverLocation = FMath::ClosestPointOnSegment( UnsnappedDraggedTo, LaserPointerStart, LaserPointerEnd );
-				}
-
-
-				if( OtherHandThatWasAssistingDrag != nullptr )
-				{
-					OtherHandThatWasAssistingDrag->LastDragToLocation = OtherHandDraggedTo;
-
-					// Apply damping
-					const bool bVelocitySensitive = Hand.DraggingMode == EVREditorDraggingMode::World;
-					Owner.ApplyVelocityDamping( OtherHandThatWasAssistingDrag->DragTranslationVelocity, bVelocitySensitive );
-				}
-
-				// If we were dragging the world, then play some haptic feedback
-				if( Hand.DraggingMode == EVREditorDraggingMode::World )
-				{
-					const FVector NewViewLocation = ViewportClient->GetViewLocation();
-
-					// @todo vreditor: Consider doing this for inertial moves too (we need to remember the last hand that invoked the move.)
-
-					const float RoomSpaceHapticTranslationInterval = 25.0f;		// @todo vreditor tweak: Hard coded haptic distance
-					const float WorldSpaceHapticTranslationInterval = RoomSpaceHapticTranslationInterval * WorldScaleFactor;
-
-					// @todo vreditor: Make this a velocity-sensitive strength?
-					const float ForceFeedbackStrength = VREd::GridHapticFeedbackStrength->GetFloat();		// @todo vreditor tweak: Force feedback strength and enable/disable should be user configurable in options
-
-					bool bCrossedGridLine = false;
-					for( int32 AxisIndex = 0; AxisIndex < 3; ++AxisIndex )
-					{
-						const int32 OldGridCellIndex = FMath::TruncToInt( OldViewLocation[ AxisIndex ] / WorldSpaceHapticTranslationInterval );
-						const int32 NewGridCellIndex = FMath::TruncToInt( NewViewLocation[ AxisIndex ] / WorldSpaceHapticTranslationInterval );
-						if( OldGridCellIndex != NewGridCellIndex )
-						{
-							bCrossedGridLine = true;
-						}
-					}
-
-					if( bCrossedGridLine )
-					{
-						float LeftStrength = 0.0f;
-						float RightStrength = 0.0f;
-						if( HandIndex == VREditorConstants::LeftHandIndex || bWithTwoHands )
-						{
-							LeftStrength = ForceFeedbackStrength;
-						}
-						if( HandIndex == VREditorConstants::RightHandIndex || bWithTwoHands )
-						{
-							RightStrength = ForceFeedbackStrength;
-						}
-						Owner.PlayHapticEffect( LeftStrength, RightStrength );
-					}
-				}
-			}
-			else if (Hand.DraggingMode == EVREditorDraggingMode::DockableWindow)
-			{
-				AVREditorDockableWindow* DraggingUI = Owner.GetUISystem().GetDraggingDockUI();
-				if( DraggingUI )
-				{
-					//@todo: VR Editor - This should be in the UI system
-					const bool bIsAbsolute = ( Owner.GetHMDDeviceType() == EHMDDeviceType::DT_SteamVR );
-					float SlideDeltaY = GetTrackpadSlideDelta( HandIndex );
-
-					float NewUIScale = DraggingUI->GetScale() + SlideDeltaY;
-					if( NewUIScale <= Owner.GetUISystem().GetMinDockWindowSize() )
-					{
-						NewUIScale = Owner.GetUISystem().GetMinDockWindowSize();
-					}
-					else if( NewUIScale >= Owner.GetUISystem().GetMaxDockWindowSize() )
-					{
-						NewUIScale = Owner.GetUISystem().GetMaxDockWindowSize();
-					}
-					DraggingUI->SetScale( NewUIScale );
-
-					const FTransform DockedUIToWorld = Owner.GetUISystem().MakeDockableUITransform( DraggingUI, HandIndex, DockSelectDistance );
-					DraggingUI->SetActorTransform( DockedUIToWorld );
-					
-					DraggingUI->UpdateRelativeRoomTransform();
-				}
-			}
-			// If we're not actively dragging, apply inertia to any selected elements that we've dragged around recently
-			else 
-			{
-				if( !Hand.DragTranslationVelocity.IsNearlyZero( KINDA_SMALL_NUMBER ) &&
-					!Hand.bWasAssistingDrag && 	// If we were only assisting, let the other hand take care of doing the update
-					!Hand.bIsDrivingVelocityOfSimulatedTransformables )	// If simulation mode is on, let the physics engine take care of inertia
-				{
-					const FVector DragDelta = Hand.DragTranslationVelocity;
-					const FVector DraggedTo = Hand.LastDragToLocation + DragDelta;
-					const FVector DragDeltaFromStart = DraggedTo - Hand.LaserPointerImpactAtDragStart;
-
-					// Check to see if the other hand has any inertia to contribute
-					FVirtualHand* OtherHandThatWasAssistingDrag = nullptr;
-					{
-						FVirtualHand& OtherHand = Owner.GetOtherHand( HandIndex );
-
-						// If the other hand isn't doing anything, but the last thing it was doing was assisting a drag, then allow it
-						// to contribute inertia!
-						if( OtherHand.DraggingMode == EVREditorDraggingMode::Nothing && OtherHand.bWasAssistingDrag )
-						{
-							if( !OtherHand.DragTranslationVelocity.IsNearlyZero( KINDA_SMALL_NUMBER ) )
-							{
-								OtherHandThatWasAssistingDrag = &OtherHand;
-							}
-						}
-					}
-					const bool bWithTwoHands = ( OtherHandThatWasAssistingDrag != nullptr );
-
-					FVector OtherHandDragDelta = FVector::ZeroVector;
-					FVector OtherHandDragDeltaFromStart = FVector::ZeroVector;
-					FVector OtherHandDraggedTo = FVector::ZeroVector;
-					if( bWithTwoHands )
-					{
-						OtherHandDragDelta = OtherHandThatWasAssistingDrag->DragTranslationVelocity;
-						OtherHandDraggedTo = OtherHandThatWasAssistingDrag->LastDragToLocation + OtherHandDragDelta;
-						OtherHandDragDeltaFromStart = OtherHandDraggedTo - OtherHandThatWasAssistingDrag->LaserPointerImpactAtDragStart;
-					}
-
-					const bool bIsLaserPointerValid = false;	// Laser pointer has moved on to other things
-					FVector UnsnappedDraggedTo = FVector::ZeroVector;
-					UpdateDragging(
-						DeltaTime,
-						Hand.bIsFirstDragUpdate, 
-						Hand.LastDraggingMode, 
-						Hand.TransformGizmoInteractionType, 
-						bWithTwoHands, 
-						*ViewportClient, 
-						Hand.OptionalHandlePlacement, 
-						DragDelta, 
-						OtherHandDragDelta,
-						DraggedTo,
-						OtherHandDraggedTo,
-						DragDeltaFromStart, 
-						OtherHandDragDeltaFromStart,
-						FVector::ZeroVector,	// No valid laser pointer during inertia
-						FVector::ZeroVector,	// No valid laser pointer during inertia
-						bIsLaserPointerValid, 
-						Hand.GizmoStartTransform, 
-						Hand.GizmoStartLocalBounds, 
-						Hand.DraggingTransformGizmoComponent.Get(),
-						/* In/Out */ Hand.GizmoSpaceFirstDragUpdateOffsetAlongAxis,
-						/* In/Out */ Hand.GizmoSpaceDragDeltaFromStartOffset,
-						Hand.bIsDrivingVelocityOfSimulatedTransformables,
-						/* Out */ UnsnappedDraggedTo );
-
-					Hand.LastDragToLocation = DraggedTo;
-					const bool bVelocitySensitive = Hand.LastDraggingMode == EVREditorDraggingMode::World;
-					Owner.ApplyVelocityDamping( Hand.DragTranslationVelocity, bVelocitySensitive );
-
-					if( OtherHandThatWasAssistingDrag != nullptr )
-					{
-						OtherHandThatWasAssistingDrag->LastDragToLocation = OtherHandDraggedTo;
-						Owner.ApplyVelocityDamping( OtherHandThatWasAssistingDrag->DragTranslationVelocity, bVelocitySensitive );
-					}
-				}
-				else
-				{
-					Hand.DragTranslationVelocity = FVector::ZeroVector;
-				}
-			}
-		}
-	}
-
-
-	// Refresh the transform gizmo every frame, just in case actors were moved by some external
-	// influence.  Also, some features of the transform gizmo respond to camera position (like the
-	// measurement text labels), so it just makes sense to keep it up to date.
-	{
-		const bool bNewObjectsSelected = false;
-		const bool bAllHandlesVisible = ( DraggingActorsWithHandIndex == INDEX_NONE );
-		UActorComponent* SingleVisibleHandle = ( DraggingActorsWithHandIndex != INDEX_NONE ) ? Owner.GetVirtualHand( DraggingActorsWithHandIndex ).DraggingTransformGizmoComponent.Get() : nullptr;
-
-		static TArray< UActorComponent* > HoveringOverHandles;
-		HoveringOverHandles.Reset();
-		for( int32 HandIndex = 0; HandIndex < VREditorConstants::NumVirtualHands; ++HandIndex )
-		{
-			FVirtualHand& Hand = Owner.GetVirtualHand( HandIndex );
-			UActorComponent* HoveringOverHandle = Hand.HoveringOverTransformGizmoComponent.Get();
-			if( HoveringOverHandle != nullptr )
-			{
-				HoveringOverHandles.Add( HoveringOverHandle );
-			}
-		}
-
-		RefreshTransformGizmo( bNewObjectsSelected, bAllHandlesVisible, SingleVisibleHandle, HoveringOverHandles );
-	}
-
-
-	// Update transformables
-	{
-		const bool bSmoothSnappingEnabled = IsSmoothSnappingEnabled();
-		if( bSmoothSnappingEnabled || bIsInterpolatingTransformablesFromSnapshotTransform )
-		{
-			const float SmoothSnapSpeed = VREd::SmoothSnapSpeed->GetFloat();
-			const bool bUseElasticSnapping = bSmoothSnappingEnabled && VREd::ElasticSnap->GetInt() > 0 && TrackingTransaction.IsActive();	// Only while we're still dragging stuff!
-			const float ElasticSnapStrength = VREd::ElasticSnapStrength->GetFloat();
-
-			// NOTE: We never sweep while smooth-snapping as the object will never end up where we want it to
-			// due to friction with the ground and other objects.
-			const bool bSweep = false;
-
-			for( TUniquePtr<FVREditorTransformable>& TransformablePtr : Transformables )
-			{
-				FVREditorTransformable& Transformable = *TransformablePtr;
-
-				FTransform ActualTargetTransform = Transformable.TargetTransform;
-
-				// If 'elastic snapping' is turned on, we'll have the object 'reach' toward its unsnapped position, so
-				// that the user always has visual feedback when they are dragging something around, even if they
-				// haven't dragged far enough to overcome the snap threshold.
-				if( bUseElasticSnapping )
-				{
-					ActualTargetTransform.BlendWith( Transformable.UnsnappedTargetTransform, ElasticSnapStrength );
-				}
-
-				FTransform InterpolatedTransform = ActualTargetTransform;
-				if( !ActualTargetTransform.Equals( Transformable.LastTransform, 0.0f ) )
-				{
-					// If we're really close, just snap it.
-					if( ActualTargetTransform.Equals( Transformable.LastTransform, KINDA_SMALL_NUMBER ) )
-					{
-						InterpolatedTransform = Transformable.LastTransform = Transformable.UnsnappedTargetTransform = ActualTargetTransform;
-					}
-					else
-					{
-						InterpolatedTransform.Blend(
-							Transformable.LastTransform,
-							ActualTargetTransform,
-							FMath::Min( 1.0f, SmoothSnapSpeed * FMath::Min( DeltaTime, 1.0f / 30.0f ) ) );
-					}
-					Transformable.LastTransform = InterpolatedTransform;
-				}
-
-				if( bIsInterpolatingTransformablesFromSnapshotTransform )
-				{
-					const float InterpProgress = FMath::Clamp( (float)( CurrentTime - TransformablesInterpolationStartTime ).GetTotalSeconds() / TransformablesInterpolationDuration, 0.0f, 1.0f );
-					InterpolatedTransform.BlendWith( Transformable.InterpolationSnapshotTransform, 1.0f - InterpProgress );
-					if( InterpProgress >= 1.0f - KINDA_SMALL_NUMBER )
-					{
-						// Finished interpolating
-						bIsInterpolatingTransformablesFromSnapshotTransform = false;
-					}
-				}
-
-				// Got an actor?
-				AActor* Actor = Cast<AActor>( Transformable.Object.Get() );
-				if( Actor != nullptr )
-				{
-					const FTransform& ExistingTransform = Actor->GetActorTransform();
-					if( !ExistingTransform.Equals( InterpolatedTransform, 0.0f ) )
-					{
-						const bool bOnlyTranslationChanged =
-							ExistingTransform.GetRotation() == InterpolatedTransform.GetRotation() &&
-							ExistingTransform.GetScale3D() == InterpolatedTransform.GetScale3D();
-
-						Actor->SetActorTransform( InterpolatedTransform, bSweep );
-						//GWarn->Logf( TEXT( "SMOOTH: Actor %s to %s" ), *Actor->GetName(), *Transformable.TargetTransform.ToString() );
-
-						Actor->InvalidateLightingCacheDetailed( bOnlyTranslationChanged );
-
-						const bool bFinished = false;
-						Actor->PostEditMove( bFinished );
-					}
-				}
-				else
-				{
-					// Some other object that we don't know how to drag
-				}
-			}
-		}
-	}
-
-
-	// Update teleporting
-	if( bIsTeleporting )
-	{
-		Teleport( DeltaTime );
-	}
-
-	LastWorldToMetersScale = WorldToMetersScale;
-}
-
-
-void FVREditorWorldInteraction::Render( const FSceneView* SceneView, FViewport* Viewport, FPrimitiveDrawInterface* PDI )
-{
-}
-
-
-void FVREditorWorldInteraction::OnActorSelectionChanged( UObject* ChangedObject )
-{
-	const bool bNewObjectsSelected = true;
-	const bool bAllHandlesVisible = true;
-	UActorComponent* SingleVisibleHandle = nullptr;
-	TArray<UActorComponent*> HoveringOverHandles;
-	RefreshTransformGizmo( bNewObjectsSelected, bAllHandlesVisible, SingleVisibleHandle, HoveringOverHandles );
-
-	// Clear our transformables.  They'll be doing things relative to the gizmo, which will now be in a different place.
-	// @todo vreditor: We could solve this by moving GizmoLocalBounds into the a Transformables base structure, and making sure
-	// that we never talk to TransformGizmoActor (or selected actors) when updating in tick (currently safe!)
-	Transformables.Reset();
-
-	bDraggedSinceLastSelection = false;
-	LastDragGizmoStartTransform = FTransform::Identity;
-}
-
-
-void FVREditorWorldInteraction::RefreshTransformGizmo( const bool bNewObjectsSelected, bool bAllHandlesVisible, UActorComponent* SingleVisibleHandle, const TArray< UActorComponent* >& HoveringOverHandles )
-{
-	if( GEditor->GetSelectedActorCount() > 0 )
-	{
-		// Check if there no gizmo or if the wrong gizmo is being used
-		bool bSpawnNewGizmo = false;
-		if( TransformGizmoActor == nullptr || 
-			( ( TransformGizmoActor->IsA( APivotTransformGizmo::StaticClass() ) && VREd::PivotPointTransformGizmo->GetInt() == 0 ) || 
-			  ( TransformGizmoActor->IsA( ATransformGizmo::StaticClass() ) && VREd::PivotPointTransformGizmo->GetInt() == 1 ) ) )
-		{
-			bSpawnNewGizmo = true;
-		}
-
-		if ( bSpawnNewGizmo )
-		{
-			// Destroy the previous gizmo
-			if ( TransformGizmoActor != nullptr )
-			{
-				TransformGizmoActor->Destroy();
-			}
-
-			// Create the correct gizmo
-			const bool bWithSceneComponent = false;	// We already have our own scene component
-			if( VREd::PivotPointTransformGizmo->GetInt() == 1 )
-			{
-				TransformGizmoActor = Owner.SpawnTransientSceneActor<APivotTransformGizmo>( TEXT( "PivotTransformGizmo" ), bWithSceneComponent );
-			}
-			else
-			{
-				TransformGizmoActor = Owner.SpawnTransientSceneActor<ATransformGizmo>( TEXT( "TransformGizmo" ), bWithSceneComponent );
-			}
-
-			check( TransformGizmoActor != nullptr );
-			TransformGizmoActor->SetOwnerMode( &Owner ) ;
-
-			if( VREd::ShowTransformGizmo->GetInt() == 0 )
-			{
-				TransformGizmoActor->SetIsTemporarilyHiddenInEditor( true );
-			}
-		}
-
-		const ECoordSystem CurrentCoordSystem = Owner.GetTransformGizmoCoordinateSpace();
-		const bool bIsWorldSpaceGizmo = ( CurrentCoordSystem == COORD_World );
-
-		USelection* ActorSelectionSet = GEditor->GetSelectedActors();
-
-		static TArray<UObject*> SelectedActorObjects;
-		SelectedActorObjects.Reset();
-		ActorSelectionSet->GetSelectedObjects( AActor::StaticClass(), /* Out */ SelectedActorObjects );
-
-		AActor* LastSelectedActor = CastChecked<AActor>( SelectedActorObjects.Last() );
-
-		// Use the location and orientation of the last selected actor, to be consistent with the regular editor's gizmo
-		FTransform GizmoToWorld = LastSelectedActor->GetTransform();
-		GizmoToWorld.RemoveScaling();	// We don't need the pivot itself to be scaled
-
-		if( bIsWorldSpaceGizmo )
-		{
-			GizmoToWorld.SetRotation( FQuat::Identity );
-		}
-
-		// Create a gizmo-local bounds around all of the selected actors
-		FBox GizmoSpaceSelectedActorsBounds;
-		GizmoSpaceSelectedActorsBounds.Init();
-		{
-			const FTransform WorldToGizmo = GizmoToWorld.Inverse();
-			for( UObject* SelectedActorObject : SelectedActorObjects )
-			{
-				AActor* SelectedActor = CastChecked<AActor>( SelectedActorObject );
-
-				// Figure out the local space bounds of the entire actor
-				const bool bIncludeNonCollidingComponents = false;	// @todo vreditor: Disabled this because it causes lights to have huge bounds
-				const FBox ActorSpaceBounds = SelectedActor->CalculateComponentsBoundingBoxInLocalSpace( bIncludeNonCollidingComponents );
-
-				// Transform the bounds into the gizmo's space
-				const FTransform ActorToWorld = SelectedActor->GetTransform();
-				const FTransform ActorToGizmo = ActorToWorld * WorldToGizmo;
-				const FBox GizmoSpaceBounds = ActorSpaceBounds.TransformBy( ActorToGizmo );
-
-				// Get that bounding box into the local space of the gizmo
-				GizmoSpaceSelectedActorsBounds += GizmoSpaceBounds;
-			}
-		}
-
-		if( VREd::ForceGizmoPivotToCenterOfSelectedActorsBounds->GetInt() > 0 )
-		{
-			GizmoToWorld.SetLocation( GizmoToWorld.TransformPosition( GizmoSpaceSelectedActorsBounds.GetCenter() ) );
-			GizmoSpaceSelectedActorsBounds = GizmoSpaceSelectedActorsBounds.ShiftBy( -GizmoSpaceSelectedActorsBounds.GetCenter() );
-		}
-
-		if( bNewObjectsSelected )
-		{
-			TransformGizmoActor->OnNewObjectsSelected();
-		}
-		
-		const EGizmoHandleTypes GizmoType = Owner.GetCurrentGizmoType();
-		const ECoordSystem GizmoCoordinateSpace = Owner.GetTransformGizmoCoordinateSpace();
-
-		GizmoLocalBounds = GizmoSpaceSelectedActorsBounds;
-
-		TransformGizmoActor->UpdateGizmo( GizmoType, GizmoCoordinateSpace, GizmoToWorld, GizmoSpaceSelectedActorsBounds, Owner.GetHeadTransform().GetLocation(), bAllHandlesVisible, SingleVisibleHandle, 
-			HoveringOverHandles, VREd::GizmoHandleHoverScale->GetFloat(), VREd::GizmoHandleHoverAnimationDuration->GetFloat() );
-
-		// Only draw if snapping is turned on
-		AActor* SnapGridActor = Owner.GetSnapGridActor();
-		if( FSnappingUtils::IsSnapToGridEnabled() )
-		{
-			SnapGridActor->GetRootComponent()->SetVisibility( true );
-
-			const float GizmoAnimationAlpha = TransformGizmoActor->GetAnimationAlpha();
-
-			// Position the grid snap actor
-			const FVector GizmoLocalBoundsBottom( 0.0f, 0.0f, GizmoLocalBounds.Min.Z );
-			const float SnapGridZOffset = 0.1f;	// @todo vreditor tweak: Offset the grid position a little bit to avoid z-fighting with objects resting directly on a floor
-			const FVector SnapGridLocation = GizmoToWorld.TransformPosition( GizmoLocalBoundsBottom ) + FVector( 0.0f, 0.0f, SnapGridZOffset );
-			SnapGridActor->SetActorLocationAndRotation( SnapGridLocation, GizmoToWorld.GetRotation() );
-
-			// Figure out how big to make the snap grid.  We'll use a multiplier of the object's bounding box size (ignoring local 
-			// height, because we currently only draw the grid at the bottom of the object.)
-			FBox GizmoLocalBoundsFlattened = GizmoLocalBounds;
-			GizmoLocalBoundsFlattened.Min.Z = GizmoLocalBoundsFlattened.Max.Z = 0.0f;
-			const FBox GizmoWorldBoundsFlattened = GizmoLocalBoundsFlattened.TransformBy( GizmoToWorld );
-			const float SnapGridSize = GizmoWorldBoundsFlattened.GetSize().GetAbsMax() * VREd::SnapGridSize->GetFloat();
-
-			// The mesh is 100x100, so we'll scale appropriately
-			const float SnapGridScale = SnapGridSize / 100.0f;
-			SnapGridActor->SetActorScale3D( FVector( SnapGridScale ) );
-
-			EVREditorDraggingMode DraggingMode = EVREditorDraggingMode::Nothing;
-			FVector GizmoStartLocationWhileDragging = FVector::ZeroVector;
-			for( int32 HandIndex = 0; HandIndex < VREditorConstants::NumVirtualHands; ++HandIndex )
-			{
-				const FVirtualHand& Hand = Owner.GetVirtualHand( HandIndex );
-				if( Hand.DraggingMode == EVREditorDraggingMode::ActorsFreely ||
-					Hand.DraggingMode == EVREditorDraggingMode::ActorsWithGizmo ||
-					Hand.DraggingMode == EVREditorDraggingMode::ActorsAtLaserImpact )
-				{
-					DraggingMode = Hand.DraggingMode;
-					GizmoStartLocationWhileDragging = Hand.GizmoStartTransform.GetLocation();
-					break;
-				}
-				else if( bDraggedSinceLastSelection && 
-					( Hand.LastDraggingMode == EVREditorDraggingMode::ActorsFreely ||
-					  Hand.LastDraggingMode == EVREditorDraggingMode::ActorsWithGizmo ||
-					  Hand.LastDraggingMode == EVREditorDraggingMode::ActorsAtLaserImpact ) )
-				{
-					DraggingMode = Hand.LastDraggingMode;
-					GizmoStartLocationWhileDragging = LastDragGizmoStartTransform.GetLocation();
-				}
-			}
-
-			FVector SnapGridCenter = SnapGridLocation;
-			if( DraggingMode != EVREditorDraggingMode::Nothing )
-			{
-				SnapGridCenter = GizmoStartLocationWhileDragging;
-			}
-
-			// Fade the grid in with the transform gizmo
-			const float GridAlpha = GizmoAnimationAlpha;
-
-			// Animate the grid a little bit
-			const FLinearColor GridColor = FLinearColor::LerpUsingHSV( 
-				FLinearColor::White, 
-				FLinearColor::Yellow, 
-				FMath::MakePulsatingValue( Owner.GetTimeSinceModeEntered().GetTotalSeconds(), 0.5f ) ).CopyWithNewOpacity( GridAlpha );
-
-			const float GridInterval = GEditor->GetGridSize();
-
-			// If the grid size is very small, shrink the size of our lines so that they don't overlap
-			float LineWidth = VREd::SnapGridLineWidth->GetFloat();
- 			while( GridInterval < LineWidth * 3.0f )	// @todo vreditor tweak
- 			{
- 				LineWidth *= 0.5f;
- 			}
-
-			{
-				// Update snap grid material state
-				UMaterialInstanceDynamic* SnapGridMID = Owner.GetSnapGridMID();
-
-				static FName GridColorParameterName( "GridColor" );
-				SnapGridMID->SetVectorParameterValue( GridColorParameterName, GridColor );
-
-				static FName GridCenterParameterName( "GridCenter" );
-				SnapGridMID->SetVectorParameterValue( GridCenterParameterName, SnapGridCenter );
-
-				static FName GridIntervalParameterName( "GridInterval" );
-				SnapGridMID->SetScalarParameterValue( GridIntervalParameterName, GridInterval );
-
-				static FName GridRadiusParameterName( "GridRadius" );
-				SnapGridMID->SetScalarParameterValue( GridRadiusParameterName, SnapGridSize * 0.5f );
-
-				static FName LineWidthParameterName( "LineWidth" );
-				SnapGridMID->SetScalarParameterValue( LineWidthParameterName, LineWidth );
-
-				FMatrix GridRotationMatrix = GizmoToWorld.ToMatrixNoScale().Inverse();
-				GridRotationMatrix.SetOrigin( FVector::ZeroVector );
-
-				static FName GridRotationMatrixXParameterName( "GridRotationMatrixX" );
-				static FName GridRotationMatrixYParameterName( "GridRotationMatrixY" );
-				static FName GridRotationMatrixZParameterName( "GridRotationMatrixZ" );
-				SnapGridMID->SetVectorParameterValue( GridRotationMatrixXParameterName, GridRotationMatrix.GetScaledAxis( EAxis::X ) );
-				SnapGridMID->SetVectorParameterValue( GridRotationMatrixYParameterName, GridRotationMatrix.GetScaledAxis( EAxis::Y ) );
-				SnapGridMID->SetVectorParameterValue( GridRotationMatrixZParameterName, GridRotationMatrix.GetScaledAxis( EAxis::Z ) );
-			}
-		}
-		else
-		{
-			// Grid snap not enabled
-			SnapGridActor->GetRootComponent()->SetVisibility( false );
-		}
-	}
-	else
-	{
-		// Nothing selected, so kill our gizmo actor
-		if( TransformGizmoActor != nullptr )
-		{
-			Owner.DestroyTransientActor( TransformGizmoActor );
-			TransformGizmoActor = nullptr;
-		}
-		GizmoLocalBounds = FBox( 0 );
-
-		// Hide the snap actor
-		AActor* SnapGridActor = Owner.GetSnapGridActor();
-		SnapGridActor->GetRootComponent()->SetVisibility( false );
-	}
-}
-
-
-bool FVREditorWorldInteraction::IsInteractableComponent( const UActorComponent* Component ) const
-{
-	// Don't interact with frozen actors
-	if( Component != nullptr )
-	{
-		static const bool bIsVREditorDemo = FParse::Param( FCommandLine::Get(), TEXT( "VREditorDemo" ) );	// @todo vreditor: Remove this when no longer needed (console variable, too!)
-		const bool bIsFrozen = bIsVREditorDemo && Component->GetOwner()->GetActorLabel().StartsWith( TEXT( "Frozen_" ) );
-		return
-			!bIsFrozen &&
-			( !Owner.GetUISystem().IsWidgetAnEditorUIWidget( Component ) ); // Don't allow user to move around our UI widgets
-	}
-
-	return false;
-}
-
-
-FVector FVREditorWorldInteraction::ComputeConstrainedDragDeltaFromStart( const bool bIsFirstDragUpdate, const ETransformGizmoInteractionType InteractionType, const TOptional<FTransformGizmoHandlePlacement> OptionalHandlePlacement, const FVector& DragDeltaFromStart, const FVector& LaserPointerStart, const FVector& LaserPointerDirection, const bool bIsLaserPointerValid, const FTransform& GizmoStartTransform, FVector& GizmoSpaceFirstDragUpdateOffsetAlongAxis, FVector& GizmoSpaceDragDeltaFromStartOffset ) const
-{
-	FVector ConstrainedGizmoSpaceDeltaFromStart;
-
-	bool bConstrainDirectlyToLineOrPlane = false;
-	if( OptionalHandlePlacement.IsSet() )
-	{
-		// Constrain to line/plane
-		const FTransformGizmoHandlePlacement& HandlePlacement = OptionalHandlePlacement.GetValue();
-
-		int32 CenterHandleCount, FacingAxisIndex, CenterAxisIndex;
-		HandlePlacement.GetCenterHandleCountAndFacingAxisIndex( /* Out */ CenterHandleCount, /* Out */ FacingAxisIndex, /* Out */ CenterAxisIndex );
-
-		// Our laser pointer ray won't be valid for inertial transform, since it's already moved on to other things.  But the existing velocity should already be on axis.
-		bConstrainDirectlyToLineOrPlane = ( CenterHandleCount == 2 ) && bIsLaserPointerValid;
-		if( bConstrainDirectlyToLineOrPlane )
-		{
-			const FVector GizmoSpaceLaserPointerStart = GizmoStartTransform.InverseTransformPosition( LaserPointerStart );
-			const FVector GizmoSpaceLaserPointerDirection = GizmoStartTransform.InverseTransformVectorNoScale( LaserPointerDirection );
-
-			const float LaserPointerMaxLength = Owner.GetLaserPointerMaxLength();
-
-			FVector GizmoSpaceConstraintAxis =
-				FacingAxisIndex == 0 ? FVector::ForwardVector : ( FacingAxisIndex == 1 ? FVector::RightVector : FVector::UpVector );
-			if( HandlePlacement.Axes[ FacingAxisIndex ] == ETransformGizmoHandleDirection::Negative )
-			{
-				GizmoSpaceConstraintAxis *= -1.0f;
-			}
-
-			const bool bOnPlane = ( InteractionType == ETransformGizmoInteractionType::TranslateOnPlane );
-			if( bOnPlane )
-			{
-				const FPlane GizmoSpaceConstrainToPlane( FVector::ZeroVector, GizmoSpaceConstraintAxis );
-
-				// 2D Plane
-				FVector GizmoSpacePointOnPlane = FVector::ZeroVector;
-				if( !FMath::SegmentPlaneIntersection(
-					GizmoSpaceLaserPointerStart,
-					GizmoSpaceLaserPointerStart + LaserPointerMaxLength * GizmoSpaceLaserPointerDirection,
-					GizmoSpaceConstrainToPlane,
-					/* Out */ GizmoSpacePointOnPlane ) )
-				{
-					// No intersect.  Default to no delta.
-					GizmoSpacePointOnPlane = FVector::ZeroVector;
-				}
-
-				// Gizmo space is defined as the starting position of the interaction, so we simply take the closest position
-				// along the plane as our delta from the start in gizmo space
-				ConstrainedGizmoSpaceDeltaFromStart = GizmoSpacePointOnPlane;
-			}
-			else
-			{
-				// 1D Line
-				const float AxisSegmentLength = LaserPointerMaxLength * 2.0f;	// @todo vreditor: We're creating a line segment to represent an infinitely long axis, but really it just needs to be further than our laser pointer can reach
-
-				DVector GizmoSpaceClosestPointOnRayDouble, GizmoSpaceClosestPointOnAxisDouble;
-				SegmentDistToSegmentDouble(
-					GizmoSpaceLaserPointerStart, DVector( GizmoSpaceLaserPointerStart ) + DVector( GizmoSpaceLaserPointerDirection ) * LaserPointerMaxLength,	// @todo vreditor: Should store laser pointer length rather than hard code
-					DVector( GizmoSpaceConstraintAxis ) * -AxisSegmentLength, DVector( GizmoSpaceConstraintAxis ) * AxisSegmentLength,
-					/* Out */ GizmoSpaceClosestPointOnRayDouble,
-					/* Out */ GizmoSpaceClosestPointOnAxisDouble );
-
-				// Gizmo space is defined as the starting position of the interaction, so we simply take the closest position
-				// along the axis as our delta from the start in gizmo space
-				ConstrainedGizmoSpaceDeltaFromStart = GizmoSpaceClosestPointOnAxisDouble.ToFVector();
-			}
-
-			// Account for the handle position on the outside of the bounds.  This is a bit confusing but very important for dragging
-			// to feel right.  When constraining movement to either a plane or single axis, we always want the object to move exactly 
-			// to the location under the laser/cursor -- it's an absolute movement.  But if we did that on the first update frame, it 
-			// would teleport from underneath the gizmo handle to that location. 
-			{
-				if( bIsFirstDragUpdate )
-				{
-					GizmoSpaceFirstDragUpdateOffsetAlongAxis = ConstrainedGizmoSpaceDeltaFromStart;
-				}
-				ConstrainedGizmoSpaceDeltaFromStart -= GizmoSpaceFirstDragUpdateOffsetAlongAxis;
-
-				// OK, it gets even more complicated.  When dragging and releasing for inertial movement, this code path
-				// will no longer be executed (because the hand/laser has moved on to doing other things, and our drag
-				// is driven by velocity only).  So we need to remember the LAST offset from the object's position to
-				// where we actually constrained it to, and continue to apply that delta during the inertial drag.
-				// That actually happens in the code block near the bottom of this function.
-				const FVector GizmoSpaceDragDeltaFromStart = GizmoStartTransform.InverseTransformVectorNoScale( DragDeltaFromStart );
-				GizmoSpaceDragDeltaFromStartOffset = ConstrainedGizmoSpaceDeltaFromStart - GizmoSpaceDragDeltaFromStart;
-			}
-		}
-	}
-
-	if( !bConstrainDirectlyToLineOrPlane )
-	{
-		ConstrainedGizmoSpaceDeltaFromStart = GizmoStartTransform.InverseTransformVectorNoScale( DragDeltaFromStart );
-
-		// Apply axis constraint if we have one (and we haven't already constrained directly to a line)
-		if( OptionalHandlePlacement.IsSet() )
-		{
-			// OK in this case, inertia is moving our selected objects while they remain constrained to
-			// either a plane or a single axis.  Every frame, we need to apply the original delta between
-			// where the laser was aiming and where the object was constrained to on the LAST frame before
-			// the user released the object and it moved inertially.  See the big comment up above for 
-			// more information.  Inertial movement of constrained objects is actually very complicated!
-			ConstrainedGizmoSpaceDeltaFromStart += GizmoSpaceDragDeltaFromStartOffset;
-
-			const bool bOnPlane = ( InteractionType == ETransformGizmoInteractionType::TranslateOnPlane );
-
-			const FTransformGizmoHandlePlacement& HandlePlacement = OptionalHandlePlacement.GetValue();
-			for( int32 AxisIndex = 0; AxisIndex < 3; ++AxisIndex )
-			{
-				if( bOnPlane )
-				{
-					if( HandlePlacement.Axes[ AxisIndex ] == ETransformGizmoHandleDirection::Positive )
-					{
-						// Lock the opposing axes out (plane translation)
-						ConstrainedGizmoSpaceDeltaFromStart[ AxisIndex ] = 0.0f;
-					}
-				}
-				else
-				{
-					if( HandlePlacement.Axes[ AxisIndex ] == ETransformGizmoHandleDirection::Center )
-					{
-						// Lock the centered axis out (line translation)
-						ConstrainedGizmoSpaceDeltaFromStart[ AxisIndex ] = 0.0f;
-					}
-				}
-			}
-		}
-	}
-
-	FVector ConstrainedWorldSpaceDeltaFromStart = GizmoStartTransform.TransformVectorNoScale( ConstrainedGizmoSpaceDeltaFromStart );
-	return ConstrainedWorldSpaceDeltaFromStart;
-}
-
-
-void FVREditorWorldInteraction::UpdateDragging(
-	const float DeltaTime,
-	bool& bIsFirstDragUpdate,
-	const EVREditorDraggingMode DraggingMode,
-	const ETransformGizmoInteractionType InteractionType,
-	const bool bWithTwoHands,
-	FEditorViewportClient& ViewportClient,
-	const TOptional<FTransformGizmoHandlePlacement> OptionalHandlePlacement,
-	const FVector& DragDelta,
-	const FVector& OtherHandDragDelta,
-	const FVector& DraggedTo,
-	const FVector& OtherHandDraggedTo,
-	const FVector& DragDeltaFromStart,
-	const FVector& OtherHandDragDeltaFromStart,
-	const FVector& LaserPointerStart,
-	const FVector& LaserPointerDirection,
-	const bool bIsLaserPointerValid,
-	const FTransform& GizmoStartTransform,
-	const FBox& GizmoStartLocalBounds,
-	const USceneComponent* const DraggingTransformGizmoComponent,
-	FVector& GizmoSpaceFirstDragUpdateOffsetAlongAxis,
-	FVector& DragDeltaFromStartOffset,
-	bool& bIsDrivingVelocityOfSimulatedTransformables,
-	FVector& OutUnsnappedDraggedTo )
-{
-	// IMPORTANT NOTE: Depending on the DraggingMode, the incoming coordinates can be in different coordinate spaces.
-	//		-> When dragging objects around, everything is in world space unless otherwise specified in the variable name.
-	//		-> When dragging the world around, everything is in room space unless otherwise specified.
-
-	bool bMovedAnyTransformables = false;
-
-	// This will be set to true if we want to set the physics velocity on the dragged objects based on how far
-	// we dragged this frame.  Used for pushing objects around in Simulate mode.  If this is set to false, we'll
-	// simply zero out the velocities instead to cancel out the physics engine's applied forces (e.g. gravity) 
-	// while we're dragging objects around (otherwise, the objects will spaz out as soon as dragging stops.)
-	bool bShouldApplyVelocitiesFromDrag = false;
-
-	// Always snap objects relative to where they were when we first grabbed them.  We never want objects to immediately
-	// teleport to the closest snap, but instead snaps should always be relative to the gizmo center
-	// NOTE: While placing objects, we always snap in world space, since the initial position isn't really at all useful
-	const bool bLocalSpaceSnapping =
-		Owner.GetTransformGizmoCoordinateSpace() == COORD_Local &&
-		DraggingMode != EVREditorDraggingMode::ActorsAtLaserImpact;
-	const FVector SnapGridBase = ( DraggingMode == EVREditorDraggingMode::ActorsAtLaserImpact || bLocalSpaceSnapping ) ? FVector::ZeroVector : GizmoStartTransform.GetLocation();
-
-
-	// Okay, time to move stuff!  We'll do this part differently depending on whether we're dragging actual actors around
-	// or we're moving the camera (aka. dragging the world)
-	if( DraggingMode == EVREditorDraggingMode::ActorsWithGizmo ||
-		DraggingMode == EVREditorDraggingMode::ActorsAtLaserImpact )
-	{
-		FVector ConstrainedDragDeltaFromStart = ComputeConstrainedDragDeltaFromStart(
-			bIsFirstDragUpdate,
-			InteractionType,
-			OptionalHandlePlacement,
-			DragDeltaFromStart,
-			LaserPointerStart,
-			LaserPointerDirection,
-			bIsLaserPointerValid,
-			GizmoStartTransform,
-			/* In/Out */ GizmoSpaceFirstDragUpdateOffsetAlongAxis,
-			/* In/Out */ DragDeltaFromStartOffset );
-
-		// @todo vreditor debug
-		if( false )
-		{
-			GWarn->Logf( TEXT( "bIsLaserPointerValid=%s  DragDelta=%s  DraggedTo=%s  DragDeltaFromStart=%s  ConstrainedDragDeltaFromStart=%s  Offset=%s  DragOffset=%s  LastDragToLocation=%s" ),
-				bIsLaserPointerValid ? TEXT( "1" ) : TEXT( "0" ), 
-				*DragDelta.ToString(),
-				*DraggedTo.ToString(),
-				*DragDeltaFromStart.ToString(), 
-				*ConstrainedDragDeltaFromStart.ToString(),
-				*Owner.GetVirtualHand(0).GizmoSpaceFirstDragUpdateOffsetAlongAxis.ToString(),
-				*DragDeltaFromStartOffset.ToString(),
-				*Owner.GetVirtualHand(0).LastDragToLocation.ToString());
-		}
-
-
-		FVector UnsnappedDraggedTo = GizmoStartTransform.GetLocation() + ConstrainedDragDeltaFromStart;
-		OutUnsnappedDraggedTo = UnsnappedDraggedTo;
-
-		// Grid snap!
-		FVector SnappedDraggedTo = UnsnappedDraggedTo;
-		if( FSnappingUtils::IsSnapToGridEnabled() )
-		{
-			// Snap in local space, if we need to
-			if( bLocalSpaceSnapping )
-			{
-				SnappedDraggedTo = GizmoStartTransform.InverseTransformPositionNoScale( SnappedDraggedTo );
-				FSnappingUtils::SnapPointToGrid( SnappedDraggedTo, SnapGridBase );
-				SnappedDraggedTo = GizmoStartTransform.TransformPositionNoScale( SnappedDraggedTo );
-			}
-			else
-			{
-				FSnappingUtils::SnapPointToGrid( SnappedDraggedTo, SnapGridBase );
-			}
-		}
-
-		// Two passes.  First update the real transform.  Then update the unsnapped transform.
-		for( int32 PassIndex = 0; PassIndex < 2; ++PassIndex )
-		{
-			const bool bIsUpdatingUnsnappedTarget = ( PassIndex == 1 );
-			const FVector& PassDraggedTo = bIsUpdatingUnsnappedTarget ? UnsnappedDraggedTo : SnappedDraggedTo;
-
-			if( InteractionType == ETransformGizmoInteractionType::Translate ||
-				InteractionType == ETransformGizmoInteractionType::TranslateOnPlane )
-			{
-				// Translate the objects!
-				for( TUniquePtr<FVREditorTransformable>& TransformablePtr : Transformables )
-				{
-					FVREditorTransformable& Transformable = *TransformablePtr;
-					FTransform& PassTargetTransform = bIsUpdatingUnsnappedTarget ? Transformable.UnsnappedTargetTransform : Transformable.TargetTransform;
-
-					const FVector OldLocation = PassTargetTransform.GetLocation();
-					PassTargetTransform.SetLocation( ( Transformable.StartTransform.GetLocation() - GizmoStartTransform.GetLocation() ) + PassDraggedTo );
-
-					bMovedAnyTransformables = true;
-					bShouldApplyVelocitiesFromDrag = true;
-				}
-			}
-			else if( InteractionType == ETransformGizmoInteractionType::StretchAndReposition )
-			{
-				// We only support stretching by a handle currently
-				const FTransformGizmoHandlePlacement& HandlePlacement = OptionalHandlePlacement.GetValue();
-
-				const FVector PassGizmoSpaceDraggedTo = GizmoStartTransform.InverseTransformPositionNoScale( PassDraggedTo );
-
-				FBox NewGizmoLocalBounds = GizmoStartLocalBounds;
-				FVector GizmoSpacePivotLocation = FVector::ZeroVector;
-				for( int32 AxisIndex = 0; AxisIndex < 3; ++AxisIndex )
-				{
-					// Figure out how much the gizmo bounds changes
-					if( HandlePlacement.Axes[ AxisIndex ] == ETransformGizmoHandleDirection::Negative )	// Negative direction
-					{
-						GizmoSpacePivotLocation[ AxisIndex ] = GizmoStartLocalBounds.Max[ AxisIndex ];
-						NewGizmoLocalBounds.Min[ AxisIndex ] = GizmoStartLocalBounds.Min[ AxisIndex ] + PassGizmoSpaceDraggedTo[ AxisIndex ];
-
-					}
-					else if( HandlePlacement.Axes[ AxisIndex ] == ETransformGizmoHandleDirection::Positive )	// Positive direction
-					{
-						GizmoSpacePivotLocation[ AxisIndex ] = GizmoStartLocalBounds.Min[ AxisIndex ];
-						NewGizmoLocalBounds.Max[ AxisIndex ] = GizmoStartLocalBounds.Max[ AxisIndex ] + PassGizmoSpaceDraggedTo[ AxisIndex ];
-					}
-					else
-					{
-						// Must be ETransformGizmoHandleDirection::Center.  
-						GizmoSpacePivotLocation[ AxisIndex ] = GizmoStartLocalBounds.GetCenter()[ AxisIndex ];
-					}
-				}
-
-				const FVector GizmoStartLocalSize = GizmoStartLocalBounds.GetSize();
-				const FVector NewGizmoLocalSize = NewGizmoLocalBounds.GetSize();
-
-				FVector NewGizmoLocalScaleFromStart = FVector( 1.0f );
-				if( !FMath::IsNearlyZero( NewGizmoLocalSize.GetAbsMin() ) )
-				{
-					NewGizmoLocalScaleFromStart = NewGizmoLocalSize / GizmoStartLocalSize;
-				}
-				else
-				{
-					// Zero scale.  This is allowed in Unreal, for better or worse.
-					NewGizmoLocalScaleFromStart = FVector( 0 );
-				}
-
-				// Stretch and reposition the objects!
-				for( TUniquePtr<FVREditorTransformable>& TransformablePtr : Transformables )
-				{
-					FVREditorTransformable& Transformable = *TransformablePtr;
-
-					FTransform& PassTargetTransform = bIsUpdatingUnsnappedTarget ? Transformable.UnsnappedTargetTransform : Transformable.TargetTransform;
-
-					{
-						const FVector GizmoSpaceTransformableStartLocation = GizmoStartTransform.InverseTransformPositionNoScale( Transformable.StartTransform.GetLocation() );
-						const FVector NewGizmoSpaceLocation = ( GizmoSpaceTransformableStartLocation - GizmoSpacePivotLocation ) * NewGizmoLocalScaleFromStart + GizmoSpacePivotLocation;
-						const FVector OldLocation = PassTargetTransform.GetLocation();
-						PassTargetTransform.SetLocation( GizmoStartTransform.TransformPosition( NewGizmoSpaceLocation ) );
-					}
-
-					// @todo vreditor: This scale is still in gizmo space, but we're setting it in world space
-					PassTargetTransform.SetScale3D( Transformable.StartTransform.GetScale3D() * NewGizmoLocalScaleFromStart );
-
-					bMovedAnyTransformables = true;
-					bShouldApplyVelocitiesFromDrag = false;
-				}
-			}
-			else if ( InteractionType == ETransformGizmoInteractionType::Scale || InteractionType == ETransformGizmoInteractionType::UniformScale )
-			{
-				const FTransformGizmoHandlePlacement& HandlePlacement = OptionalHandlePlacement.GetValue();
-
-				int32 CenterHandleCount, FacingAxisIndex, CenterAxisIndex;
-				HandlePlacement.GetCenterHandleCountAndFacingAxisIndex( /* Out */ CenterHandleCount, /* Out */ FacingAxisIndex, /* Out */ CenterAxisIndex);
-
-				const FVector PassGizmoSpaceDraggedTo = GizmoStartTransform.InverseTransformPositionNoScale( PassDraggedTo );
-
-				float AddedScaleOnAxis = 0.0f;
-				if ( InteractionType == ETransformGizmoInteractionType::Scale )
-				{
-					AddedScaleOnAxis = PassGizmoSpaceDraggedTo[ FacingAxisIndex ] * VREd::ScaleSensitivity->GetFloat();
-
-					// Invert if we we are scaling on the negative side of the gizmo
-					if ( DraggingTransformGizmoComponent && DraggingTransformGizmoComponent->GetRelativeTransform().GetLocation()[ FacingAxisIndex ] < 0 )
-					{
-						AddedScaleOnAxis *= -1;
-					}
-				}
-				else if ( InteractionType == ETransformGizmoInteractionType::UniformScale )
-				{
-					// Always use Z for uniform scale
-					const FVector RelativeDraggedTo = PassDraggedTo - GizmoStartTransform.GetLocation();
-					AddedScaleOnAxis = RelativeDraggedTo[ 2 ] * VREd::ScaleSensitivity->GetFloat();
-				}
-
-				// Scale the objects!
-				for ( TUniquePtr<FVREditorTransformable>& TransformablePtr : Transformables )
-				{
-					FVREditorTransformable& Transformable = *TransformablePtr;
-					FTransform& PassTargetTransform = bIsUpdatingUnsnappedTarget ? Transformable.UnsnappedTargetTransform : Transformable.TargetTransform;
-					FVector NewTotalScale = Transformable.StartTransform.GetScale3D();
-
-					if ( InteractionType == ETransformGizmoInteractionType::Scale )
-					{
-						NewTotalScale[ FacingAxisIndex ] += AddedScaleOnAxis;
-					}
-					else if ( InteractionType == ETransformGizmoInteractionType::UniformScale )
-					{
-						NewTotalScale += FVector( AddedScaleOnAxis );
-					}
-
-					// Scale snap!
-					if ( !bIsUpdatingUnsnappedTarget && FSnappingUtils::IsScaleSnapEnabled() )
-					{
-						FSnappingUtils::SnapScale( NewTotalScale, FVector::ZeroVector );
-					}
-
-					PassTargetTransform.SetScale3D( NewTotalScale );
-					bMovedAnyTransformables = true;
-					bShouldApplyVelocitiesFromDrag = true;
-
-				}
-			}
-			else if( InteractionType == ETransformGizmoInteractionType::Rotate )
-			{
-				// We only support rotating by a handle currently
-				const FTransformGizmoHandlePlacement& HandlePlacement = OptionalHandlePlacement.GetValue();
-
-				int32 CenterHandleCount, FacingAxisIndex, CenterAxisIndex;
-				HandlePlacement.GetCenterHandleCountAndFacingAxisIndex( /* Out */ CenterHandleCount, /* Out */ FacingAxisIndex, /* Out */ CenterAxisIndex );
-				check( CenterAxisIndex != INDEX_NONE );
-
-				// Get the local location of the rotation handle
-				FVector HandleRelativeLocation = FVector::ZeroVector;
-				for( int32 AxisIndex = 0; AxisIndex < 3; ++AxisIndex )
-				{
-					if( HandlePlacement.Axes[ AxisIndex ] == ETransformGizmoHandleDirection::Negative )	// Negative direction
-					{
-						HandleRelativeLocation[ AxisIndex ] = GizmoStartLocalBounds.Min[ AxisIndex ];
-					}
-					else if( HandlePlacement.Axes[ AxisIndex ] == ETransformGizmoHandleDirection::Positive )	// Positive direction
-					{
-						HandleRelativeLocation[ AxisIndex ] = GizmoStartLocalBounds.Max[ AxisIndex ];
-					}
-					else // ETransformGizmoHandleDirection::Center
-					{
-						HandleRelativeLocation[ AxisIndex ] = GizmoStartLocalBounds.GetCenter()[ AxisIndex ];
-					}
-				}
-
-				const FVector GizmoSpaceRotationAxis = ( CenterAxisIndex == 0 ) ? FVector::ForwardVector : ( CenterAxisIndex == 1 ? FVector::RightVector : FVector::UpVector );
-
-				// Make a vector that points along the tangent vector of the bounding box edge
-				const FVector GizmoSpaceTowardHandleVector = ( HandleRelativeLocation - GizmoStartLocalBounds.GetCenter() ).GetSafeNormal();
-				const FVector GizmoSpaceDiagonalAxis = FQuat( GizmoSpaceRotationAxis, FMath::DegreesToRadians( 90.0f ) ).RotateVector( GizmoSpaceTowardHandleVector );
-
-				// Figure out how far we've dragged in the direction of the diagonal axis.
-				const float RotationProgress = FVector::DotProduct( DragDeltaFromStart, GizmoStartTransform.TransformVectorNoScale( GizmoSpaceDiagonalAxis ) );
-
-				const float RotationDegreesAroundAxis = RotationProgress * VREd::GizmoRotationSensitivity->GetFloat();
-				const FQuat RotationDeltaFromStart = FQuat( GizmoSpaceRotationAxis, FMath::DegreesToRadians( RotationDegreesAroundAxis ) );
-
-				const FTransform WorldToGizmo = GizmoStartTransform.Inverse();
-
-				// Rotate (and reposition) the objects!
-				for( TUniquePtr<FVREditorTransformable>& TransformablePtr : Transformables )
-				{
-					FVREditorTransformable& Transformable = *TransformablePtr;
-					FTransform& PassTargetTransform = bIsUpdatingUnsnappedTarget ? Transformable.UnsnappedTargetTransform : Transformable.TargetTransform;
-
-					const FTransform GizmoSpaceStartTransform = Transformable.StartTransform * WorldToGizmo;
-					FTransform GizmoSpaceRotatedTransform = GizmoSpaceStartTransform * RotationDeltaFromStart;
-
-					// Snap rotation in gizmo space
-					if( !bIsUpdatingUnsnappedTarget )
-					{
-						FRotator SnappedRotation = GizmoSpaceRotatedTransform.GetRotation().Rotator();
-						FSnappingUtils::SnapRotatorToGrid( SnappedRotation );
-						GizmoSpaceRotatedTransform.SetRotation( SnappedRotation.Quaternion() );
-					}
-
-					const FTransform WorldSpaceRotatedTransform = GizmoSpaceRotatedTransform * GizmoStartTransform;
-
-					PassTargetTransform = WorldSpaceRotatedTransform;
-					bMovedAnyTransformables = true;
-					bShouldApplyVelocitiesFromDrag = false;
-
-					// Snap location in world space
-					if( !bIsUpdatingUnsnappedTarget )
-					{
-						FVector SnappedLocation = PassTargetTransform.GetLocation();
-						
-						// Snap in local space, if we need to
-						if( bLocalSpaceSnapping )
-						{
-							SnappedLocation = GizmoStartTransform.InverseTransformPositionNoScale( SnappedLocation );
-							FSnappingUtils::SnapPointToGrid( SnappedLocation, SnapGridBase );
-							SnappedLocation = GizmoStartTransform.TransformPositionNoScale( SnappedLocation );
-						}
-						else
-						{
-							FSnappingUtils::SnapPointToGrid( SnappedLocation, SnapGridBase );
-						}
-						PassTargetTransform.SetLocation( SnappedLocation );
-					}
-				}
-			}
-			else if( InteractionType == ETransformGizmoInteractionType::RotateOnAngle )
-			{			
-				const FTransformGizmoHandlePlacement& HandlePlacement = OptionalHandlePlacement.GetValue();
-
-				int32 CenterHandleCount, FacingAxisIndex, CenterAxisIndex;
-				HandlePlacement.GetCenterHandleCountAndFacingAxisIndex( /* Out */ CenterHandleCount, /* Out */ FacingAxisIndex, /* Out */ CenterAxisIndex);
-		
-				FVector GizmoSpaceFacingAxisVector = UVREditorGizmoHandleGroup::GetAxisVector( FacingAxisIndex, HandlePlacement.Axes[ FacingAxisIndex ] );
-				
-				if ( DraggingTransformGizmoComponent )
-				{
-					const FVector RotationHandlePlaneLocation = TransformGizmoActor->GetActorLocation();		
-					FVector RotationHandlePlaneNormal = DraggingTransformGizmoComponent->GetComponentTransform().GetRotation().Vector();
-					RotationHandlePlaneNormal.Normalize();
-
-					// Get the laserpointer intersection on the plane of the handle
-					FPlane Plane = FPlane( RotationHandlePlaneLocation, RotationHandlePlaneNormal );
-					const FVector IntersectWorldLocationOnPlane = FMath::LinePlaneIntersection( LaserPointerStart, LaserPointerStart + LaserPointerDirection, RotationHandlePlaneLocation, RotationHandlePlaneNormal );
-					
-					const FTransform WorldToGizmo = GizmoStartTransform.Inverse();
-					
-					DrawDebugLine( Owner.GetWorld(), TransformGizmoActor->GetActorLocation(), IntersectWorldLocationOnPlane, FColor::Red, false, -1.f, 0, 3.0f * Owner.GetWorldScaleFactor() );
-
-					// Rotate the objects!
-					for ( TUniquePtr<FVREditorTransformable>& TransformablePtr : Transformables )
-					{
-						FVREditorTransformable& Transformable = *TransformablePtr;
-						FTransform& PassTargetTransform = bIsUpdatingUnsnappedTarget ? Transformable.UnsnappedTargetTransform : Transformable.TargetTransform;
-
-						// The starting transformation of this transformable in gizmo space
-						const FTransform GizmoSpaceStartTransform = Transformable.StartTransform * WorldToGizmo;
-						const FVector TransformableSpaceFacingAxisVector = Transformable.StartTransform.InverseTransformVectorNoScale( GizmoStartTransform.InverseTransformVectorNoScale( GizmoSpaceFacingAxisVector ) );
-
-						// Relative offset of the intersection on the plane
-						const FVector IntersectLocalLocationOnPlane = Transformable.StartTransform.InverseTransformPositionNoScale( IntersectWorldLocationOnPlane );
-						FVector RotatedIntersectLocationOnPlane = TransformableSpaceFacingAxisVector.Rotation().UnrotateVector( IntersectLocalLocationOnPlane );
-						float AngleToIntersectedLocation = FMath::Atan2( RotatedIntersectLocationOnPlane.Y, RotatedIntersectLocationOnPlane.Z );
-
-						if ( !StartDragAngleOnRotation.IsSet() )
-						{
-							StartDragAngleOnRotation = AngleToIntersectedLocation;
-						}
-
-						// Delta rotation in gizmo space between the starting and the intersection rotation
-						const float AngleDeltaRotationFromStart = FMath::FindDeltaAngleRadians( AngleToIntersectedLocation, StartDragAngleOnRotation.GetValue() );
-						const FQuat GizmoSpaceDeltaRotation = FQuat( GizmoSpaceFacingAxisVector, AngleDeltaRotationFromStart );
-
-						// Add the delta rotation to the starting transformation
-						FTransform GizmoSpaceRotatedTransform = GizmoSpaceStartTransform * GizmoSpaceDeltaRotation;
-
-						// Snap rotation in gizmo space
-						if ( !bIsUpdatingUnsnappedTarget )
-						{
-							FRotator SnappedRotation = GizmoSpaceRotatedTransform.GetRotation().Rotator();
-							FSnappingUtils::SnapRotatorToGrid( SnappedRotation );
-							GizmoSpaceRotatedTransform.SetRotation( SnappedRotation.Quaternion() );
-						}
-
-						// Rotation back to world space 
-						const FTransform WorldSpaceRotatedTransform = GizmoSpaceRotatedTransform * GizmoStartTransform;
-
-						PassTargetTransform = WorldSpaceRotatedTransform;
-						bMovedAnyTransformables = true;
-						bShouldApplyVelocitiesFromDrag = false;
-					}
-				}
-			}
-		}
-	}
-	else if( DraggingMode == EVREditorDraggingMode::ActorsFreely ||
-			 DraggingMode == EVREditorDraggingMode::World )
-	{
-		FVector TranslationOffset = FVector::ZeroVector;
-		FQuat RotationOffset = FQuat::Identity;
-		FVector ScaleOffset = FVector( 1.0f );
-		
-		bool bHasPivotLocation = false;
-		FVector PivotLocation = FVector::ZeroVector;
-		
-
-		if( bWithTwoHands )
-		{
-			// Rotate/scale while simultaneously counter-translating (two hands)
-
-			// NOTE: The reason that we use per-frame deltas (rather than the delta from the start of the whole interaction,
-			// like we do with the other interaction modes), is because the point that we're rotating/scaling relative to
-			// actually changes every update as the user moves their hands.  So it's iteratively getting the user to
-			// the result they naturally expect.
-
-			const FVector LineStart = DraggedTo;
-			const FVector LineStartDelta = DragDelta;
-			const float LineStartDistance = LineStartDelta.Size();
-			const FVector LastLineStart = LineStart - LineStartDelta;
-
-			const FVector LineEnd = OtherHandDraggedTo;
-			const FVector LineEndDelta = OtherHandDragDelta;
-			const float LineEndDistance = LineEndDelta.Size();
-			const FVector LastLineEnd = LineEnd - LineEndDelta;
-
-			// Choose a point along the new line segment to rotate and scale relative to.  We'll weight it toward the
-			// side of the line that moved the most this update.
-			const float TotalDistance = LineStartDistance + LineEndDistance;
-			float LineStartToEndActivityWeight = 0.5f;	// Default to right in the center, if no distance moved yet.
-			if( !FMath::IsNearlyZero( TotalDistance ) )	// Avoid division by zero
-			{
-				LineStartToEndActivityWeight = LineStartDistance / TotalDistance;
-			}
-
-			PivotLocation = FMath::Lerp( LastLineStart, LastLineEnd, LineStartToEndActivityWeight );
-			bHasPivotLocation = true;
-
-			if( DraggingMode == EVREditorDraggingMode::World && VREd::ScaleWorldFromFloor->GetInt() != 0 )
-			{
-				PivotLocation.Z = 0.0f;
-			}
-
-			// @todo vreditor debug
-			if( false )
-			{
-				const FTransform LinesRelativeTo = DraggingMode == EVREditorDraggingMode::World ? Owner.GetRoomTransform() : FTransform::Identity;
-				DrawDebugLine( Owner.GetWorld(), LinesRelativeTo.TransformPosition( LineStart ), LinesRelativeTo.TransformPosition( LineEnd ), FColor::Green, false, 0.0f );
-				DrawDebugLine( Owner.GetWorld(), LinesRelativeTo.TransformPosition( LastLineStart ), LinesRelativeTo.TransformPosition( LastLineEnd ), FColor::Red, false, 0.0f );
-
-				DrawDebugSphere( Owner.GetWorld(), LinesRelativeTo.TransformPosition( PivotLocation ), 2.5f * Owner.GetWorldScaleFactor(), 32, FColor::White, false, 0.0f );
-			}
-
-			const float LastLineLength = ( LastLineEnd - LastLineStart ).Size();
-			const float LineLength = ( LineEnd - LineStart ).Size();
-			ScaleOffset = FVector( LineLength / LastLineLength );
-//			ScaleOffset = FVector( 0.98f + 0.04f * FMath::MakePulsatingValue( FPlatformTime::Seconds(), 0.1f ) ); // FVector( LineLength / LastLineLength );
-
-			// How much did the line rotate since last time?
-			RotationOffset = FQuat::FindBetweenVectors( LastLineEnd - LastLineStart, LineEnd - LineStart );
-
-			// For translation, only move proportionally to the common vector between the two deltas.  Basically,
-			// you need to move both hands in the same direction to translate while gripping with two hands.
-			const FVector AverageDelta = ( DragDelta + OtherHandDragDelta ) * 0.5f;
-			const float TranslationWeight = FMath::Max( 0.0f, FVector::DotProduct( DragDelta.GetSafeNormal(), OtherHandDragDelta.GetSafeNormal() ) );
-			TranslationOffset = FMath::Lerp( FVector::ZeroVector, AverageDelta, TranslationWeight );
-		}
-		else
-		{
-			// Translate only (one hand)
-			TranslationOffset = DragDelta;
-		}
-
-		if( VREd::AllowVerticalWorldMovement->GetInt() == 0 )
-		{
-			TranslationOffset.Z = 0.0f;
-		}
-
-
-		if( DraggingMode == EVREditorDraggingMode::ActorsFreely )
-		{
-			for( TUniquePtr<FVREditorTransformable>& TransformablePtr : Transformables )
-			{
-				FVREditorTransformable& Transformable = *TransformablePtr;
-
-				const FTransform OldTransform = Transformable.UnsnappedTargetTransform;
-
-				// Two passes.  First update the real transform.  Then update the unsnapped transform.
-				for( int32 PassIndex = 0; PassIndex < 2; ++PassIndex )
-				{
-					const bool bIsUpdatingUnsnappedTarget = ( PassIndex == 1 );
-
-					// Scale snap!
-					FVector SnappedScaleOffset = ScaleOffset;
-					if( !bIsUpdatingUnsnappedTarget && FSnappingUtils::IsScaleSnapEnabled() )
-					{
-						// Did scale even change?
-						if( !( FMath::IsNearlyEqual( ScaleOffset.GetAbsMax(), 1.0f ) && FMath::IsNearlyEqual( ScaleOffset.GetAbsMin(), 1.0f ) ) )
-						{
-							FVector NewTotalScale = OldTransform.GetScale3D() * ScaleOffset;
-							FSnappingUtils::SnapScale( NewTotalScale, FVector::ZeroVector );
-							SnappedScaleOffset = NewTotalScale / OldTransform.GetScale3D();
-						}
-					}
-
-					const FTransform ScaleOffsetTransform( FQuat::Identity, FVector::ZeroVector, SnappedScaleOffset );
-					const FTransform TranslationOffsetTransform( FQuat::Identity, TranslationOffset );
-					const FTransform RotationOffsetTransform( RotationOffset, FVector::ZeroVector );
-
-					if( !bHasPivotLocation )
-					{
-						PivotLocation = OldTransform.GetLocation();
-					}
-
-					// @todo vreditor multi: This is only solving for rotating/scaling one object.  Really we want to rotate/scale the GIZMO and have all objects rotated and repositioned within it!
-					FTransform NewTransform;
-					const bool bIgnoreGizmoSpace = true;
-					if( bIgnoreGizmoSpace )
-					{
-						const FTransform PivotToWorld = FTransform( FQuat::Identity, PivotLocation );
-						const FTransform WorldToPivot = FTransform( FQuat::Identity, -PivotLocation );
-						NewTransform = OldTransform * WorldToPivot * ScaleOffsetTransform * RotationOffsetTransform * PivotToWorld * TranslationOffsetTransform;
-					}
-					else
-					{
-						// @todo vreditor perf: Can move most of this outside of the inner loop.  Also could be optimized by reducing the number of transforms.
-						const FTransform GizmoToWorld = GizmoStartTransform;
-						const FTransform WorldToGizmo = GizmoToWorld.Inverse();
-						const FVector GizmoSpacePivotLocation = WorldToGizmo.TransformPositionNoScale( PivotLocation );
-						const FTransform GizmoToPivot = FTransform( FQuat::Identity, -GizmoSpacePivotLocation );
-						const FTransform PivotToGizmo = FTransform( FQuat::Identity, GizmoSpacePivotLocation );
-						const FQuat GizmoSpaceRotationOffsetQuat = WorldToGizmo.GetRotation() * RotationOffset;	// NOTE: FQuat composition order is opposite that of FTransform!
-						const FTransform GizmoSpaceRotationOffsetTransform( GizmoSpaceRotationOffsetQuat, FVector::ZeroVector );
-						NewTransform = OldTransform * TranslationOffsetTransform * WorldToGizmo * GizmoToPivot * ScaleOffsetTransform * GizmoSpaceRotationOffsetTransform * PivotToGizmo * GizmoToWorld;
-					}
-
-					if( bIsUpdatingUnsnappedTarget )
-					{
-						const FVector OldLocation = Transformable.UnsnappedTargetTransform.GetLocation();
-						Transformable.UnsnappedTargetTransform = NewTransform;
-					}
-					else 
-					{
-						Transformable.TargetTransform = NewTransform;
-
-						// Grid snap!
-						FVector SnappedLocation = Transformable.TargetTransform.GetLocation();
-						if( bLocalSpaceSnapping )
-						{
-							SnappedLocation = GizmoStartTransform.InverseTransformPositionNoScale( SnappedLocation );
-							FSnappingUtils::SnapPointToGrid( SnappedLocation, SnapGridBase );
-							SnappedLocation = GizmoStartTransform.TransformPositionNoScale( SnappedLocation );
-						}
-						else
-						{
-							FSnappingUtils::SnapPointToGrid( SnappedLocation, SnapGridBase );
-						}
-						Transformable.TargetTransform.SetLocation( SnappedLocation );
-
-						// Rotation snap!
-						FRotator SnappedRotation = Transformable.TargetTransform.GetRotation().Rotator();
-						FSnappingUtils::SnapRotatorToGrid( SnappedRotation );
-						Transformable.TargetTransform.SetRotation( SnappedRotation.Quaternion() );
-					}
-
-				}
-
-				bMovedAnyTransformables = true;
-				bShouldApplyVelocitiesFromDrag = true;
-			}
-		}
-		else if( DraggingMode == EVREditorDraggingMode::World )
-		{
-			FTransform RoomTransform = Owner.GetRoomTransform();
-
-			// Adjust world scale
-			const float WorldScaleOffset = ScaleOffset.GetAbsMax();
-			if( WorldScaleOffset != 0.0f )
-			{
-				const float OldWorldToMetersScale = Owner.GetWorld()->GetWorldSettings()->WorldToMeters;
-				const float NewWorldToMetersScale = OldWorldToMetersScale / WorldScaleOffset;
-
-				// NOTE: Instead of clamping, we simply skip changing the W2M this frame if it's out of bounds.  Clamping makes our math more complicated.
-				if( NewWorldToMetersScale != OldWorldToMetersScale &&
-					NewWorldToMetersScale >= Owner.GetMinScale() && NewWorldToMetersScale <= Owner.GetMaxScale() )
-				{
-					Owner.SetWorldToMetersScale( NewWorldToMetersScale );
-
-					// Because the tracking space size has changed, but our head position within that space relative to the origin
-					// of the room is the same (before scaling), we need to offset our location within the tracking space to compensate.
-					// This makes the user feel like their head and hands remain in the same location.
-					const FVector RoomSpacePivotLocation = PivotLocation;
-					const FVector WorldSpacePivotLocation = RoomTransform.TransformPosition( RoomSpacePivotLocation );
-
-					const FVector NewRoomSpacePivotLocation = ( RoomSpacePivotLocation / OldWorldToMetersScale ) * NewWorldToMetersScale;
-					const FVector NewWorldSpacePivotLocation = RoomTransform.TransformPosition( NewRoomSpacePivotLocation );
-
-					const FVector WorldSpacePivotDelta = ( NewWorldSpacePivotLocation - WorldSpacePivotLocation );
-
-					const FVector NewWorldSpaceRoomLocation = RoomTransform.GetLocation() - WorldSpacePivotDelta;
-
-					RoomTransform.SetLocation( NewWorldSpaceRoomLocation );
-				}
-			}
-
-			// Apply rotation and translation
-			{
-				FTransform RotationOffsetTransform = FTransform::Identity;
-				RotationOffsetTransform.SetRotation( RotationOffset );
-
-				if( VREd::AllowWorldRotationPitchAndRoll->GetInt() == 0 )
-				{
-					// Eliminate pitch and roll in rotation offset.  We don't want the user to get sick!
-					FRotator YawRotationOffset = RotationOffset.Rotator();
-					YawRotationOffset.Pitch = YawRotationOffset.Roll = 0.0f;
-					RotationOffsetTransform.SetRotation( YawRotationOffset.Quaternion() );
-				}
-
-				// Move the camera in the opposite direction, so it feels to the user as if they're dragging the entire world around
-				const FTransform TranslationOffsetTransform( FQuat::Identity, TranslationOffset );
-				const FTransform PivotToWorld = FTransform( FQuat::Identity, PivotLocation ) * RoomTransform;
-				const FTransform WorldToPivot = PivotToWorld.Inverse();
-				RoomTransform = TranslationOffsetTransform.Inverse() * RoomTransform * WorldToPivot * RotationOffsetTransform.Inverse() * PivotToWorld;
-			}
-
-			Owner.SetRoomTransform( RoomTransform );
-		}
-
-	}
-
-
-	// If we're not using smooth snapping, go ahead and immediately update the transforms of all objects.  If smooth
-	// snapping is enabled, this will be done in Tick() instead.
-	if( bMovedAnyTransformables )
-	{
-		// Update velocity if we're simulating in editor
-		if( GEditor->bIsSimulatingInEditor )
-		{
-			for( TUniquePtr<FVREditorTransformable>& TransformablePtr : Transformables )
-			{
-				FVREditorTransformable& Transformable = *TransformablePtr;
-
-				AActor* Actor = Cast<AActor>( Transformable.Object.Get() );
-				if( Actor != nullptr )
-				{
-					UPrimitiveComponent* RootPrimitiveComponent = Cast<UPrimitiveComponent>( Actor->GetRootComponent() );
-					if( RootPrimitiveComponent != nullptr )
-					{
-						// @todo vreditor physics: When freely transforming, should set angular velocity too (will pivot be the same though??)
-						if( RootPrimitiveComponent->IsSimulatingPhysics() )
-						{
-							FVector MoveDelta = Transformable.UnsnappedTargetTransform.GetLocation() - Transformable.LastTransform.GetLocation();
-							if( bShouldApplyVelocitiesFromDrag )
-							{
-								bIsDrivingVelocityOfSimulatedTransformables = true;
-							}
-							else
-							{
-								MoveDelta = FVector::ZeroVector;
-							}
-							RootPrimitiveComponent->SetAllPhysicsLinearVelocity( MoveDelta * DeltaTime * VREd::InertiaVelocityBoost->GetFloat() );
-						}
-					}
-				}
-			}
-		}
-
-		const bool bSmoothSnappingEnabled = IsSmoothSnappingEnabled();
-		if( !bSmoothSnappingEnabled && !bIsInterpolatingTransformablesFromSnapshotTransform )
-		{
-			const bool bSweep = bShouldApplyVelocitiesFromDrag && VREd::SweepPhysicsWhileSimulating->GetInt() != 0 && bIsDrivingVelocityOfSimulatedTransformables;
-			for( TUniquePtr<FVREditorTransformable>& TransformablePtr : Transformables )
-			{
-				FVREditorTransformable& Transformable = *TransformablePtr;
-
-				Transformable.LastTransform = Transformable.TargetTransform;
-
-				// Got an actor?
-				AActor* Actor = Cast<AActor>( Transformable.Object.Get() );
-				if( Actor != nullptr )
-				{
-					const FTransform& ExistingTransform = Actor->GetTransform();
-
-					const bool bOnlyTranslationChanged =
-						ExistingTransform.GetRotation() == Transformable.TargetTransform.GetRotation() &&
-						ExistingTransform.GetScale3D() == Transformable.TargetTransform.GetScale3D();
+// For actor placement
+#include "ObjectTools.h"
+#include "AssetSelection.h"
+#include "IPlacementModeModule.h"
+#include "Kismet/GameplayStatics.h"
 
-					Actor->SetActorTransform( Transformable.TargetTransform, bSweep );
+#include "LevelEditor.h"
+#include "SLevelViewport.h"
+#include "Editor/LevelEditor/Public/LevelEditorActions.h"
 
-					Actor->InvalidateLightingCacheDetailed( bOnlyTranslationChanged );
+#define LOCTEXT_NAMESPACE "VREditor"
 
-					const bool bFinished = false;
-					Actor->PostEditMove( bFinished );
-				}
-				else
-				{
-					// Some other object that we don't know how to drag
-				}
-			}
-		}
-	}
 
-	bIsFirstDragUpdate = false;
+namespace VREd
+{
+	static FAutoConsoleVariable SizeOfActorsOverContentBrowserThumbnail( TEXT( "VREd.SizeOfActorsOverContentBrowserThumbnail" ), 6.0f, TEXT( "How large objects should be when rendered 'thumbnail size' over the Content Browser" ) );
+	static FAutoConsoleVariable OverrideSizeOfPlacedActors( TEXT( "VREd.OverrideSizeOfPlacedActors" ), 0.0f, TEXT( "If set to a value greater than zero, sets the size in cm of a placed actor (relative to your world's scale)" ) );
+	static FAutoConsoleVariable TrackpadAbsoluteDragSpeed( TEXT( "VREd.TrackpadAbsoluteDragSpeed" ), 40.0f, TEXT( "How fast objects move toward or away when you drag on the touchpad while carrying them" ) );
+	static FAutoConsoleVariable TrackpadRelativeDragSpeed( TEXT( "VREd.TrackpadRelativeDragSpeed" ), 8.0f, TEXT( "How fast objects move toward or away when you hold a direction on an analog stick while carrying them" ) );
+	static FAutoConsoleVariable HoverHapticFeedbackStrength( TEXT( "VREd.HoverHapticFeedbackStrength" ), 0.1f, TEXT( "Default strength for haptic feedback when hovering" ) );
+	static FAutoConsoleVariable HoverHapticFeedbackTime( TEXT( "VREd.HoverHapticFeedbackTime" ), 0.2f, TEXT( "The minimum time between haptic feedback for hovering" ) );
+	static FAutoConsoleVariable PivotPointTransformGizmo( TEXT( "VREd.PivotPointTransformGizmo" ), 1, TEXT( "If the pivot point transform gizmo is used instead of the bounding box gizmo" ) );
+	static FAutoConsoleVariable DragHapticFeedbackStrength( TEXT( "VREd.DragHapticFeedbackStrength" ), 1.0f, TEXT( "Default strength for haptic feedback when starting to drag objects" ) );
+	static FAutoConsoleVariable PlacementInterpolationDuration( TEXT( "VREd.PlacementInterpolationDuration" ), 0.6f, TEXT( "How long we should interpolate newly-placed objects to their target location." ) );
 }
 
-void FVREditorWorldInteraction::StartDraggingActors(FVRAction VRAction, UActorComponent* ClickedComponent, const FVector HitLocation, const bool bIsPlacingActors)
+UVREditorWorldInteraction::UVREditorWorldInteraction( const FObjectInitializer& Initializer ) : 
+	Super( Initializer ),
+	Owner( nullptr ),
+	PlacingMaterialOrTextureAsset( nullptr ),
+	DropMaterialOrMaterialSound( nullptr ),
+	FloatingUIAssetDraggedFrom( nullptr )
 {
-	if( IsInteractableComponent( ClickedComponent ) )
-	{
-		FVector LaserPointerStart, LaserPointerEnd;
-		const bool bHaveLaserPointer = Owner.GetLaserPointer( VRAction.HandIndex, /* Out */ LaserPointerStart, /* Out */ LaserPointerEnd );
-		if( bHaveLaserPointer )
-		{
-			AActor* Actor = ClickedComponent->GetOwner();
-
-			FVirtualHand& Hand = Owner.GetVirtualHand( VRAction.HandIndex );
-
-			// Capture undo state
-			if( bIsPlacingActors )
-			{
-				// When placing actors, a transaction should already be in progress
-				check( TrackingTransaction.IsActive() );
-			}
-			else
-			{
-				TrackingTransaction.TransCount++;
-				TrackingTransaction.Begin( LOCTEXT( "MovingActors", "Moving Selected Actors" ) );
-
-				// Suspend actor/component modification during each delta step to avoid recording unnecessary overhead into the transaction buffer
-				GEditor->DisableDeltaModification( true );
-			}
-
-			// If the user is holding down the modifier key, go ahead and duplicate the selection first.  This needs to
-			// happen before we create our transformables, because duplicating actors will change which actors are
-			// selected!  Don't do this if we're placing objects right now though.
-			if( Hand.bIsModifierPressed && !bIsPlacingActors )
-			{
-				Owner.Duplicate();
-			}
-
-			const bool bUsingGizmo =
-				VRAction.ActionType == EVRActionType::SelectAndMove_LightlyPressed &&		// Only use the gizmo when lightly pressed
-				( Actor == TransformGizmoActor ) &&
-				ClickedComponent != nullptr;
-
-			// Start dragging the objects right away!
-			Hand.DraggingMode = Hand.LastDraggingMode = bIsPlacingActors ? EVREditorDraggingMode::ActorsAtLaserImpact : ( bUsingGizmo ? EVREditorDraggingMode::ActorsWithGizmo : EVREditorDraggingMode::ActorsFreely );
-
-			// Starting a new drag, so make sure the other hand doesn't think it's assisting us
-			FVirtualHand& OtherHand = Owner.GetOtherHand( VRAction.HandIndex );
-			OtherHand.bWasAssistingDrag = false;
-
-			Hand.bIsFirstDragUpdate = true;
-			Hand.bWasAssistingDrag = false;
-			Hand.DragRayLength = ( HitLocation - LaserPointerStart ).Size();
-			Hand.LastLaserPointerStart = LaserPointerStart;
-			Hand.LastDragToLocation = HitLocation;
-			Hand.LaserPointerImpactAtDragStart = HitLocation;
-			Hand.DragTranslationVelocity = FVector::ZeroVector;
-			Hand.DragRayLengthVelocity = 0.0f;
-			Hand.bIsDrivingVelocityOfSimulatedTransformables = false;
-
-			// Start dragging the transform gizmo.  Even if the user clicked on the actor itself, we'll use
-			// the gizmo to transform it.
-			if( bUsingGizmo )
-			{
-				Hand.DraggingTransformGizmoComponent = Cast<USceneComponent>( ClickedComponent );
-			}
-			else
-			{
-				Hand.DraggingTransformGizmoComponent = nullptr;
-			}
-
-			if( TransformGizmoActor != nullptr )
-			{
-				Hand.TransformGizmoInteractionType = TransformGizmoActor->GetInteractionType( Hand.DraggingTransformGizmoComponent.Get(), Hand.OptionalHandlePlacement );
-				Hand.GizmoStartTransform = TransformGizmoActor->GetTransform();
-				Hand.GizmoStartLocalBounds = GizmoLocalBounds;
-			}
-			else
-			{
-				Hand.TransformGizmoInteractionType = ETransformGizmoInteractionType::None;
-				Hand.GizmoStartTransform = FTransform::Identity;
-				Hand.GizmoStartLocalBounds = FBox( 0 );
-			}
-			Hand.GizmoSpaceFirstDragUpdateOffsetAlongAxis = FVector::ZeroVector;	// Will be determined on first update
-			Hand.GizmoSpaceDragDeltaFromStartOffset = FVector::ZeroVector;	// Set every frame while dragging
-
-			bDraggedSinceLastSelection = true;
-			LastDragGizmoStartTransform = Hand.GizmoStartTransform;
-
-			SetupTransformablesForSelectedActors();
-
-
-			// If we're placing actors, start interpolating to their actual location.  This helps smooth everything out when
-			// using the laser impact point as the target transform
-			if( bIsPlacingActors )
-			{
-				bIsInterpolatingTransformablesFromSnapshotTransform = true;
-				const FTimespan CurrentTime = FTimespan::FromSeconds( FPlatformTime::Seconds() );
-				TransformablesInterpolationStartTime = CurrentTime;
-				TransformablesInterpolationDuration = VREd::PlacementInterpolationDuration->GetFloat();
-			}
+	// Find out when the user drags stuff out of a content browser
+	FEditorDelegates::OnAssetDragStarted.AddUObject( this, &UVREditorWorldInteraction::OnAssetDragStartedFromContentBrowser );
 
-			// Play a haptic effect when objects are picked up
-			const float Strength = VREd::DragHapticFeedbackStrength->GetFloat();
-			Owner.PlayHapticEffect(
-				VRAction.HandIndex == VREditorConstants::LeftHandIndex ? Strength : 0.0f,
-				VRAction.HandIndex == VREditorConstants::RightHandIndex ? Strength : 0.0f );
-		}
-	}
+	// Load sounds
+	DropMaterialOrMaterialSound = LoadObject<USoundCue>( nullptr, TEXT( "/Engine/VREditor/Sounds/VR_grab_Cue" ) );
+	check( DropMaterialOrMaterialSound != nullptr );
 }
 
-
-void FVREditorWorldInteraction::StartDraggingMaterialOrTexture( FVRAction VRAction, const FVector HitLocation, UObject* MaterialOrTextureAsset )
+void UVREditorWorldInteraction::Shutdown()
 {
-	check( MaterialOrTextureAsset != nullptr );
-
-	FVector LaserPointerStart, LaserPointerEnd;
-	const bool bHaveLaserPointer = Owner.GetLaserPointer( VRAction.HandIndex, /* Out */ LaserPointerStart, /* Out */ LaserPointerEnd );
-	if( bHaveLaserPointer )
-	{
-		this->PlacingMaterialOrTextureAsset = MaterialOrTextureAsset;
-
-		FVirtualHand& Hand = Owner.GetVirtualHand( VRAction.HandIndex );
-
-		Hand.DraggingMode = EVREditorDraggingMode::Material;
-
-		// Starting a new drag, so make sure the other hand doesn't think it's assisting us
-		FVirtualHand& OtherHand = Owner.GetOtherHand( VRAction.HandIndex );
-		OtherHand.bWasAssistingDrag = false;
-
-		Hand.bIsFirstDragUpdate = true;
-		Hand.bWasAssistingDrag = false;
-		Hand.DragRayLength = ( HitLocation - LaserPointerStart ).Size();
-		Hand.LastLaserPointerStart = LaserPointerStart;
-		Hand.LastDragToLocation = HitLocation;
-		Hand.LaserPointerImpactAtDragStart = HitLocation;
-		Hand.DragTranslationVelocity = FVector::ZeroVector;
-		Hand.DragRayLengthVelocity = 0.0f;
-		Hand.bIsDrivingVelocityOfSimulatedTransformables = false;
-
-		Hand.DraggingTransformGizmoComponent = nullptr;
-	
-		Hand.TransformGizmoInteractionType = ETransformGizmoInteractionType::None;
-		Hand.GizmoStartTransform = FTransform::Identity;
-		Hand.GizmoStartLocalBounds = FBox( 0 );
+	Super::Shutdown();
 
-		Hand.GizmoSpaceFirstDragUpdateOffsetAlongAxis = FVector::ZeroVector;	// Will be determined on first update
-		Hand.GizmoSpaceDragDeltaFromStartOffset = FVector::ZeroVector;	// Set every frame while dragging
-
-		bDraggedSinceLastSelection = false;
-		LastDragGizmoStartTransform = FTransform::Identity;
+	FEditorDelegates::OnAssetDragStarted.RemoveAll( this );
 
-		// Play a haptic effect when objects are picked up
-		const float Strength = VREd::DragHapticFeedbackStrength->GetFloat();
-		Owner.PlayHapticEffect(
-			VRAction.HandIndex == VREditorConstants::LeftHandIndex ? Strength : 0.0f,
-			VRAction.HandIndex == VREditorConstants::RightHandIndex ? Strength : 0.0f );
-	}
+	PlacingMaterialOrTextureAsset = nullptr;
+	FloatingUIAssetDraggedFrom = nullptr;
+	DropMaterialOrMaterialSound = nullptr;
+	Owner = nullptr;
 }
 
-
-void FVREditorWorldInteraction::StopDragging( const int32 HandIndex )
+bool UVREditorWorldInteraction::IsInteractableComponent( const UActorComponent* Component ) const
 {
-	FVirtualHand& Hand = Owner.GetVirtualHand( HandIndex );
-	if( Hand.DraggingMode != EVREditorDraggingMode::Nothing )
-	{
-		// If the other hand started dragging after we started, allow that hand to "take over" the drag, so the user
-		// doesn't have to click again to continue their action.  Inertial effects of the hand that stopped dragging
-		// will still be in effect.
-		FVirtualHand& OtherHand = Owner.GetOtherHand( HandIndex );
-		if( OtherHand.DraggingMode == EVREditorDraggingMode::AssistingDrag )
-		{
-			// The other hand takes over whatever this hand was doing
-			OtherHand.DraggingMode = OtherHand.LastDraggingMode = Hand.DraggingMode;
-			OtherHand.bIsDrivingVelocityOfSimulatedTransformables = Hand.bIsDrivingVelocityOfSimulatedTransformables;
+	bool bResult = false;
+	bResult = UViewportWorldInteraction::IsInteractableComponent( Component );
 
-			// The other hand is no longer assisting, as it's now the primary interacting hand.
-			OtherHand.bWasAssistingDrag = false;
-
-			// The hand that stopped dragging will be remembered as "assisting" the other hand, so that its
-			// inertia will still influence interactions
-			Hand.bWasAssistingDrag = true;	
-		}
-		else
+	// Don't bother if the base class already found out that this is a custom interactable
+	if ( bResult )
+	{
+		// Don't interact with frozen actors or our UI
+		if ( Component != nullptr )
 		{
-			// No hand is dragging the objects anymore.
-			if( Hand.DraggingMode == EVREditorDraggingMode::ActorsWithGizmo ||
-				Hand.DraggingMode == EVREditorDraggingMode::ActorsFreely ||
-				Hand.DraggingMode == EVREditorDraggingMode::ActorsAtLaserImpact )
-			{
-				if ( Hand.DraggingMode == EVREditorDraggingMode::ActorsWithGizmo &&
-					 Hand.TransformGizmoInteractionType == ETransformGizmoInteractionType::RotateOnAngle )
-				{
-					StartDragAngleOnRotation.Reset();
-				}
-
-				// Finalize undo
-				{
-					// @todo vreditor undo: This doesn't actually encapsulate any inertial movement that happens after the drag is released! 
-					// We need to figure out whether that matters or not.  Also, look into PostEditMove( bFinished=true ) and how that relates to this.
-					--TrackingTransaction.TransCount;
-					TrackingTransaction.End();
-					GEditor->DisableDeltaModification( false );
-				}
-			}
+			static const bool bIsVREditorDemo = FParse::Param( FCommandLine::Get(), TEXT( "VREditorDemo" ) );	// @todo vreditor: Remove this when no longer needed (console variable, too!)
+			const bool bIsFrozen = bIsVREditorDemo && Component->GetOwner()->GetActorLabel().StartsWith( TEXT( "Frozen_" ) );
+			bResult = !bIsFrozen && !GetOwner().GetUISystem().IsWidgetAnEditorUIWidget( Component ); // Don't allow user to move around our UI widgets //@todo VREditor: This wouldn't be necessary if the UI used the interactable interface
 		}
-
-		Hand.DraggingMode = EVREditorDraggingMode::Nothing;
-		Hand.IsInputCaptured[ (int32)EVRActionType::SelectAndMove ] = false;
 	}
-
-	// NOTE: Even though we're done dragging, we keep our list of transformables so that inertial drag can still push them around!
+	
+	return bResult;
 }
 
-
-void FVREditorWorldInteraction::SetupTransformablesForSelectedActors()
+void UVREditorWorldInteraction::StopDragging( UViewportInteractor* Interactor )
 {
-	bIsInterpolatingTransformablesFromSnapshotTransform = false;
-	TransformablesInterpolationStartTime = FTimespan::Zero();
-	TransformablesInterpolationDuration = 1.0f;
-
-	Transformables.Reset();
-
-	USelection* ActorSelectionSet = GEditor->GetSelectedActors();
+	EViewportInteractionDraggingMode InteractorDraggingMode = Interactor->GetDraggingMode();
 
-	static TArray<UObject*> SelectedActorObjects;
-	SelectedActorObjects.Reset();
-	ActorSelectionSet->GetSelectedObjects( AActor::StaticClass(), /* Out */ SelectedActorObjects );
-
-	for( UObject* SelectedActorObject : SelectedActorObjects )
+	// If we were placing a material, go ahead and do that now
+	if ( InteractorDraggingMode == EViewportInteractionDraggingMode::Material )
 	{
-		AActor* SelectedActor = CastChecked<AActor>( SelectedActorObject );
-
-		Transformables.Add( TUniquePtr< FVREditorTransformable>( new FVREditorTransformable() ) );
-		FVREditorTransformable& Transformable = *Transformables[ Transformables.Num() - 1 ];
-
-		Transformable.Object = SelectedActor;
-		Transformable.StartTransform = Transformable.LastTransform = Transformable.TargetTransform = Transformable.UnsnappedTargetTransform = Transformable.InterpolationSnapshotTransform = SelectedActor->GetTransform();
+		PlaceDraggedMaterialOrTexture( Interactor );
 	}
-}
 
-
-float FVREditorWorldInteraction::GetTrackpadSlideDelta( const int32 HandIndex )
-{
-	const bool bIsAbsolute = (Owner.GetHMDDeviceType() == EHMDDeviceType::DT_SteamVR);
-	float SlideDelta = 0.0f;
-	FVirtualHand& Hand = Owner.GetVirtualHand( HandIndex );
-	if (Hand.bIsTouchingTrackpad || !bIsAbsolute)
+	// If we were placing something, bring the window back
+	if ( InteractorDraggingMode == EViewportInteractionDraggingMode::ActorsAtLaserImpact ||
+		InteractorDraggingMode == EViewportInteractionDraggingMode::Material )
 	{
-		if (bIsAbsolute)
-		{
-			SlideDelta = ((Hand.TrackpadPosition.Y - Hand.LastTrackpadPosition.Y) * VREd::TrackpadAbsoluteDragSpeed->GetFloat());
-		}
-		else
-		{
-			SlideDelta = (Hand.TrackpadPosition.Y * VREd::TrackpadRelativeDragSpeed->GetFloat());
-		}
+		GetOwner().GetUISystem().ShowEditorUIPanel( FloatingUIAssetDraggedFrom, Cast<UVREditorInteractor>( Interactor->GetOtherInteractor() ), 
+			true, false, false, false );
+		FloatingUIAssetDraggedFrom = nullptr;
 	}
 
-	return SlideDelta;
+
+	Super::StopDragging( Interactor );
 }
 
-bool FVREditorWorldInteraction::IsTransformingActor( AActor* Actor ) const
+void UVREditorWorldInteraction::StartDraggingMaterialOrTexture( UViewportInteractor* Interactor, const FViewportActionKeyInput& Action, const FVector HitLocation, UObject* MaterialOrTextureAsset )
 {
-	bool bFoundActor = false;
-	for( const TUniquePtr<FVREditorTransformable>& TransformablePtr : Transformables )
+	check( MaterialOrTextureAsset != nullptr );
+
+	FVector LaserPointerStart, LaserPointerEnd;
+	const bool bHaveLaserPointer = Interactor->GetLaserPointer( /* Out */ LaserPointerStart, /* Out */ LaserPointerEnd );
+	if( bHaveLaserPointer )
 	{
-		const FVREditorTransformable& Transformable = *TransformablePtr;
+		PlacingMaterialOrTextureAsset = MaterialOrTextureAsset;
 
-		if( Transformable.Object.Get() == Actor )
-		{
-			bFoundActor = true;
-			break;
-		}
-	}
+		FViewportInteractorData& InteractorData = Interactor->GetInteractorData();
 
-	return bFoundActor;
-}
+		Interactor->SetDraggingMode( EViewportInteractionDraggingMode::Material );
 
-void FVREditorWorldInteraction::DeleteSelectedObjects()
-{
-	const FScopedTransaction Transaction( LOCTEXT( "DeleteSelection", "Delete selection" ) );
+		// Starting a new drag, so make sure the other hand doesn't think it's assisting us
+		FViewportInteractorData& OtherInteractorData = Interactor->GetOtherInteractor()->GetInteractorData();
+		OtherInteractorData.bWasAssistingDrag = false;
+
+		InteractorData.bIsFirstDragUpdate = true;
+		InteractorData.bWasAssistingDrag = false;
+		InteractorData.DragRayLength = ( HitLocation - LaserPointerStart ).Size();
+		InteractorData.LastLaserPointerStart = LaserPointerStart;
+		InteractorData.LastDragToLocation = HitLocation;
+		InteractorData.LaserPointerImpactAtDragStart = HitLocation;
+		InteractorData.DragTranslationVelocity = FVector::ZeroVector;
+		InteractorData.DragRayLengthVelocity = 0.0f;
+		InteractorData.bIsDrivingVelocityOfSimulatedTransformables = false;
+
+		InteractorData.DraggingTransformGizmoComponent = nullptr;
+	
+		InteractorData.TransformGizmoInteractionType = ETransformGizmoInteractionType::None;
+		InteractorData.GizmoStartTransform = FTransform::Identity;
+		InteractorData.GizmoStartLocalBounds = FBox( 0 );
 
-	GUnrealEd->edactDeleteSelected( Owner.GetWorld() );
-}
+		InteractorData.GizmoSpaceFirstDragUpdateOffsetAlongAxis = FVector::ZeroVector;	// Will be determined on first update
+		InteractorData.GizmoSpaceDragDeltaFromStartOffset = FVector::ZeroVector;	// Set every frame while dragging
 
+		bDraggedSinceLastSelection = false;
+		LastDragGizmoStartTransform = FTransform::Identity;
 
-void FVREditorWorldInteraction::CleanUpActorsBeforeMapChangeOrSimulate()
-{
-	// Destroy actors
-	{
-		Owner.DestroyTransientActor( TransformGizmoActor );
-		TransformGizmoActor = nullptr;
+		// Play a haptic effect when objects aTrackingTransactionre picked up
+		Interactor->PlayHapticEffect( VREd::DragHapticFeedbackStrength->GetFloat() );
 	}
 }
 
-
-void FVREditorWorldInteraction::OnAssetDragStartedFromContentBrowser( const TArray<FAssetData>& DraggedAssets, UActorFactory* FactoryToUse )
+void UVREditorWorldInteraction::OnAssetDragStartedFromContentBrowser( const TArray<FAssetData>& DraggedAssets, UActorFactory* FactoryToUse )
 {
 	const bool bIsPreview = false;	 // @todo vreditor placement: Consider supporting a "drop preview" actor (so you can cancel placement interactively)
 
@@ -2844,44 +160,57 @@ void FVREditorWorldInteraction::OnAssetDragStartedFromContentBrowser( const TArr
 
 	// Figure out which controller pressed the button and started dragging
 	// @todo vreditor placement: This logic could misfire.  Ideally we would be routed information from the pointer event, so we can determine the hand.
-	int32 PlacingWithHandIndex = INDEX_NONE;
-	int32 OtherHandIndex = INDEX_NONE;
-	for( int32 HandIndex = 0; HandIndex < VREditorConstants::NumVirtualHands; ++HandIndex )
+	UVREditorInteractor* PlacingWithInteractor = nullptr;
+	for ( UViewportInteractor* Interactor : GetInteractors() )
 	{
-		FVirtualHand& Hand = Owner.GetVirtualHand( HandIndex );
-		if( ( Hand.IsInputCaptured[ (int32)EVRActionType::SelectAndMove ] ||
-			  Hand.IsInputCaptured[ (int32)EVRActionType::SelectAndMove_LightlyPressed ] ) && 
-			( Hand.bIsClickingOnUI && !Hand.bIsRightClickingOnUI ) )
+		UVREditorInteractor* VRInteractor = Cast<UVREditorInteractor>( Interactor );
+		FViewportActionKeyInput* SelectAndMoveAction = Interactor->GetActionWithName( ViewportWorldActionTypes::SelectAndMove );
+		FViewportActionKeyInput* SelectAndMoveLightlyPressedAction = Interactor->GetActionWithName( ViewportWorldActionTypes::SelectAndMove_LightlyPressed );
+
+		if ( VRInteractor && ( ( SelectAndMoveAction && SelectAndMoveAction->bIsInputCaptured ) 
+			|| ( SelectAndMoveLightlyPressedAction && SelectAndMoveLightlyPressedAction->bIsInputCaptured ) ) )
 		{
-			PlacingWithHandIndex = HandIndex;
-			OtherHandIndex = Owner.GetOtherHandIndex( PlacingWithHandIndex );
-			break;
+			if ( VRInteractor->IsClickingOnUI() && !VRInteractor->IsRightClickingOnUI() ) 
+			{
+				PlacingWithInteractor = VRInteractor;
+				break;
+			}
 		}
 	}
 
+
 	// We're always expecting a hand to be hovering at the time we receive this event
-	if( PlacingWithHandIndex != INDEX_NONE )
+	if( PlacingWithInteractor )
 	{
-		FVirtualHand& Hand = Owner.GetVirtualHand( PlacingWithHandIndex );
-
 		// Cancel UI input
 		{
-			Hand.bIsClickingOnUI = false;
-			Hand.bIsRightClickingOnUI = false;
-			Hand.IsInputCaptured[ (int32)EVRActionType::SelectAndMove ] = false;
-			Hand.IsInputCaptured[ (int32)EVRActionType::SelectAndMove_LightlyPressed ] = false;
+			FViewportActionKeyInput* SelectAndMoveAction = PlacingWithInteractor->GetActionWithName( ViewportWorldActionTypes::SelectAndMove );
+			FViewportActionKeyInput* SelectAndMoveLightlyPressedAction = PlacingWithInteractor->GetActionWithName( ViewportWorldActionTypes::SelectAndMove_LightlyPressed );
+
+			PlacingWithInteractor->SetIsClickingOnUI( false );
+			PlacingWithInteractor->SetIsRightClickingOnUI( false );
+
+			if ( SelectAndMoveAction )
+			{
+				SelectAndMoveAction->bIsInputCaptured = false;
+			}
+
+			if ( SelectAndMoveLightlyPressedAction )
+			{
+				SelectAndMoveLightlyPressedAction->bIsInputCaptured = false;
+			}
 		}
 
-		FloatingUIAssetDraggedFrom = Hand.HoveringOverWidgetComponent;
+		FloatingUIAssetDraggedFrom = PlacingWithInteractor->GetHoveringOverWidgetComponent();
 		// Hide the UI panel that's being used to drag
-		Owner.GetUISystem().ShowEditorUIPanel( FloatingUIAssetDraggedFrom, OtherHandIndex, false, false, true, false );
+		GetOwner().GetUISystem().ShowEditorUIPanel( FloatingUIAssetDraggedFrom, PlacingWithInteractor, false, false, true, false );
 
 		TArray< UObject* > DroppedObjects;
 		TArray< AActor* > AllNewActors;
 
 		UObject* DraggingSingleMaterialOrTexture = nullptr;
 
-		FVector PlaceAt = Hand.HoverLocation;
+		FVector PlaceAt = PlacingWithInteractor->GetHoverLocation();
 
 		// Only place the object at the laser impact point if we're NOT going to interpolate to the impact
 		// location.  When interpolation is enabled, it looks much better to blend to the new location
@@ -2889,7 +218,7 @@ void FVREditorWorldInteraction::OnAssetDragStartedFromContentBrowser( const TArr
 		if( !bShouldInterpolateFromDragLocation )
 		{
 			FVector HitLocation = FVector::ZeroVector;
-			if( FindPlacementPointUnderLaser( PlacingWithHandIndex, /* Out */ HitLocation ) )
+			if( FindPlacementPointUnderLaser( PlacingWithInteractor, /* Out */ HitLocation ) )
 			{
 				PlaceAt = HitLocation;
 			}
@@ -2900,7 +229,7 @@ void FVREditorWorldInteraction::OnAssetDragStartedFromContentBrowser( const TArr
 			bool bCanPlace = true;
 
 			UObject* AssetObj = AssetData.GetAsset();
-			if( !ObjectTools::IsAssetValidForPlacing( Owner.GetWorld(), AssetData.ObjectPath.ToString() ) )
+			if( !ObjectTools::IsAssetValidForPlacing( GetViewportWorld(), AssetData.ObjectPath.ToString() ) )
 			{
 				bCanPlace = false;
 			}
@@ -2954,7 +283,7 @@ void FVREditorWorldInteraction::OnAssetDragStartedFromContentBrowser( const TArr
 				const bool bSelectNewActors = true;
 				const EObjectFlags ObjectFlags = bIsPreview ? RF_Transient : RF_Transactional;
 
-				TArray<AActor*> NewActors = FLevelEditorViewportClient::TryPlacingActorFromObject( Owner.GetWorld()->GetCurrentLevel(), AssetObj, bSelectNewActors, ObjectFlags, FactoryToUse );
+				TArray<AActor*> NewActors = FLevelEditorViewportClient::TryPlacingActorFromObject( GetViewportWorld()->GetCurrentLevel(), AssetObj, bSelectNewActors, ObjectFlags, FactoryToUse );
 
 				if( NewActors.Num() > 0 )
 				{
@@ -2972,7 +301,7 @@ void FVREditorWorldInteraction::OnAssetDragStartedFromContentBrowser( const TArr
 								( bShouldInterpolateFromDragLocation ? 
 									VREd::SizeOfActorsOverContentBrowserThumbnail->GetFloat() :
 									VREd::OverrideSizeOfPlacedActors->GetFloat() )
-								* Owner.GetWorldScaleFactor();
+								* GetWorldScaleFactor();
 							NewActor->SetActorScale3D( FVector( DesiredSize / LocalBoundsSize ) );
 						}
 					}
@@ -3003,30 +332,26 @@ void FVREditorWorldInteraction::OnAssetDragStartedFromContentBrowser( const TArr
 
 		if( AllNewActors.Num() > 0 )
 		{
-			const FVRAction VRAction( EVRActionType::SelectAndMove_LightlyPressed, PlacingWithHandIndex );
+			const FViewportActionKeyInput Action( ViewportWorldActionTypes::SelectAndMove_LightlyPressed );
 
 			// Start dragging the new actor(s)
 			const bool bIsPlacingActors = true;
-			StartDraggingActors( 
-				VRAction, 
-				TransformGizmoActor->GetRootComponent(), 
-				PlaceAt,
-				bIsPlacingActors );
+			StartDraggingActors( PlacingWithInteractor, Action, GetTransformGizmoActor()->GetRootComponent(), PlaceAt, bIsPlacingActors );
 
 			// If we're interpolating, update the target transform of the actors to use our overridden size.  When
 			// we placed them we set their size to be 'thumbnail sized', and we want them to interpolate to
 			// their actual size in the world
 			if( bShouldInterpolateFromDragLocation )
 			{
-				for( TUniquePtr<FVREditorTransformable>& TransformablePtr : Transformables )
+				for( TUniquePtr<FViewportTransformable>& TransformablePtr : Transformables )
 				{
-					FVREditorTransformable& Transformable = *TransformablePtr;
+					FViewportTransformable& Transformable = *TransformablePtr;
 
 					float ObjectSize = 1.0f;
 
 					if( VREd::OverrideSizeOfPlacedActors->GetFloat() > KINDA_SMALL_NUMBER )
 					{
-						ObjectSize = 1.0f * Owner.GetWorldScaleFactor();
+						ObjectSize = 1.0f * GetWorldScaleFactor();
 
 						// Got an actor?
 						AActor* Actor = Cast<AActor>( Transformable.Object.Get() );
@@ -3035,7 +360,7 @@ void FVREditorWorldInteraction::OnAssetDragStartedFromContentBrowser( const TArr
 							const FBox LocalSpaceBounds = Actor->CalculateComponentsBoundingBoxInLocalSpace();
 							const float LocalBoundsSize = LocalSpaceBounds.GetSize().GetAbsMax();
 
-							const float DesiredSize = VREd::OverrideSizeOfPlacedActors->GetFloat() * Owner.GetWorldScaleFactor();
+							const float DesiredSize = VREd::OverrideSizeOfPlacedActors->GetFloat() * GetWorldScaleFactor();
 							const FVector NewScale( DesiredSize / LocalBoundsSize );
 
 							Transformable.UnsnappedTargetTransform.SetScale3D( NewScale );
@@ -3050,22 +375,18 @@ void FVREditorWorldInteraction::OnAssetDragStartedFromContentBrowser( const TArr
 
 		if( DraggingSingleMaterialOrTexture != nullptr )
 		{
-			const FVRAction VRAction( EVRActionType::SelectAndMove_LightlyPressed, PlacingWithHandIndex );
+			const FViewportActionKeyInput Action( ViewportWorldActionTypes::SelectAndMove_LightlyPressed );
 
 			// Start dragging the material
-			StartDraggingMaterialOrTexture(
-				VRAction,
-				Hand.HoverLocation,
-				DraggingSingleMaterialOrTexture );
+			StartDraggingMaterialOrTexture( PlacingWithInteractor,	Action, PlacingWithInteractor->GetHoverLocation(), DraggingSingleMaterialOrTexture );
 		}
 	}
 }
 
 
-void FVREditorWorldInteraction::PlaceDraggedMaterialOrTexture( const int32 HandIndex )
+void UVREditorWorldInteraction::PlaceDraggedMaterialOrTexture( UViewportInteractor* Interactor )
 {
-	FVirtualHand& Hand = Owner.GetVirtualHand( HandIndex );
-	if( ensure( Hand.DraggingMode == EVREditorDraggingMode::Material ) &&
+	if( ensure( Interactor->GetDraggingMode() == EViewportInteractionDraggingMode::Material ) &&
 		PlacingMaterialOrTextureAsset != nullptr )
 	{
 		// Check to see if the laser pointer is over something we can drop on
@@ -3073,7 +394,7 @@ void FVREditorWorldInteraction::PlaceDraggedMaterialOrTexture( const int32 HandI
 		{
 			const bool bIgnoreGizmos = true;	// Never place on top of gizmos, just ignore them
 			const bool bEvenIfUIIsInFront = true;	// Don't let the UI block placement
-			FHitResult HitResult = Owner.GetHitResultFromLaserPointer( HandIndex, nullptr, bIgnoreGizmos, bEvenIfUIIsInFront );
+			FHitResult HitResult = Interactor->GetHitResultFromLaserPointer( nullptr, bIgnoreGizmos, nullptr, bEvenIfUIIsInFront );
 			if( HitResult.Actor.IsValid() )
 			{
 				if( IsInteractableComponent( HitResult.GetComponent() ) )	// @todo vreditor placement: We don't necessarily need to restrict to only VR-interactive components here
@@ -3087,7 +408,6 @@ void FVREditorWorldInteraction::PlaceDraggedMaterialOrTexture( const int32 HandI
 			}
 		}
 
-
 		if( HitComponent != nullptr )
 		{
 			UObject* ObjToUse = PlacingMaterialOrTextureAsset;
@@ -3112,7 +432,7 @@ void FVREditorWorldInteraction::PlaceDraggedMaterialOrTexture( const int32 HandI
 
 			if( DroppedObjAsMaterial || DroppedObjAsTexture )
 			{
-				UGameplayStatics::PlaySound2D( Owner.GetWorld(), DropMaterialOrMaterialSound );
+				UGameplayStatics::PlaySound2D( GetViewportWorld(), DropMaterialOrMaterialSound );
 			}
 		}
 	}
@@ -3120,140 +440,20 @@ void FVREditorWorldInteraction::PlaceDraggedMaterialOrTexture( const int32 HandI
 	this->PlacingMaterialOrTextureAsset = nullptr;
 }
 
-
-bool FVREditorWorldInteraction::FindPlacementPointUnderLaser( const int32 HandIndex, FVector& OutHitLocation )
-{
-	// Check to see if the laser pointer is over something we can drop on
-	bool bHitSomething = false;
-	FVector HitLocation = FVector::ZeroVector;
-
-	static TArray<AActor*> IgnoredActors;
-	{
-		IgnoredActors.Reset();
-
-		// Don't trace against stuff that we're dragging!
-		for( TUniquePtr<FVREditorTransformable>& TransformablePtr : Transformables )
-		{
-			FVREditorTransformable& Transformable = *TransformablePtr;
-
-			// Got an actor?
-			AActor* Actor = Cast<AActor>( Transformable.Object.Get() );
-			if( Actor != nullptr )
-			{
-				IgnoredActors.Add( Actor );
-			}
-		}
-
-		// Ignore UI widgets too
-		for( TActorIterator<AVREditorFloatingUI> UIActorIt( Owner.GetWorld() ); UIActorIt; ++UIActorIt )
-		{
-			IgnoredActors.Add( *UIActorIt );
-		}
-	}
-
-
-	const bool bIgnoreGizmos = true;	// Never place on top of gizmos, just ignore them
-	const bool bEvenIfUIIsInFront = true;	// Don't let the UI block placement
-	FHitResult HitResult = Owner.GetHitResultFromLaserPointer( HandIndex, &IgnoredActors, bIgnoreGizmos, bEvenIfUIIsInFront );
-	if( HitResult.Actor.IsValid() )
-	{
-		bHitSomething = true;
-		HitLocation = HitResult.ImpactPoint;
-	}
-
-	if( bHitSomething )
-	{
-		FVirtualHand& Hand = Owner.GetVirtualHand( HandIndex );
-
-		// Pull back from the impact point
-		{
-			const FVector GizmoSpaceImpactNormal = Hand.GizmoStartTransform.InverseTransformVectorNoScale( HitResult.ImpactNormal );
-
-			FVector ExtremePointOnBox;
-			const FVector ExtremeDirection = -GizmoSpaceImpactNormal;
-			const FBox& Box = GizmoLocalBounds;
-			{
-				const FVector ExtremePoint(
-					ExtremeDirection.X >= 0.0f ? Box.Max.X : Box.Min.X,
-					ExtremeDirection.Y >= 0.0f ? Box.Max.Y : Box.Min.Y,
-					ExtremeDirection.Z >= 0.0f ? Box.Max.Z : Box.Min.Z );
-				const float ProjectionDistance = FVector::DotProduct( ExtremePoint, ExtremeDirection );
-				ExtremePointOnBox = ExtremeDirection * ProjectionDistance;
-			}
-
-			const FVector WorldSpaceExtremePointOnBox = Hand.GizmoStartTransform.TransformVectorNoScale( ExtremePointOnBox );
-			HitLocation -= WorldSpaceExtremePointOnBox;
-		}
-
-		OutHitLocation = HitLocation;
-	}
-
-	return bHitSomething;
-}
-
-void FVREditorWorldInteraction::StartTeleport( const int32 HandIndex  )
+void UVREditorWorldInteraction::SnapSelectedActorsToGround()
 {
-	FVirtualHand Hand = Owner.GetVirtualHand( HandIndex );
-	if( !Hand.bIsHoveringOverUI )
+	TSharedPtr<SLevelViewport> LevelEditorViewport = StaticCastSharedPtr<SLevelViewport>( GetViewportClient()->GetEditorViewportWidget() );
+	if ( LevelEditorViewport.IsValid() )
 	{
-		FVector LaserPointerStart, LaserPointerEnd;
-		if (Owner.GetLaserPointer( HandIndex, LaserPointerStart, LaserPointerEnd ))
-		{
-			FHitResult HitResult = Owner.GetHitResultFromLaserPointer( HandIndex, nullptr, true, false, VREd::TeleportLaserPointerLength->GetFloat() );
-			if ( HitResult.bBlockingHit )
-			{
-				const FVector Offset = (LaserPointerStart - HitResult.Location).GetSafeNormal() * VREd::TeleportOffset->GetFloat() * Owner.GetWorldScaleFactor();	// Offset from the hitresult, TeleportOffset is in centimeter
-				const FVector NewHandLocation = HitResult.Location + Offset;																					// New world location offset from the hitresult
-				const FVector HandHeadOffset = Owner.GetHeadTransform().GetLocation() - Hand.Transform.GetLocation();											// The offset between the hand and the head
-				const FVector NewHeadLocation = NewHandLocation + HandHeadOffset;																				// The location in world space where the head is going to be
-
-				// The new roomspace location in world space
-				const FVector HeadRoomOffset = Owner.GetHeadTransform().GetLocation() - Owner.GetRoomTransform().GetLocation();									// Where the head is going to be in world space
-				const FVector NewRoomSpaceLocation = NewHeadLocation - HeadRoomOffset;																			// Result, where the roomspace is going to be in world space
-				
-				// Set values to start teleporting
-				bIsTeleporting = true;
-				TeleportLerpAlpha = 0.0f;
-				TeleportStartLocation = Owner.GetRoomTransform().GetLocation();
-				TeleportGoalLocation = NewRoomSpaceLocation;
-
-				UGameplayStatics::PlaySound2D( Owner.GetWorld(), TeleportSound );
-			}
-		}
-	}
-}
+		FLevelEditorModule& LevelEditorModule = FModuleManager::GetModuleChecked<FLevelEditorModule>( TEXT( "LevelEditor" ) );
+		const FLevelEditorCommands& Commands = LevelEditorModule.GetLevelEditorCommands();
+		const TSharedPtr< FUICommandList >& CommandList = LevelEditorViewport->GetParentLevelEditor().Pin()->GetLevelEditorActions(); //@todo vreditor: Cast on leveleditor
 
+		CommandList->ExecuteAction( Commands.SnapBottomCenterBoundsToFloor.ToSharedRef() );
 
-void FVREditorWorldInteraction::Teleport( const float DeltaTime )
-{
-	FTransform RoomTransform = Owner.GetRoomTransform();
-	TeleportLerpAlpha += DeltaTime;
-	
-	if( TeleportLerpAlpha > VREd::TeleportLerpTime->GetFloat() )
-	{
-		// Teleporting is finished
-		TeleportLerpAlpha = VREd::TeleportLerpTime->GetFloat();
-		bIsTeleporting = false;
+		// @todo vreditor: This should not be needed after to allow transformables to stop animating the transforms of actors after they come to rest
+		SetupTransformablesForSelectedActors();
 	}
-
-	// Calculate the new position of the roomspace
-	FVector NewLocation = FMath::Lerp( TeleportStartLocation, TeleportGoalLocation, TeleportLerpAlpha / VREd::TeleportLerpTime->GetFloat() );
-	RoomTransform.SetLocation( NewLocation );
-	Owner.SetRoomTransform( RoomTransform );
-}
-
-
-bool FVREditorWorldInteraction::IsSmoothSnappingEnabled() const
-{
-	const float SmoothSnapSpeed = VREd::SmoothSnapSpeed->GetFloat();
-	const bool bSmoothSnappingEnabled =
-		!GEditor->bIsSimulatingInEditor &&
-		( FSnappingUtils::IsSnapToGridEnabled() || FSnappingUtils::IsRotationSnapEnabled() || FSnappingUtils::IsScaleSnapEnabled() ) &&
-		VREd::SmoothSnap->GetInt() != 0 &&
-		!FMath::IsNearlyZero( SmoothSnapSpeed );
-
-	return bSmoothSnappingEnabled;
 }
 
-
 #undef LOCTEXT_NAMESPACE
\ No newline at end of file
