VR Editor: - Removed AutoConsoleVariable duplicates
- Fixed mac errors: declarations that shadows another field
- Fixed mac errors: Constructor initialization list order

diff --git a/Engine/Source/Editor/MeshPaint/Private/MeshPaintEdMode.cpp b/Engine/Source/Editor/MeshPaint/Private/MeshPaintEdMode.cpp
--- a/Engine/Source/Editor/MeshPaint/Private/MeshPaintEdMode.cpp
+++ b/Engine/Source/Editor/MeshPaint/Private/MeshPaintEdMode.cpp
@@ -2681,21 +2681,17 @@ void FEdModeMeshPaint::Render( const FSceneView* View, FViewport* Viewport, FPri
 						float StrengthScale = ( PaintingWithInteractorInVR == nullptr && FMeshPaintSettings::Get().bEnableFlow ) ? FMeshPaintSettings::Get().FlowAmount : 1.0f;
 
 						// Apply VR controller trigger pressure if we're painting in VR
-						if ( bIsPainting && PaintingWithInteractorInVR )
+						if ( bIsPainting && PaintingWithInteractorInVR && VREditorMode != nullptr && PaintingWithInteractorInVR )
 						{
-							IVREditorMode* VREditorMode = static_cast< IVREditorMode* >( GetModeManager()->GetActiveMode( IVREditorModule::Get().GetVREditorModeID() ) );
-							if ( VREditorMode != nullptr && PaintingWithInteractorInVR )
+							UVREditorInteractor* VRInteractor = Cast<UVREditorInteractor>( PaintingWithInteractorInVR );
+							if ( VRInteractor )
 							{
-								UVREditorInteractor* VRInteractor = Cast<UVREditorInteractor>( PaintingWithInteractorInVR );
-								if ( VRInteractor )
-								{
-									StrengthScale *= VRInteractor->GetSelectAndMoveTriggerValue();
+								StrengthScale *= VRInteractor->GetSelectAndMoveTriggerValue();
 
-									// Make sure light press locking is enabled in VR for the hand we're painting with.  We don't want
-									// a full press to be converted to a move action, however locking will have been disabled already
-									// by the world interaction code after a light press happened on an object. 
-									VRInteractor->SetAllowTriggerLightPressLocking( true );
-								}
+								// Make sure light press locking is enabled in VR for the hand we're painting with.  We don't want
+								// a full press to be converted to a move action, however locking will have been disabled already
+								// by the world interaction code after a light press happened on an object. 
+								VRInteractor->SetAllowTriggerLightPressLocking( true );
 							}
 						}
 
diff --git a/Engine/Source/Editor/VREditor/Teleporter/VREditorTeleporter.cpp b/Engine/Source/Editor/VREditor/Teleporter/VREditorTeleporter.cpp
--- a/Engine/Source/Editor/VREditor/Teleporter/VREditorTeleporter.cpp
+++ b/Engine/Source/Editor/VREditor/Teleporter/VREditorTeleporter.cpp
@@ -17,12 +17,12 @@ namespace VREd
 
 UVREditorTeleporter::UVREditorTeleporter( const FObjectInitializer& Initializer ) :
 	Super( Initializer ),
+	VRMode( nullptr ),
 	bIsTeleporting( false ),
 	TeleportLerpAlpha( 0 ),
 	TeleportStartLocation( FVector::ZeroVector ),
 	TeleportGoalLocation( FVector::ZeroVector ),
-	bJustTeleported( false ),
-	VRMode( nullptr )
+	bJustTeleported( false )
 {
 	// Load sounds
 	TeleportSound = LoadObject<USoundCue>( nullptr, TEXT( "/Engine/VREditor/Sounds/VR_teleport_Cue" ) );
diff --git a/Engine/Source/Editor/VREditor/VREditorFloatingUI.cpp b/Engine/Source/Editor/VREditor/VREditorFloatingUI.cpp
--- a/Engine/Source/Editor/VREditor/VREditorFloatingUI.cpp
+++ b/Engine/Source/Editor/VREditor/VREditorFloatingUI.cpp
@@ -15,7 +15,6 @@ namespace VREd
 {
 	static FAutoConsoleVariable UIFadeSpeed( TEXT( "VREd.UIFadeSpeed" ), 6.0f, TEXT( "How fast UI should fade in and out" ) );
 	static FAutoConsoleVariable UIBrightness( TEXT( "VREd.UIBrightness" ), 0.3f, TEXT( "How bright the UI should be" ) );
-	static FAutoConsoleVariable UIOnHandRotationOffset( TEXT( "VREd.UIOnHandRotationOffset" ), 45.0f, TEXT( "Rotation offset for UI that's docked to your hand, to make it more comfortable to hold" ) );
 }
 
 
diff --git a/Engine/Source/Editor/VREditor/VREditorMode.cpp b/Engine/Source/Editor/VREditor/VREditorMode.cpp
--- a/Engine/Source/Editor/VREditor/VREditorMode.cpp
+++ b/Engine/Source/Editor/VREditor/VREditorMode.cpp
@@ -39,8 +39,6 @@ namespace VREd
 	static FAutoConsoleVariable GridMaxOpacity( TEXT( "VREd.GridMaxFade" ), 0.8f, TEXT( "Grid maximum opacity" ) );
 	static FAutoConsoleVariable GridHeightOffset( TEXT( "VREd.GridHeightOffset" ), 0.0f, TEXT( "Height offset for the world movement grid.  Useful when tracking space is not properly calibrated" ) );
 
-	static FAutoConsoleVariable TriggerLightlyPressedLockTime( TEXT( "VREd.TriggerLightlyPressedLockTime" ), 0.15f, TEXT( "If the trigger remains lightly pressed for longer than this, we'll continue to treat it as a light press in some cases" ) );
-
 	static FAutoConsoleVariable WorldMovementFogOpacity( TEXT( "VREd.WorldMovementFogOpacity" ), 0.8f, TEXT( "How opaque the fog should be at the 'end distance' (0.0 - 1.0)" ) );
 	static FAutoConsoleVariable WorldMovementFogStartDistance( TEXT( "VREd.WorldMovementFogStartDistance" ), 300.0f, TEXT( "How far away fog will start rendering while in world movement mode" ) );
 	static FAutoConsoleVariable WorldMovementFogEndDistance( TEXT( "VREd.WorldMovementFogEndDistance" ), 800.0f, TEXT( "How far away fog will finish rendering while in world movement mode" ) );
@@ -49,7 +47,6 @@ namespace VREd
 	static FAutoConsoleVariable ScaleProgressBarLength( TEXT( "VREd.ScaleProgressBarLength" ), 50.0f, TEXT( "Length of the progressbar that appears when scaling" ) );
 	static FAutoConsoleVariable ScaleProgressBarRadius( TEXT( "VREd.ScaleProgressBarRadius" ), 1.0f, TEXT( "Radius of the progressbar that appears when scaling" ) );
 
-	static FAutoConsoleVariable InvertTrackpadVertical( TEXT( "VREd.InvertTrackpadVertical" ), 1, TEXT( "Toggles inverting the touch pad vertical axis" ) );
 	static FAutoConsoleVariable ForceOculusMirrorMode( TEXT( "VREd.ForceOculusMirrorMode" ), 3, TEXT( "Which Oculus display mirroring mode to use (see MirrorWindowModeType in OculusRiftHMD.h)" ) );
 
 	static FAutoConsoleVariable ShowMovementGrid( TEXT( "VREd.ShowMovementGrid" ), 1, TEXT( "Showing the ground movement grid" ) );
diff --git a/Engine/Source/Editor/VREditor/VREditorUISystem.cpp b/Engine/Source/Editor/VREditor/VREditorUISystem.cpp
--- a/Engine/Source/Editor/VREditor/VREditorUISystem.cpp
+++ b/Engine/Source/Editor/VREditor/VREditorUISystem.cpp
@@ -1067,36 +1067,37 @@ FTransform UVREditorUISystem::MakeDockableUITransform( AVREditorDockableWindow*
 	
 	const FTransform UIToWorld = UIToUIOnLaser * UIOnLaserToWorld;
 
-	if( DraggingUI != nullptr &&  DraggingUI->GetDockedTo() == AVREditorFloatingUI::EDockedTo::Dragging && InteractorDraggingUI )
+	if( DraggingUI != nullptr &&  DraggingUI->GetDockedTo() == AVREditorFloatingUI::EDockedTo::Dragging && InteractorDraggingUI && Interactor == InteractorDraggingUI )
 	{
 		UVREditorInteractor* OtherInteractor = Cast<UVREditorInteractor>( InteractorDraggingUI->GetOtherInteractor() );
-		
-		const float WorldScaleFactor = GetOwner().GetWorldScaleFactor();
-		const FVector DraggingUILocation = InteractorDraggingUI->GetHoverLocation();
-		const FVector OtherHandLocation =  OtherInteractor->GetTransform().GetLocation();
-		LastDraggingHoverLocation = UIOnLaserToWorld.GetLocation();
+		if ( OtherInteractor )
+		{
+			const float WorldScaleFactor = GetOwner().GetWorldScaleFactor();
+			const FVector DraggingUILocation = InteractorDraggingUI->GetHoverLocation();
+			const FVector OtherHandLocation = OtherInteractor->GetTransform().GetLocation();
+			LastDraggingHoverLocation = UIOnLaserToWorld.GetLocation();
 
-		const float Distance = FVector::Dist( UIOnLaserToWorld.GetLocation(), OtherHandLocation ) / WorldScaleFactor;
+			const float Distance = FVector::Dist( UIOnLaserToWorld.GetLocation(), OtherHandLocation ) / WorldScaleFactor;
 
-		// if dragged passed the threshold since starting dragging
-		if ( Distance > VREd::MinDockDragDistance->GetFloat() && !bDraggedDockFromHandPassedThreshold && 
-			( DraggingUI->GetPreviouslyDockedTo() == AVREditorFloatingUI::EDockedTo::LeftHand || DraggingUI->GetPreviouslyDockedTo() == AVREditorFloatingUI::EDockedTo::RightHand ) )
-		{
-			bDraggedDockFromHandPassedThreshold = true;
-		}
+			// if dragged passed the threshold since starting dragging
+			if ( Distance > VREd::MinDockDragDistance->GetFloat() && !bDraggedDockFromHandPassedThreshold &&
+				( DraggingUI->GetPreviouslyDockedTo() == AVREditorFloatingUI::EDockedTo::LeftHand || DraggingUI->GetPreviouslyDockedTo() == AVREditorFloatingUI::EDockedTo::RightHand ) )
+			{
+				bDraggedDockFromHandPassedThreshold = true;
+			}
 
-		// Snapping to a hand when in range
-		if ( Distance <= VREd::MinDockDragDistance->GetFloat() && ( bDraggedDockFromHandPassedThreshold || DraggingUI->GetPreviouslyDockedTo() == AVREditorFloatingUI::EDockedTo::Room ) )
-		{
-			UVREditorInteractor* OtherInteractor = Cast<UVREditorInteractor>( Interactor->GetOtherInteractor() );
-			const FVector EditorUIRelativeOffset( DraggingUI->GetSize().Y * 0.5f + VREd::EditorUIOffsetFromHand->GetFloat(), 0.0f, 0.0f );
-			const FTransform MovingUIGoalTransform = DraggingUI->MakeUITransformLockedToHand( OtherInteractor, false, EditorUIRelativeOffset, FRotator( 90.0f, 180.0f, 0.0f ) );
-			const AVREditorFloatingUI::EDockedTo NewDockedTo = OtherInteractor->GetControllerSide() == EControllerHand::Left ? AVREditorFloatingUI::EDockedTo::LeftHand : AVREditorFloatingUI::EDockedTo::RightHand;
-			DraggingUI->MoveTo( MovingUIGoalTransform, VREd::DockDragSpeed->GetFloat(), NewDockedTo );
-		}
-		else
-		{
-			DraggingUI->MoveTo( UIToWorld, VREd::DockDragSpeed->GetFloat(), AVREditorFloatingUI::EDockedTo::Room );
+			// Snapping to a hand when in range
+			if ( Distance <= VREd::MinDockDragDistance->GetFloat() && ( bDraggedDockFromHandPassedThreshold || DraggingUI->GetPreviouslyDockedTo() == AVREditorFloatingUI::EDockedTo::Room ) )
+			{
+				const FVector EditorUIRelativeOffset( DraggingUI->GetSize().Y * 0.5f + VREd::EditorUIOffsetFromHand->GetFloat(), 0.0f, 0.0f );
+				const FTransform MovingUIGoalTransform = DraggingUI->MakeUITransformLockedToHand( OtherInteractor, false, EditorUIRelativeOffset, FRotator( 90.0f, 180.0f, 0.0f ) );
+				const AVREditorFloatingUI::EDockedTo NewDockedTo = OtherInteractor->GetControllerSide() == EControllerHand::Left ? AVREditorFloatingUI::EDockedTo::LeftHand : AVREditorFloatingUI::EDockedTo::RightHand;
+				DraggingUI->MoveTo( MovingUIGoalTransform, VREd::DockDragSpeed->GetFloat(), NewDockedTo );
+			}
+			else
+			{
+				DraggingUI->MoveTo( UIToWorld, VREd::DockDragSpeed->GetFloat(), AVREditorFloatingUI::EDockedTo::Room );
+			}
 		}
 	}
 
diff --git a/Engine/Source/Editor/VREditor/VREditorWorldInteraction.cpp b/Engine/Source/Editor/VREditor/VREditorWorldInteraction.cpp
--- a/Engine/Source/Editor/VREditor/VREditorWorldInteraction.cpp
+++ b/Engine/Source/Editor/VREditor/VREditorWorldInteraction.cpp
@@ -28,8 +28,6 @@ namespace VREd
 {
 	static FAutoConsoleVariable SizeOfActorsOverContentBrowserThumbnail( TEXT( "VREd.SizeOfActorsOverContentBrowserThumbnail" ), 6.0f, TEXT( "How large objects should be when rendered 'thumbnail size' over the Content Browser" ) );
 	static FAutoConsoleVariable OverrideSizeOfPlacedActors( TEXT( "VREd.OverrideSizeOfPlacedActors" ), 0.0f, TEXT( "If set to a value greater than zero, sets the size in cm of a placed actor (relative to your world's scale)" ) );
-	static FAutoConsoleVariable TrackpadAbsoluteDragSpeed( TEXT( "VREd.TrackpadAbsoluteDragSpeed" ), 40.0f, TEXT( "How fast objects move toward or away when you drag on the touchpad while carrying them" ) );
-	static FAutoConsoleVariable TrackpadRelativeDragSpeed( TEXT( "VREd.TrackpadRelativeDragSpeed" ), 8.0f, TEXT( "How fast objects move toward or away when you hold a direction on an analog stick while carrying them" ) );
 	static FAutoConsoleVariable HoverHapticFeedbackStrength( TEXT( "VREd.HoverHapticFeedbackStrength" ), 0.1f, TEXT( "Default strength for haptic feedback when hovering" ) );
 	static FAutoConsoleVariable HoverHapticFeedbackTime( TEXT( "VREd.HoverHapticFeedbackTime" ), 0.2f, TEXT( "The minimum time between haptic feedback for hovering" ) );
 	static FAutoConsoleVariable PivotPointTransformGizmo( TEXT( "VREd.PivotPointTransformGizmo" ), 1, TEXT( "If the pivot point transform gizmo is used instead of the bounding box gizmo" ) );
@@ -40,9 +38,9 @@ namespace VREd
 UVREditorWorldInteraction::UVREditorWorldInteraction( const FObjectInitializer& Initializer ) : 
 	Super( Initializer ),
 	Owner( nullptr ),
-	PlacingMaterialOrTextureAsset( nullptr ),
 	DropMaterialOrMaterialSound( nullptr ),
-	FloatingUIAssetDraggedFrom( nullptr )
+	FloatingUIAssetDraggedFrom( nullptr ),
+	PlacingMaterialOrTextureAsset( nullptr )
 {
 	// Find out when the user drags stuff out of a content browser
 	FEditorDelegates::OnAssetDragStarted.AddUObject( this, &UVREditorWorldInteraction::OnAssetDragStartedFromContentBrowser );
diff --git a/Engine/Source/Editor/ViewportInteraction/Public/ViewportWorldInteraction.h b/Engine/Source/Editor/ViewportInteraction/Public/ViewportWorldInteraction.h
--- a/Engine/Source/Editor/ViewportInteraction/Public/ViewportWorldInteraction.h
+++ b/Engine/Source/Editor/ViewportInteraction/Public/ViewportWorldInteraction.h
@@ -414,12 +414,6 @@ private:
 	/** Duration to interpolate transformables over */
 	float TransformablesInterpolationDuration;
 
-	/** Event that fires every frame to update hover state based on what's under the cursor.  Set bWasHandled to true if you detect something to hover. */
-	FOnVIHoverUpdate OnHoverUpdateEvent;
-
-	/** Event that is fired for when a key is pressed by an interactor */
-	FOnVIActionHandle OnInputActionEvent;
-
 	//
 	// Transform gizmo
 	//
@@ -471,5 +465,11 @@ private:
 	//
 
 	/** The current dragged interactable */
-	class IViewportInteractableInterface* DraggedInteractable;	
+	class IViewportInteractableInterface* DraggedInteractable;		
+	
+	/** Event that fires every frame to update hover state based on what's under the cursor.  Set bWasHandled to true if you detect something to hover. */
+	FOnVIHoverUpdate OnHoverUpdateEvent;
+
+	/** Event that is fired for when a key is pressed by an interactor */
+	FOnVIActionHandle OnInputActionEvent;
 };
\ No newline at end of file
diff --git a/Engine/Source/Editor/ViewportInteraction/ViewportWorldInteraction.cpp b/Engine/Source/Editor/ViewportInteraction/ViewportWorldInteraction.cpp
--- a/Engine/Source/Editor/ViewportInteraction/ViewportWorldInteraction.cpp
+++ b/Engine/Source/Editor/ViewportInteraction/ViewportWorldInteraction.cpp
@@ -209,10 +209,13 @@ static void SegmentDistToSegmentDouble( DVector A1, DVector B1, DVector A2, DVec
 
 UViewportWorldInteraction::UViewportWorldInteraction( const FObjectInitializer& Initializer ) :
 	Super( Initializer ),
+	bDraggedSinceLastSelection( false ),
+	LastDragGizmoStartTransform( FTransform::Identity ),
+	TrackingTransaction( FTrackingTransaction() ),
 	AppTimeEntered( FTimespan::Zero() ),
+	EditorViewportClient( nullptr ),
 	LastFrameNumberInputWasPolled( 0 ),
 	MotionControllerID( 0 ),	// @todo ViewportInteraction: We only support a single controller, and we assume the first controller are the motion controls
-	EditorViewportClient( nullptr ),
 	LastWorldToMetersScale( 100.0f ),
 	bIsInterpolatingTransformablesFromSnapshotTransform( false ),
 	TransformablesInterpolationStartTime( FTimespan::Zero() ),
@@ -220,8 +223,6 @@ UViewportWorldInteraction::UViewportWorldInteraction( const FObjectInitializer&
 	TransformGizmoActor( nullptr ),
 	TransformGizmoClass( APivotTransformGizmo::StaticClass() ),
 	GizmoLocalBounds( FBox( 0 ) ),
-	bDraggedSinceLastSelection( false ),
-	LastDragGizmoStartTransform( FTransform::Identity ),
 	StartDragAngleOnRotation( ),
 	StartDragHandleDirection( ),
 	CurrentGizmoType( EGizmoHandleTypes::All ),
