VR Editor:
- Removed InputEvent parameter from the ViewportFOnVIActionHandle, because this data was already stored in the FViewportActionKeyInput that is also passed
- Fixed typo ViewportDragOperation SetDragOperationClass
- Removed OnDrag from interactables
- Fixed mac warning with order of ViewportWorldInteraction constructor initialization list
- Added events for AVREditorButton
- Added Switching between VR Editor to VR PIE:
  - Changed quickmenu to add PIE in VR button
  - Module starts PIE when the reason for exiting the mode was to start PIE session
  - Module sets the initial worldtometers for the VR Editor and stores it, so when exiting the VR Editor mode the worldtometers is set back from what it was before VR Editor.

diff --git a/Engine/Source/Editor/MeshPaint/Private/MeshPaintEdMode.cpp b/Engine/Source/Editor/MeshPaint/Private/MeshPaintEdMode.cpp
--- a/Engine/Source/Editor/MeshPaint/Private/MeshPaintEdMode.cpp
+++ b/Engine/Source/Editor/MeshPaint/Private/MeshPaintEdMode.cpp
@@ -5060,7 +5060,7 @@ TWeakObjectPtr<AActor> FEdModeMeshPaint::GetEditingActor() const
 
 
 void FEdModeMeshPaint::OnVRAction( class FEditorViewportClient& ViewportClient, UViewportInteractor* Interactor, 
-	const FViewportActionKeyInput& Action, const EInputEvent Event, bool& bOutIsInputCaptured, bool& bWasHandled )
+	const FViewportActionKeyInput& Action, bool& bOutIsInputCaptured, bool& bWasHandled )
 {
 	IVREditorMode* VREditorMode = static_cast<IVREditorMode*>( GetModeManager()->GetActiveMode( IVREditorModule::Get().GetVREditorModeID() ) );
 	UVREditorInteractor* VRInteractor = Cast<UVREditorInteractor>( Interactor );
@@ -5079,7 +5079,7 @@ void FEdModeMeshPaint::OnVRAction( class FEditorViewportClient& ViewportClient,
 
 		if( Action.ActionType == ViewportWorldActionTypes::SelectAndMove_LightlyPressed )
 		{
-			if( !bIsPainting && Event == IE_Pressed && !VRInteractor->IsHoveringOverUI() )
+			if( !bIsPainting && Action.Event == IE_Pressed && !VRInteractor->IsHoveringOverUI() )
 			{
 				// NOTE: We intentionally don't capture input for "LightlyPressed" because we want to allow objects to still
 				// become selected when clicked on.  We avoid allowing the objects to be moved by calling the 
@@ -5109,7 +5109,7 @@ void FEdModeMeshPaint::OnVRAction( class FEditorViewportClient& ViewportClient,
 			}
 
 			// Stop current tracking if the user is no longer painting
-			else if( bIsPainting && Event == IE_Released && PaintingWithInteractorInVR && PaintingWithInteractorInVR == Interactor )
+			else if( bIsPainting && Action.Event == IE_Released && PaintingWithInteractorInVR && PaintingWithInteractorInVR == Interactor )
 			{
 				EndPainting();
 
diff --git a/Engine/Source/Editor/MeshPaint/Public/MeshPaintEdMode.h b/Engine/Source/Editor/MeshPaint/Public/MeshPaintEdMode.h
--- a/Engine/Source/Editor/MeshPaint/Public/MeshPaintEdMode.h
+++ b/Engine/Source/Editor/MeshPaint/Public/MeshPaintEdMode.h
@@ -716,7 +716,7 @@ private:
 
 	/** Called when the user presses a button on their motion controller device */
 	void OnVRAction( class FEditorViewportClient& ViewportClient, class UViewportInteractor* Interactor,
-		const struct FViewportActionKeyInput& Action, const EInputEvent Event, bool& bOutIsInputCaptured, bool& bWasHandled );
+		const struct FViewportActionKeyInput& Action, bool& bOutIsInputCaptured, bool& bWasHandled );
 
 private:
 
diff --git a/Engine/Source/Editor/VREditor/Interactables/VREditorButton.cpp b/Engine/Source/Editor/VREditor/Interactables/VREditorButton.cpp
--- a/Engine/Source/Editor/VREditor/Interactables/VREditorButton.cpp
+++ b/Engine/Source/Editor/VREditor/Interactables/VREditorButton.cpp
@@ -37,30 +37,25 @@ AVREditorButton::AVREditorButton()
 
 void AVREditorButton::OnPressed( UViewportInteractor* Interactor, const FHitResult& InHitResult, bool& bOutResultedInDrag  )
 {
-	GLog->Logf( TEXT( "OnPressed" ) );
+	OnPressedEvent.Broadcast( this, Interactor );
 }
 
 void AVREditorButton::OnHover( UViewportInteractor* Interactor )
 {
-	GLog->Logf( TEXT( "OnHover" ) );
+	OnHoverEvent.Broadcast( this, Interactor );
 }
 
 void AVREditorButton::OnHoverEnter( UViewportInteractor* Interactor, const FHitResult& InHitResult  )
 {
-	GLog->Logf( TEXT("OnHoverEnter") );
+	OnHoverEnterEvent.Broadcast( this, Interactor );
 }
 
 void AVREditorButton::OnHoverLeave( UViewportInteractor* Interactor, const UActorComponent* NewComponent )
 {
-	GLog->Logf( TEXT("OnHoverLeave") );
-}
-
-void AVREditorButton::OnDrag( UViewportInteractor* Interactor )
-{
-	GLog->Logf( TEXT( "OnDrag" ) );
+	OnHoverLeaveEvent.Broadcast( this, Interactor );
 }
 
 void AVREditorButton::OnDragRelease( UViewportInteractor* Interactor )
 {
-	GLog->Logf( TEXT( "OnDragRelease" ) );
-}
+	OnReleaseEvent.Broadcast( this, Interactor );
+}
\ No newline at end of file
diff --git a/Engine/Source/Editor/VREditor/Interactables/VREditorButton.h b/Engine/Source/Editor/VREditor/Interactables/VREditorButton.h
--- a/Engine/Source/Editor/VREditor/Interactables/VREditorButton.h
+++ b/Engine/Source/Editor/VREditor/Interactables/VREditorButton.h
@@ -9,6 +9,7 @@
 //Forward declarations
 class UViewportInteractor;
 
+
 /**
 * A button for VR Editor
 */
@@ -27,15 +28,48 @@ public:
 	virtual void OnHover( UViewportInteractor* Interactor ) override;
 	virtual void OnHoverEnter( UViewportInteractor* Interactor, const FHitResult& InHitResult ) override;
 	virtual void OnHoverLeave( UViewportInteractor* Interactor, const UActorComponent* NewComponent ) override;
-	virtual void OnDrag( UViewportInteractor* Interactor ) override;
 	virtual void OnDragRelease( UViewportInteractor* Interactor ) override;
 	virtual class UViewportDragOperationComponent* GetDragOperationComponent() override { return nullptr; };
 	// End IViewportInteractableInterface
 
+	/** Gets the on pressed event */
+	DECLARE_EVENT_TwoParams( AVREditorButton, FOnPressedEvent, AVREditorButton*, UViewportInteractor* );
+	FOnPressedEvent& GetOnPressed() { return OnPressedEvent; }
+
+	/** Gets the on hover event */
+	DECLARE_EVENT_TwoParams( AVREditorButton, FOnHoverEvent, AVREditorButton*, UViewportInteractor* );
+	FOnHoverEvent& GetOnHover() { return OnHoverEvent; }
+
+	/** Gets the on hover enter event */
+	DECLARE_EVENT_TwoParams( AVREditorButton, FOnHoverEnterEvent, AVREditorButton*, UViewportInteractor* );
+	FOnHoverEnterEvent& GetOnHoverEnter() { return OnHoverEnterEvent; }
+
+	/** Gets the on hover leave event */
+	DECLARE_EVENT_TwoParams( AVREditorButton, FOnHoverLeaveEvent, AVREditorButton*, UViewportInteractor* );
+	FOnHoverLeaveEvent	& GetOnHoverLeave() { return OnHoverLeaveEvent; }
+
+	/** Gets the on release event */
+	DECLARE_EVENT_TwoParams( AVREditorButton, FOnReleaseEvent, AVREditorButton*, UViewportInteractor* );
+	FOnReleaseEvent& GetOnRelease() { return OnReleaseEvent; }
+
 private:
 
 	/** Mesh for the button to interact with */
 	UPROPERTY()
 	UStaticMeshComponent* ButtonMeshComponent;
 
+	/** Broadcasted when pressed on this interactable by an interactor */
+	FOnPressedEvent OnPressedEvent;
+
+	/** Broadcasted when interactor is hovering over this interactable */
+	FOnHoverEvent OnHoverEvent;
+
+	/** Broadcasted when interactor started hovering over this interactable */
+	FOnHoverEnterEvent OnHoverEnterEvent;
+
+	/** Broadcasted when interactor stopped hovering over this interactable */
+	FOnHoverLeaveEvent OnHoverLeaveEvent;
+
+	/** Broadcasted when interactor released this interactable */
+	FOnReleaseEvent OnReleaseEvent;
 };
\ No newline at end of file
diff --git a/Engine/Source/Editor/VREditor/Teleporter/VREditorTeleporter.cpp b/Engine/Source/Editor/VREditor/Teleporter/VREditorTeleporter.cpp
--- a/Engine/Source/Editor/VREditor/Teleporter/VREditorTeleporter.cpp
+++ b/Engine/Source/Editor/VREditor/Teleporter/VREditorTeleporter.cpp
@@ -89,12 +89,12 @@ void UVREditorTeleporter::StartTeleport( UViewportInteractor* Interactor )
 	}
 }
 
-void UVREditorTeleporter::OnInteractorAction( class FEditorViewportClient& ViewportClient, UViewportInteractor* Interactor, const struct FViewportActionKeyInput& Action, const EInputEvent Event, bool& bOutIsInputCaptured, bool& bWasHandled )
+void UVREditorTeleporter::OnInteractorAction( class FEditorViewportClient& ViewportClient, UViewportInteractor* Interactor, const struct FViewportActionKeyInput& Action, bool& bOutIsInputCaptured, bool& bWasHandled )
 {
 	if ( ( Action.ActionType == ViewportWorldActionTypes::SelectAndMove ||
 		Action.ActionType == ViewportWorldActionTypes::SelectAndMove_LightlyPressed ) )
 	{
-		if ( Event == IE_Pressed )
+		if ( Action.Event == IE_Pressed )
 		{
 			if ( Action.ActionType == ViewportWorldActionTypes::SelectAndMove_LightlyPressed )
 			{
diff --git a/Engine/Source/Editor/VREditor/Teleporter/VREditorTeleporter.h b/Engine/Source/Editor/VREditor/Teleporter/VREditorTeleporter.h
--- a/Engine/Source/Editor/VREditor/Teleporter/VREditorTeleporter.h
+++ b/Engine/Source/Editor/VREditor/Teleporter/VREditorTeleporter.h
@@ -36,7 +36,7 @@ private:
 
 	/** Called when the user presses a button on their motion controller device */
 	void OnInteractorAction( class FEditorViewportClient& ViewportClient, UViewportInteractor* Interactor,
-		const struct FViewportActionKeyInput& Action, const EInputEvent Event, bool& bOutIsInputCaptured, bool& bWasHandled );
+		const struct FViewportActionKeyInput& Action, bool& bOutIsInputCaptured, bool& bWasHandled );
 
 	/** Actually teleport to a position */
 	void Teleport( const float DeltaTime );
diff --git a/Engine/Source/Editor/VREditor/VREditorDockableWindow.cpp b/Engine/Source/Editor/VREditor/VREditorDockableWindow.cpp
--- a/Engine/Source/Editor/VREditor/VREditorDockableWindow.cpp
+++ b/Engine/Source/Editor/VREditor/VREditorDockableWindow.cpp
@@ -147,7 +147,7 @@ AVREditorDockableWindow::AVREditorDockableWindow()
 
 	// Create the drag operation
 	DragOperationComponent = CreateDefaultSubobject<UViewportDragOperationComponent>( TEXT( "DragOperation" ) );
-	DragOperationComponent->SetDragOperactionClass( UDockableWindowDragOperation::StaticClass() );
+	DragOperationComponent->SetDragOperationClass( UDockableWindowDragOperation::StaticClass() );
 }
 
 
@@ -379,11 +379,6 @@ void AVREditorDockableWindow::OnHoverLeave( UViewportInteractor* Interactor, con
 	}
 }
 
-void AVREditorDockableWindow::OnDrag( UViewportInteractor* Interactor )
-{
-
-}
-
 void AVREditorDockableWindow::OnDragRelease( UViewportInteractor* Interactor )
 {
 	UVREditorInteractor* VREditorInteractor = Cast<UVREditorInteractor>( Interactor );
diff --git a/Engine/Source/Editor/VREditor/VREditorDockableWindow.h b/Engine/Source/Editor/VREditor/VREditorDockableWindow.h
--- a/Engine/Source/Editor/VREditor/VREditorDockableWindow.h
+++ b/Engine/Source/Editor/VREditor/VREditorDockableWindow.h
@@ -46,7 +46,6 @@ public:
 	
 	/** Leaving hover with laser changes the color of SelectionMesh and CloseButtonMesh */
 	virtual void OnHoverLeave( UViewportInteractor* Interactor, const UActorComponent* NewComponent ) override;
-	virtual void OnDrag( UViewportInteractor* Interactor ) override;
 	virtual void OnDragRelease( UViewportInteractor* Interactor ) override;
 	virtual class UViewportDragOperationComponent* GetDragOperationComponent() override;
 
diff --git a/Engine/Source/Editor/VREditor/VREditorMode.cpp b/Engine/Source/Editor/VREditor/VREditorMode.cpp
--- a/Engine/Source/Editor/VREditor/VREditorMode.cpp
+++ b/Engine/Source/Editor/VREditor/VREditorMode.cpp
@@ -96,7 +96,6 @@ FVREditorMode::~FVREditorMode()
 	FEditorDelegates::OnSwitchBeginPIEAndSIE.RemoveAll( this );
 }
 
-
 void FVREditorMode::Enter()
 {
 	// @todo vreditor urgent: Turn on global editor hacks for VR Editor mode
@@ -243,7 +242,7 @@ void FVREditorMode::Enter()
 			// @todo vreditor: Should save and restore camera position and any other settings we change (viewport type, pitch locking, etc.)
 			SavedEditorState.ViewLocation = VRViewportClient.GetViewLocation();
 			SavedEditorState.ViewRotation = VRViewportClient.GetViewRotation();
-
+			
 			// Don't allow the tracking space to pitch up or down.  People hate that in VR.
 			// @todo vreditor: This doesn't seem to prevent people from pitching the camera with RMB drag
 			SavedEditorState.bLockedPitch = VRViewportClient.GetCameraController()->GetConfig().bLockedPitch;
@@ -359,11 +358,6 @@ void FVREditorMode::Enter()
 	//	UISystem->ShowEditorUIPanel(FVREditorUISystem::EEditorUIPanel::ContentBrowser, HandIndex, bShouldShow);
 	//}
 
-	if ( bActuallyUsingVR )
-	{
-		 WorldInteraction->SetWorldToMetersScale( 100.0f );
-	}
-
 	// Call parent implementation
 	FEdMode::Enter();
 
@@ -427,7 +421,7 @@ void FVREditorMode::Exit()
 				VRViewportClient.bAlwaysShowModeWidgetAfterSelectionChanges = SavedEditorState.bAlwaysShowModeWidgetAfterSelectionChanges;
 				VRViewportClient.EngineShowFlags = SavedEditorState.ShowFlags;
 				VRViewportClient.SetGameView( SavedEditorState.bGameView );
-				VRViewportClient.SetViewLocation( GetHeadTransform().GetLocation() ); // Use SavedEditorState.ViewLocation to go back to start location when entering VR Editor Mode
+				VRViewportClient.SetViewLocation( GetHeadTransform().GetLocation() );
 				
 				FRotator HeadRotationNoRoll = GetHeadTransform().GetRotation().Rotator();
 				HeadRotationNoRoll.Roll = 0.0f;
@@ -464,9 +458,6 @@ void FVREditorMode::Exit()
 		}
 	}
 
-	// Make sure we are no longer overriding WorldToMetersScale every frame
-	WorldInteraction->SetWorldToMetersScale( 0.0f );
-
 	// Call parent implementation
 	FEdMode::Exit();
 
@@ -496,6 +487,11 @@ void FVREditorMode::Exit()
 	GEnableVREditorHacks = false;
 }
 
+void FVREditorMode::StartExitingVRMode( const EVREditorExitType InExitType /*= EVREditorExitType::To_Editor */ )
+{
+	ExitType = InExitType;
+	bWantsToExitMode = true;
+}
 
 void FVREditorMode::SpawnAvatarMeshActor()
 {
@@ -695,7 +691,28 @@ void FVREditorMode::Tick( FEditorViewportClient* ViewportClient, float DeltaTime
 	{
 		return;
 	}
-	
+
+	//Setting the initial position and rotation based on the editor viewport when going into VR mode
+	if ( bFirstTick )
+	{
+		const FTransform RoomToWorld = GetRoomTransform();
+		const FTransform WorldToRoom = RoomToWorld.Inverse();
+		FTransform ViewportToWorld = FTransform( SavedEditorState.ViewRotation, SavedEditorState.ViewLocation );
+		FTransform ViewportToRoom = ( ViewportToWorld * WorldToRoom );
+
+		FTransform ViewportToRoomYaw = ViewportToRoom;
+		ViewportToRoomYaw.SetRotation( FQuat( FRotator( 0.0f, ViewportToRoomYaw.GetRotation().Rotator().Yaw, 0.0f ) ) );
+
+		FTransform HeadToRoomYaw = GetRoomSpaceHeadTransform();
+		HeadToRoomYaw.SetRotation( FQuat( FRotator( 0.0f, HeadToRoomYaw.GetRotation().Rotator().Yaw, 0.0f ) ) );
+
+		FTransform RoomToWorldYaw = RoomToWorld;
+		RoomToWorldYaw.SetRotation( FQuat( FRotator( 0.0f, RoomToWorldYaw.GetRotation().Rotator().Yaw, 0.0f ) ) );
+
+		FTransform ResultToWorld = ( HeadToRoomYaw.Inverse() * ViewportToRoomYaw ) * RoomToWorldYaw;
+		SetRoomTransform( ResultToWorld );
+	}
+
 	if ( AvatarMeshActor == nullptr )
 	{
 		SpawnAvatarMeshActor();
@@ -928,27 +945,6 @@ void FVREditorMode::Tick( FEditorViewportClient* ViewportClient, float DeltaTime
 		}
 	}
 
-	// Setting the initial position and rotation based on the editor viewport when going into VR mode
-	if ( bActuallyUsingVR && bFirstTick )
-	{
-		const FTransform RoomToWorld = GetRoomTransform();
-		const FTransform WorldToRoom = RoomToWorld.Inverse();
-		const FTransform ViewportToWorld = FTransform( SavedEditorState.ViewRotation, SavedEditorState.ViewLocation );
-		const FTransform ViewportToRoom = ViewportToWorld * WorldToRoom;
-
-		FTransform ViewportToRoomYaw = ViewportToRoom;
-		ViewportToRoomYaw.SetRotation( FQuat( FRotator( 0.0f, ViewportToRoomYaw.GetRotation().Rotator().Yaw, 0.0f ) ) );
-
-		FTransform HeadToRoomYaw = GetRoomSpaceHeadTransform();
-		HeadToRoomYaw.SetRotation( FQuat( FRotator( 0.0f, HeadToRoomYaw.GetRotation().Rotator().Yaw, 0.0f ) ) );
-
-		FTransform RoomToWorldYaw = RoomToWorld;
-		RoomToWorldYaw.SetRotation( FQuat( FRotator( 0.0f, RoomToWorldYaw.GetRotation().Rotator().Yaw, 0.0f ) ) );
-
-		FTransform ResultToWorld = ( HeadToRoomYaw.Inverse() * ViewportToRoomYaw ) * RoomToWorldYaw;
-		SetRoomTransform( ResultToWorld );
-	}
-
 	StopOldHapticEffects();
 
 	bFirstTick = false;
diff --git a/Engine/Source/Editor/VREditor/VREditorMode.h b/Engine/Source/Editor/VREditor/VREditorMode.h
--- a/Engine/Source/Editor/VREditor/VREditorMode.h
+++ b/Engine/Source/Editor/VREditor/VREditorMode.h
@@ -10,6 +10,13 @@
 // Forward declare the GizmoHandleTypes that is defined in VREditorTransformGizmo
 enum class EGizmoHandleTypes : uint8;
 
+/** Reason for exiting the VR editor */
+enum EVREditorExitType
+{
+	Normal = 0,
+	PIE_VR = 1
+};
+
 /**
  * VR Editor Mode.  Extends editor viewports with functionality for VR controls and object manipulation
  */
@@ -48,9 +55,12 @@ public:
 	}
 
 	/** Call this to start exiting VR mode */
-	void StartExitingVRMode()
+	void StartExitingVRMode( const EVREditorExitType InExitType = EVREditorExitType::Normal );
+
+	/** Gets the reason for exiting the mode */
+	EVREditorExitType GetExitType() const
 	{
-		bWantsToExitMode = true;
+		return ExitType;
 	}
 
 	// FEdMode interface
@@ -383,5 +393,9 @@ private:
 
 	/** If the coordinate system was in world space before entering (local) scale mode */
 	bool bWasInWorldSpaceBeforeScaleMode;
+
+	/** The reason for exiting the mode, so the module can close the mode and take further actions on what should happen next */
+	EVREditorExitType ExitType;
+
 };
 
diff --git a/Engine/Source/Editor/VREditor/VREditorModule.cpp b/Engine/Source/Editor/VREditor/VREditorModule.cpp
--- a/Engine/Source/Editor/VREditor/VREditorModule.cpp
+++ b/Engine/Source/Editor/VREditor/VREditorModule.cpp
@@ -7,11 +7,14 @@
 #include "LevelEditor.h"
 #include "HeadMountedDisplay.h"
 
-
 class FVREditorModule : public IVREditorModule, public FTickableEditorObject
 {
 public:
-	FVREditorModule()
+	FVREditorModule() :
+		bEnableVRRequest( false ),
+		bPlayStartedFromVREditor( false ),
+		LastWorldToMeters( 100 ),
+		SavedWorldToMeters( 0 )
 	{
 	}
 
@@ -50,11 +53,22 @@ public:
 private:
 
 	void EnableVREditor( const bool bEnable, const bool bForceWithoutHMD );
-	void SetupMotionControls();
-};
 
+	/** Sets the world settings directly and prevents the last GNewWorldToMetersScale from the VREditor to reset it by setting it 0 */
+	void SetDirectWorldToMeters( const float NewWorldToMeters );
 
+	/** If the VR Editor mode needs to be enabled next tick */
+	bool bEnableVRRequest;
 
+	/** If PIE session started from within the VR Editor, so we can go back to the VR Editor when closing PIE with the motion controllers */
+	bool bPlayStartedFromVREditor;
+
+	/** WorldToMeters from last VR Editor session */
+	float LastWorldToMeters;
+
+	/** Saved last world to meters scale from last VR Editor session so we can restore it when entering the VR Editor when PIE starter from VR Editor */
+	float SavedWorldToMeters;
+};
 
 namespace VREd
 {
@@ -114,9 +128,19 @@ void FVREditorModule::EnableVREditor( const bool bEnable, const bool bForceWitho
 		{
 			FVREditorMode::SetActuallyUsingVR( !bForceWithoutHMD );
 
+			SavedWorldToMeters = GWorld->GetWorldSettings()->WorldToMeters;
+
+			// Set the WorldToMeters scale when we stopped playing PIE that was started by the VR Editor to that VR Editor sessions latest WorldToMeters
+			if ( bPlayStartedFromVREditor )
+			{
+				SetDirectWorldToMeters( LastWorldToMeters );
+				bPlayStartedFromVREditor = false;
+			}
+		
 			// Activate the mode right away.  We expect it to stay active forever!
 			GLevelEditorModeTools().AddDefaultMode( FVREditorMode::GetVREditorModeID() );
 			GLevelEditorModeTools().ActivateMode( FVREditorMode::GetVREditorModeID() );
+
 		}
 		else if( !bEnable )
 		{ 
@@ -126,6 +150,9 @@ void FVREditorModule::EnableVREditor( const bool bEnable, const bool bForceWitho
 				GLevelEditorModeTools().RemoveDefaultMode( FVREditorMode::GetVREditorModeID() );
 				GLevelEditorModeTools().DeactivateMode( FVREditorMode::GetVREditorModeID() );
 			}
+
+			SetDirectWorldToMeters( SavedWorldToMeters );
+			SavedWorldToMeters = 0.0f;
 		}
 	}
 }
@@ -138,6 +165,13 @@ void FVREditorModule::EnableVREditor( const bool bEnable )
 }
 
 
+void FVREditorModule::SetDirectWorldToMeters( const float NewWorldToMeters )
+{
+	GWorld->GetWorldSettings()->WorldToMeters = NewWorldToMeters;
+	ENGINE_API extern float GNewWorldToMetersScale;
+	GNewWorldToMetersScale = 0.0f;
+}
+
 void FVREditorModule::ToggleForceVRMode()
 {
 	const bool bForceWithoutHMD = true;
@@ -155,7 +189,50 @@ void FVREditorModule::Tick( float DeltaTime )
 
 		if( VREditorMode->WantsToExitMode() )
 		{
+			//Store the current WorldToMeters before exiting the mode
+			LastWorldToMeters = GWorld->GetWorldSettings()->WorldToMeters;
+
 			EnableVREditor( false );
+
+			// Start the session if that was the reason to stop the VR Editor mode
+			if ( VREditorMode->GetExitType() == EVREditorExitType::PIE_VR )
+			{
+				const bool bHMDIsReady = ( GEngine && GEngine->HMDDevice.IsValid() && GEngine->HMDDevice->IsHMDConnected() );
+
+				const FVector* StartLoc = NULL;
+				const FRotator* StartRot = NULL;
+				GEditor->RequestPlaySession( true, NULL, false, StartLoc, StartRot, -1, false, bHMDIsReady, false );
+				bPlayStartedFromVREditor = true;
+			}
+		}
+	}
+	// Only check for input if we started this play session from the VR Editor
+	else if( bPlayStartedFromVREditor && GEditor->PlayWorld && !bEnableVRRequest )
+	{
+		// Shutdown PIE if we came from the VR Editor and we are not already requesting to start the VR Editor and when any of the players is holding down the required input
+		const float ShutDownInputKeyTime = 0.5f;
+		TArray<APlayerController*> PlayerControllers;
+		GEngine->GetAllLocalPlayerControllers( PlayerControllers );
+		for ( APlayerController* PlayerController : PlayerControllers )
+		{
+			if ( ( PlayerController->GetInputKeyTimeDown( EKeys::MotionController_Left_Grip1 ) > ShutDownInputKeyTime || PlayerController->GetInputKeyTimeDown( EKeys::MotionController_Left_Grip2 ) > ShutDownInputKeyTime ) &&
+				( PlayerController->GetInputKeyTimeDown( EKeys::MotionController_Right_Grip1 ) > ShutDownInputKeyTime || PlayerController->GetInputKeyTimeDown( EKeys::MotionController_Right_Grip2 ) > ShutDownInputKeyTime ) &&
+				PlayerController->GetInputKeyTimeDown( EKeys::MotionController_Right_Trigger ) > ShutDownInputKeyTime &&
+				PlayerController->GetInputKeyTimeDown( EKeys::MotionController_Left_Trigger ) > ShutDownInputKeyTime )
+			{
+				GEditor->RequestEndPlayMap();
+				bEnableVRRequest = true;
+				break;
+			}
+		}
+	}
+	else
+	{
+		// Start the VR Editor mode
+		if ( bEnableVRRequest )
+		{
+			EnableVREditor( true );
+			bEnableVRRequest = false;
 		}
 	}
 }
diff --git a/Engine/Source/Editor/VREditor/VREditorQuickMenu.cpp b/Engine/Source/Editor/VREditor/VREditorQuickMenu.cpp
--- a/Engine/Source/Editor/VREditor/VREditorQuickMenu.cpp
+++ b/Engine/Source/Editor/VREditor/VREditorQuickMenu.cpp
@@ -379,6 +379,11 @@ void UVREditorQuickMenu::OnScreenshotButtonClicked()
 	FScreenshotRequest::RequestScreenshot(FString(), bShowUI, bUseCustomFilename);
 }
 
+void UVREditorQuickMenu::OnPlayButtonClicked()
+{
+	GetOwner()->GetOwner().GetOwner().StartExitingVRMode( EVREditorExitType::PIE_VR );
+}
+
 bool UVREditorQuickMenu::OnAssetEditorButtonClicked( const bool bIsChecked )
 {
 	GetOwner()->GetOwner().TogglePanelVisibility( UVREditorUISystem::EEditorUIPanel::AssetEditor );
diff --git a/Engine/Source/Editor/VREditor/VREditorQuickMenu.h b/Engine/Source/Editor/VREditor/VREditorQuickMenu.h
--- a/Engine/Source/Editor/VREditor/VREditorQuickMenu.h
+++ b/Engine/Source/Editor/VREditor/VREditorQuickMenu.h
@@ -141,6 +141,9 @@ public:
 	UFUNCTION(BlueprintCallable, Category = "Default")
 	void OnScreenshotButtonClicked();
 
+	UFUNCTION( BlueprintCallable, Category = "Default" )
+	void OnPlayButtonClicked();
+
 	UFUNCTION( BlueprintCallable, Category = "Default" )
 	bool OnAssetEditorButtonClicked( const bool bIsChecked );
 
diff --git a/Engine/Source/Editor/VREditor/VREditorUISystem.cpp b/Engine/Source/Editor/VREditor/VREditorUISystem.cpp
--- a/Engine/Source/Editor/VREditor/VREditorUISystem.cpp
+++ b/Engine/Source/Editor/VREditor/VREditorUISystem.cpp
@@ -168,7 +168,7 @@ void UVREditorUISystem::Shutdown()
 }
 
 void UVREditorUISystem::OnVRAction( FEditorViewportClient& ViewportClient, UViewportInteractor* Interactor,
-	const FViewportActionKeyInput& Action, const EInputEvent Event, bool& bOutIsInputCaptured, bool& bWasHandled )
+	const FViewportActionKeyInput& Action, bool& bOutIsInputCaptured, bool& bWasHandled )
 {
 	UVREditorInteractor* VREditorInteractor = Cast<UVREditorInteractor>( Interactor );
 	if ( VREditorInteractor )
@@ -177,7 +177,7 @@ void UVREditorUISystem::OnVRAction( FEditorViewportClient& ViewportClient, UView
 		{
 			if ( Action.ActionType == VRActionTypes::ConfirmRadialSelection )
 			{
-				if ( Event == IE_Pressed )
+				if ( Action.Event == IE_Pressed )
 				{
 					UVREditorRadialMenu* RadialMenu = QuickRadialMenu->GetUserWidget<UVREditorRadialMenu>();
 					if ( IsShowingRadialMenu( VREditorInteractor ) )
@@ -205,12 +205,12 @@ void UVREditorUISystem::OnVRAction( FEditorViewportClient& ViewportClient, UView
 							// Always mark the event as handled so that the editor doesn't try to select the widget component
 							bWasHandled = true;
 
-							if ( Event != IE_Repeat )
+							if ( Action.Event != IE_Repeat )
 							{
 								// If the Modifier button is held down, treat this like a right click instead of a left click
 								const bool bIsRightClicking =
-									( Event == IE_Pressed && VREditorInteractor->IsModifierPressed() ) ||
-									( Event == IE_Released && VREditorInteractor->IsRightClickingOnUI() );
+									( Action.Event == IE_Pressed && VREditorInteractor->IsModifierPressed() ) ||
+									( Action.Event == IE_Released && VREditorInteractor->IsRightClickingOnUI() );
 
 								FVector2D LastLocalHitLocation = WidgetComponent->GetLastLocalHitLocation();
 
@@ -232,7 +232,7 @@ void UVREditorUISystem::OnVRAction( FEditorViewportClient& ViewportClient, UView
 								if ( WidgetPathUnderFinger.IsValid() )
 								{
 									TSet<FKey> PressedButtons;
-									if ( Event == IE_Pressed )
+									if ( Action.Event == IE_Pressed )
 									{
 										PressedButtons.Add( bIsRightClicking ? EKeys::RightMouseButton : EKeys::LeftMouseButton );
 									}
@@ -254,7 +254,7 @@ void UVREditorUISystem::OnVRAction( FEditorViewportClient& ViewportClient, UView
 									}
 
 									FReply Reply = FReply::Unhandled();
-									if ( Event == IE_Pressed )
+									if ( Action.Event == IE_Pressed )
 									{
 										const double CurrentTime = FPlatformTime::Seconds();
 										if (CurrentTime - VREditorInteractor->GetLastClickReleaseTime() <= VREd::DoubleClickTime->GetFloat())
@@ -275,7 +275,7 @@ void UVREditorUISystem::OnVRAction( FEditorViewportClient& ViewportClient, UView
 										// Play a haptic effect on press
 										VREditorInteractor->PlayHapticEffect( VREd::UIPressHapticFeedbackStrength->GetFloat() );
 									}
-									else if ( Event == IE_Released )
+									else if ( Action.Event == IE_Released )
 									{
 										Reply = FSlateApplication::Get().RoutePointerUpEvent( WidgetPathUnderFinger, PointerEvent );
 									}
@@ -285,7 +285,7 @@ void UVREditorUISystem::OnVRAction( FEditorViewportClient& ViewportClient, UView
 					}
 				}
 
-				if ( Event == IE_Released )
+				if ( Action.Event == IE_Released )
 				{
 					bool bWasRightClicking = false;
 					if ( VREditorInteractor->IsClickingOnUI() )
diff --git a/Engine/Source/Editor/VREditor/VREditorUISystem.h b/Engine/Source/Editor/VREditor/VREditorUISystem.h
--- a/Engine/Source/Editor/VREditor/VREditorUISystem.h
+++ b/Engine/Source/Editor/VREditor/VREditorUISystem.h
@@ -124,7 +124,7 @@ protected:
 
 	/** Called when the user presses a button on their motion controller device */
 	void OnVRAction( class FEditorViewportClient& ViewportClient, UViewportInteractor* Interactor,
-		const struct FViewportActionKeyInput& Action, const EInputEvent Event, bool& bOutIsInputCaptured, bool& bWasHandled );
+		const struct FViewportActionKeyInput& Action, bool& bOutIsInputCaptured, bool& bWasHandled );
 
 	/** Called every frame to update hover state */
 	void OnVRHoverUpdate( class FEditorViewportClient& ViewportClient, UViewportInteractor* Interactor, FVector& HoverImpactPoint, bool& bWasHandled );
diff --git a/Engine/Source/Editor/ViewportInteraction/Public/ViewportDragOperation.h b/Engine/Source/Editor/ViewportInteraction/Public/ViewportDragOperation.h
--- a/Engine/Source/Editor/ViewportInteraction/Public/ViewportDragOperation.h
+++ b/Engine/Source/Editor/ViewportInteraction/Public/ViewportDragOperation.h
@@ -41,7 +41,7 @@ public:
 	UViewportDragOperation* GetDragOperation();
 
 	/** Sets the drag operation class that will be used next time starting dragging */
-	void SetDragOperactionClass( const TSubclassOf<UViewportDragOperation> InDragOperation );
+	void SetDragOperationClass( const TSubclassOf<UViewportDragOperation> InDragOperation );
 
 	/** Starts new dragging operation */
 	void StartDragOperation();
diff --git a/Engine/Source/Editor/ViewportInteraction/Public/ViewportInteractableInterface.h b/Engine/Source/Editor/ViewportInteraction/Public/ViewportInteractableInterface.h
--- a/Engine/Source/Editor/ViewportInteraction/Public/ViewportInteractableInterface.h
+++ b/Engine/Source/Editor/ViewportInteraction/Public/ViewportInteractableInterface.h
@@ -34,9 +34,6 @@ public:
 	/** Called when an interactor leave hovering over this object */
 	virtual void OnHoverLeave( UViewportInteractor* Interactor, const UActorComponent* NewComponent ) PURE_VIRTUAL( IViewportInteractableInterface::OnHoverLeave, );
 	
-	/** Called when an interactor is dragging this object */
-	virtual void OnDrag( UViewportInteractor* Interactor ) PURE_VIRTUAL( IViewportInteractableInterface::OnDrag, );
-
 	/** Called when an interactor stops dragging this object */
 	virtual void OnDragRelease( UViewportInteractor* Interactor ) PURE_VIRTUAL( IViewportInteractableInterface::OnDragRelease, );
 
diff --git a/Engine/Source/Editor/ViewportInteraction/Public/ViewportWorldInteractionInterface.h b/Engine/Source/Editor/ViewportInteraction/Public/ViewportWorldInteractionInterface.h
--- a/Engine/Source/Editor/ViewportInteraction/Public/ViewportWorldInteractionInterface.h
+++ b/Engine/Source/Editor/ViewportInteraction/Public/ViewportWorldInteractionInterface.h
@@ -35,7 +35,7 @@ public:
 	virtual FOnVIHoverUpdate& OnViewportInteractionHoverUpdate() = 0;
 
 	/** Gets the event for input actions update which is broadcasted for each interactor */
-	DECLARE_EVENT_SixParams( IViewportWorldInteractionInterface, FOnVIActionHandle, class FEditorViewportClient& /* ViewportClient */, class UViewportInteractor* /* Interactor */, 
-	const struct FViewportActionKeyInput& /* Action */, const EInputEvent /* Event */, bool& /* bOutIsInputCaptured */, bool& /* bWasHandled */ );
+	DECLARE_EVENT_FiveParams( IViewportWorldInteractionInterface, FOnVIActionHandle, class FEditorViewportClient& /* ViewportClient */, class UViewportInteractor* /* Interactor */, 
+	const struct FViewportActionKeyInput& /* Action */, bool& /* bOutIsInputCaptured */, bool& /* bWasHandled */ );
 	virtual FOnVIActionHandle& OnViewportInteractionInputAction() = 0;
 };
\ No newline at end of file
diff --git a/Engine/Source/Editor/ViewportInteraction/ViewportDragOperation.cpp b/Engine/Source/Editor/ViewportInteraction/ViewportDragOperation.cpp
--- a/Engine/Source/Editor/ViewportInteraction/ViewportDragOperation.cpp
+++ b/Engine/Source/Editor/ViewportInteraction/ViewportDragOperation.cpp
@@ -20,7 +20,7 @@ UViewportDragOperation* UViewportDragOperationComponent::GetDragOperation()
 	return DragOperation;
 }
 
-void UViewportDragOperationComponent::SetDragOperactionClass( const TSubclassOf<UViewportDragOperation> InDragOperation )
+void UViewportDragOperationComponent::SetDragOperationClass( const TSubclassOf<UViewportDragOperation> InDragOperation )
 {
 	DragOperationSubclass = InDragOperation;
 }
diff --git a/Engine/Source/Editor/ViewportInteraction/ViewportInteractor.cpp b/Engine/Source/Editor/ViewportInteraction/ViewportInteractor.cpp
--- a/Engine/Source/Editor/ViewportInteraction/ViewportInteractor.cpp
+++ b/Engine/Source/Editor/ViewportInteraction/ViewportInteractor.cpp
@@ -88,7 +88,7 @@ bool UViewportInteractor::HandleInputKey( const FKey Key, const EInputEvent Even
 		Action->Event = Event;
 		
 		// Give subsystems a chance to handle actions for this interactor
-		WorldInteraction->OnViewportInteractionInputAction().Broadcast( *WorldInteraction->GetViewportClient(), this, *Action, Event, Action->bIsInputCaptured, bHandled );
+		WorldInteraction->OnViewportInteractionInputAction().Broadcast( *WorldInteraction->GetViewportClient(), this, *Action, Action->bIsInputCaptured, bHandled );
 
 		if ( !bHandled )
 		{
diff --git a/Engine/Source/Editor/ViewportInteraction/ViewportWorldInteraction.cpp b/Engine/Source/Editor/ViewportInteraction/ViewportWorldInteraction.cpp
--- a/Engine/Source/Editor/ViewportInteraction/ViewportWorldInteraction.cpp
+++ b/Engine/Source/Editor/ViewportInteraction/ViewportWorldInteraction.cpp
@@ -210,10 +210,13 @@ static void SegmentDistToSegmentDouble( DVector A1, DVector B1, DVector A2, DVec
 
 UViewportWorldInteraction::UViewportWorldInteraction( const FObjectInitializer& Initializer ) :
 	Super( Initializer ),
+	bDraggedSinceLastSelection( false ),
+	LastDragGizmoStartTransform( FTransform::Identity ),
+	TrackingTransaction( FTrackingTransaction() ),
 	AppTimeEntered( FTimespan::Zero() ),
+	EditorViewportClient( nullptr ),
 	LastFrameNumberInputWasPolled( 0 ),
 	MotionControllerID( 0 ),	// @todo ViewportInteraction: We only support a single controller, and we assume the first controller are the motion controls
-	EditorViewportClient( nullptr ),
 	LastWorldToMetersScale( 100.0f ),
 	bIsInterpolatingTransformablesFromSnapshotTransform( false ),
 	TransformablesInterpolationStartTime( FTimespan::Zero() ),
@@ -221,8 +224,6 @@ UViewportWorldInteraction::UViewportWorldInteraction( const FObjectInitializer&
 	TransformGizmoActor( nullptr ),
 	TransformGizmoClass( APivotTransformGizmo::StaticClass() ),
 	GizmoLocalBounds( FBox( 0 ) ),
-	bDraggedSinceLastSelection( false ),
-	LastDragGizmoStartTransform( FTransform::Identity ),
 	StartDragAngleOnRotation( ),
 	StartDragHandleDirection( ),
 	CurrentGizmoType( EGizmoHandleTypes::All ),
