Viewport Interaction: Mouse cursor support (merged from CL 3026507)
- Added a new 'mouse cursor' interactor that can be enabled (when not in VR)
- We now automatically create a "mouse cursor" interactor when running in 'forced VR mode', so you can click on objects and gizmos!
- Currently it only supports 'SelectAndMove' (LMB), 'SelectAndMove_LightlyPressed' (MMB) and 'WorldMovement' (RMB)
- Interacting with UI using the mouse cursor is not supported yet.

Other fixes:
- ViewportInteraction: Fixed assertions when an interactor had no 'pair' and certain actions were triggered
- ViewportInteractor: Now creates a default transform for the interactor if no overridden implementation was provided
- VREditorInteractor: Fixed bad memory references to expired stack addresses
- VREditor: Fixed mouse buttons events not being routed through to HandleInputKey() (preprocessor did not support those.)
- VREditor: External modules can now activate VR Editor in 'forced VR mode', using new optional parameter to EnableVREditor()
- VREditor: Tutorial no longer shows up in 'forced VR mode'.  You can't click on UI anyway, and it was just in the way

diff --git a/Engine/Source/Editor/VREditor/Public/IVREditorModule.h b/Engine/Source/Editor/VREditor/Public/IVREditorModule.h
--- a/Engine/Source/Editor/VREditor/Public/IVREditorModule.h
+++ b/Engine/Source/Editor/VREditor/Public/IVREditorModule.h
@@ -61,7 +61,8 @@ public:
 	 * the editor UI for VR interaction.
 	 *
 	 * @param	bEnable	True to enable VR, or false to turn it off
+	 * @param	bForceWithoutHMD	If set to true, will enter VR mode without switching to HMD/stereo.  This can be useful for testing.
 	 */
-	virtual void EnableVREditor( const bool bEnable ) = 0;
+	virtual void EnableVREditor( const bool bEnable, const bool bForceWithoutHMD = false ) = 0;
 };
 
diff --git a/Engine/Source/Editor/VREditor/Teleporter/VREditorTeleporter.cpp b/Engine/Source/Editor/VREditor/Teleporter/VREditorTeleporter.cpp
--- a/Engine/Source/Editor/VREditor/Teleporter/VREditorTeleporter.cpp
+++ b/Engine/Source/Editor/VREditor/Teleporter/VREditorTeleporter.cpp
@@ -57,7 +57,7 @@ void UVREditorTeleporter::Tick( const float DeltaTime )
 void UVREditorTeleporter::StartTeleport( UViewportInteractor* Interactor )
 {
 	UVREditorInteractor* VREditorInteractor = Cast<UVREditorInteractor>( Interactor );
-	if ( VREditorInteractor && !VREditorInteractor->IsHoveringOverUI() )
+	if ( VREditorInteractor == nullptr || !VREditorInteractor->IsHoveringOverUI() )
 	{
 		FVector LaserPointerStart, LaserPointerEnd;
 		if ( Interactor->GetLaserPointer( LaserPointerStart, LaserPointerEnd ) )
diff --git a/Engine/Source/Editor/VREditor/VREditorDockableWindow.cpp b/Engine/Source/Editor/VREditor/VREditorDockableWindow.cpp
--- a/Engine/Source/Editor/VREditor/VREditorDockableWindow.cpp
+++ b/Engine/Source/Editor/VREditor/VREditorDockableWindow.cpp
@@ -327,22 +327,24 @@ void AVREditorDockableWindow::OnPressed( UViewportInteractor* Interactor, const
 {
 	bOutResultedInDrag = false;
 	UVREditorInteractor* VRInteractor = Cast<UVREditorInteractor>( Interactor );
-
-	if ( InHitResult.Component == GetCloseButtonMeshComponent() )
-	{
-		// Close the window
-		const bool bShouldShow = false;
-		const bool bShowOnHand = false;
-		const bool bRefreshQuickMenu = true;
-		GetOwner().ShowEditorUIPanel( this, VRInteractor, bShouldShow, bShowOnHand, bRefreshQuickMenu );
-	}
-	else if ( InHitResult.Component == GetSelectionBarMeshComponent() )
+	if( VRInteractor != nullptr )
 	{
-		bOutResultedInDrag = true;
-		
-		FVector LaserPointerStart, LaserPointerEnd;
-		DockSelectDistance = ( InHitResult.TraceStart - InHitResult.Location ).Size();
-		GetOwner().StartDraggingDockUI( this, VRInteractor, DockSelectDistance );
+		if( InHitResult.Component == GetCloseButtonMeshComponent() )
+		{
+			// Close the window
+			const bool bShouldShow = false;
+			const bool bShowOnHand = false;
+			const bool bRefreshQuickMenu = true;
+			GetOwner().ShowEditorUIPanel( this, VRInteractor, bShouldShow, bShowOnHand, bRefreshQuickMenu );
+		}
+		else if( InHitResult.Component == GetSelectionBarMeshComponent() )
+		{
+			bOutResultedInDrag = true;
+
+			FVector LaserPointerStart, LaserPointerEnd;
+			DockSelectDistance = ( InHitResult.TraceStart - InHitResult.Location ).Size();
+			GetOwner().StartDraggingDockUI( this, VRInteractor, DockSelectDistance );
+		}
 	}
 }
 
@@ -366,7 +368,11 @@ void AVREditorDockableWindow::OnHoverEnter( UViewportInteractor* Interactor, con
 
 void AVREditorDockableWindow::OnHoverLeave( UViewportInteractor* Interactor, const UActorComponent* NewComponent )
 {
-	UActorComponent* OtherInteractorHoveredComponent = Interactor->GetOtherInteractor()->GetHoverComponent();
+	UActorComponent* OtherInteractorHoveredComponent = nullptr;
+	if( Interactor->GetOtherInteractor() != nullptr )
+	{
+		OtherInteractorHoveredComponent = Interactor->GetOtherInteractor()->GetHoverComponent();
+	}
 
 	if ( OtherInteractorHoveredComponent != SelectionBarMeshComponent && NewComponent != SelectionBarMeshComponent )
 	{
@@ -382,8 +388,11 @@ void AVREditorDockableWindow::OnHoverLeave( UViewportInteractor* Interactor, con
 void AVREditorDockableWindow::OnDragRelease( UViewportInteractor* Interactor )
 {
 	UVREditorInteractor* VREditorInteractor = Cast<UVREditorInteractor>( Interactor );
-	UVREditorUISystem& UISystem = GetOwner();
-	UISystem.StopDraggingDockUI( VREditorInteractor );
+	if( VREditorInteractor != nullptr )
+	{
+		UVREditorUISystem& UISystem = GetOwner();
+		UISystem.StopDraggingDockUI( VREditorInteractor );
+	}
 }
 
 UViewportDragOperationComponent* AVREditorDockableWindow::GetDragOperationComponent()
diff --git a/Engine/Source/Editor/VREditor/VREditorMode.cpp b/Engine/Source/Editor/VREditor/VREditorMode.cpp
--- a/Engine/Source/Editor/VREditor/VREditorMode.cpp
+++ b/Engine/Source/Editor/VREditor/VREditorMode.cpp
@@ -25,6 +25,7 @@
 #include "Editor/LevelEditor/Public/LevelEditorActions.h"
 #include "Editor/ViewportInteraction/Public/ViewportInteraction.h"
 #include "VREditorInteractor.h"
+#include "MouseCursorInteractor.h"
 
 #include "Interactables/VREditorButton.h"
 
@@ -79,7 +80,11 @@ FVREditorMode::FVREditorMode()
 	  PostProcessComponent( nullptr ),
 	  MotionControllerID( 0 ),	// @todo vreditor minor: We only support a single controller, and we assume the first controller are the motion controls
 	  UISystem( nullptr ),
+	  TeleporterSystem( nullptr ),
 	  WorldInteraction( nullptr ),
+	  MouseCursorInteractor( nullptr ),
+	  LeftHandInteractor( nullptr ),
+	  RightHandInteractor( nullptr ),
 	  bFirstTick( true ),
 	  bWasInWorldSpaceBeforeScaleMode( false )
 {
@@ -327,15 +332,26 @@ void FVREditorMode::Enter()
 		WorldInteraction->SetIsUsingVR( bActuallyUsingVR );
 		WorldInteraction->Init( VREditorLevelViewportWeakPtr.Pin()->GetViewportClient().Get() );
 
-		LeftHandInteractor = NewObject<UVREditorInteractor>( WorldInteraction );
-		LeftHandInteractor->Init( EControllerHand::Left, this );
-		WorldInteraction->AddInteractor( LeftHandInteractor );
+		// Motion controllers
+		{
+			LeftHandInteractor = NewObject<UVREditorInteractor>( WorldInteraction );
+			LeftHandInteractor->Init( EControllerHand::Left, this );
+			WorldInteraction->AddInteractor( LeftHandInteractor );
 
-		RightHandInteractor = NewObject<UVREditorInteractor>( WorldInteraction );
-		RightHandInteractor->Init( EControllerHand::Right, this );
-		WorldInteraction->AddInteractor( RightHandInteractor );
+			RightHandInteractor = NewObject<UVREditorInteractor>( WorldInteraction );
+			RightHandInteractor->Init( EControllerHand::Right, this );
+			WorldInteraction->AddInteractor( RightHandInteractor );
 
-		WorldInteraction->PairInteractors( LeftHandInteractor, RightHandInteractor );
+			WorldInteraction->PairInteractors( LeftHandInteractor, RightHandInteractor );
+		}
+
+		if( !bActuallyUsingVR )
+		{
+			// Register an interactor for the mouse cursor
+			MouseCursorInteractor = NewObject<UMouseCursorInteractor>( WorldInteraction );
+			MouseCursorInteractor->Init();
+			WorldInteraction->AddInteractor( MouseCursorInteractor );
+		}
 
 		// Setup the UI system
 		UISystem = NewObject<UVREditorUISystem>();
@@ -966,6 +982,10 @@ bool FVREditorMode::InputKey(FEditorViewportClient* ViewportClient, FViewport* V
 		// User hit escape, so bail out of VR mode
 		StartExitingVRMode();
 	}
+	else if( Key.IsMouseButton() )	// Input preprocessor cannot handle mouse buttons, so we'll route those the normal way
+	{
+		return WorldInteraction->HandleInputKey( Key, Event );
+	}
 
 	return FEdMode::InputKey(ViewportClient, Viewport, Key, Event);
 }
@@ -1004,6 +1024,7 @@ void FVREditorMode::AddReferencedObjects( FReferenceCollector& Collector )
 	Collector.AddReferencedObject( UISystem );
 	Collector.AddReferencedObject( WorldInteraction );
 	Collector.AddReferencedObject( TeleporterSystem );
+	Collector.AddReferencedObject( MouseCursorInteractor );
 	Collector.AddReferencedObject( LeftHandInteractor );
 	Collector.AddReferencedObject( RightHandInteractor );
 }
diff --git a/Engine/Source/Editor/VREditor/VREditorMode.h b/Engine/Source/Editor/VREditor/VREditorMode.h
--- a/Engine/Source/Editor/VREditor/VREditorMode.h
+++ b/Engine/Source/Editor/VREditor/VREditorMode.h
@@ -349,6 +349,9 @@ protected:
 	// Interactors
 	//
 
+	/** The mouse cursor interactor (currently only available when not in VR) */
+	class UMouseCursorInteractor* MouseCursorInteractor;
+
 	/** The right motion controller */
 	class UVREditorInteractor* LeftHandInteractor; //@todo vreditor: Hardcoded interactors
 	
diff --git a/Engine/Source/Editor/VREditor/VREditorModule.cpp b/Engine/Source/Editor/VREditor/VREditorModule.cpp
--- a/Engine/Source/Editor/VREditor/VREditorModule.cpp
+++ b/Engine/Source/Editor/VREditor/VREditorModule.cpp
@@ -34,7 +34,7 @@ public:
 	}
 	virtual bool IsVREditorEnabled() const override;
 	virtual bool IsVREditorAvailable() const override;
-	virtual void EnableVREditor( const bool bEnable ) override;
+	virtual void EnableVREditor( const bool bEnable, const bool bForceWithoutHMD ) override;
 
 	// FTickableEditorObject overrides
 	virtual void Tick( float DeltaTime ) override;
@@ -52,7 +52,6 @@ public:
 
 private:
 
-	void EnableVREditor( const bool bEnable, const bool bForceWithoutHMD );
 
 	/** Sets the world settings directly and prevents the last GNewWorldToMetersScale from the VREditor to reset it by setting it 0 */
 	void SetDirectWorldToMeters( const float NewWorldToMeters );
@@ -158,13 +157,6 @@ void FVREditorModule::EnableVREditor( const bool bEnable, const bool bForceWitho
 }
 
 
-void FVREditorModule::EnableVREditor( const bool bEnable )
-{
-	const bool bForceWithoutHMD = false;
-	EnableVREditor( bEnable, bForceWithoutHMD );
-}
-
-
 void FVREditorModule::SetDirectWorldToMeters( const float NewWorldToMeters )
 {
 	GWorld->GetWorldSettings()->WorldToMeters = NewWorldToMeters;
@@ -192,7 +184,7 @@ void FVREditorModule::Tick( float DeltaTime )
 			//Store the current WorldToMeters before exiting the mode
 			LastWorldToMeters = GWorld->GetWorldSettings()->WorldToMeters;
 
-			EnableVREditor( false );
+			EnableVREditor( false, false );
 
 			// Start the session if that was the reason to stop the VR Editor mode
 			if ( VREditorMode->GetExitType() == EVREditorExitType::PIE_VR )
@@ -231,7 +223,7 @@ void FVREditorModule::Tick( float DeltaTime )
 		// Start the VR Editor mode
 		if ( bEnableVRRequest )
 		{
-			EnableVREditor( true );
+			EnableVREditor( true, false );
 			bEnableVRRequest = false;
 		}
 	}
diff --git a/Engine/Source/Editor/VREditor/VREditorUISystem.cpp b/Engine/Source/Editor/VREditor/VREditorUISystem.cpp
--- a/Engine/Source/Editor/VREditor/VREditorUISystem.cpp
+++ b/Engine/Source/Editor/VREditor/VREditorUISystem.cpp
@@ -327,120 +327,88 @@ void UVREditorUISystem::OnVRAction( FEditorViewportClient& ViewportClient, UView
 void UVREditorUISystem::OnVRHoverUpdate( FEditorViewportClient& ViewportClient, UViewportInteractor* Interactor, FVector& HoverImpactPoint, bool& bWasHandled )
 {
 	UVREditorInteractor* VREditorInteractor = Cast<UVREditorInteractor>( Interactor );
-
-	if( !bWasHandled && Interactor->GetDraggingMode() != EViewportInteractionDraggingMode::Interactable )
+	if( VREditorInteractor != nullptr )
 	{
-		FVector LaserPointerStart, LaserPointerEnd;
-		if( Interactor->GetLaserPointer( LaserPointerStart, LaserPointerEnd ) )
+		if( !bWasHandled && Interactor->GetDraggingMode() != EViewportInteractionDraggingMode::Interactable )
 		{
-			FHitResult HitResult = Interactor->GetHitResultFromLaserPointer();
-			if( HitResult.Actor.IsValid() )
+			FVector LaserPointerStart, LaserPointerEnd;
+			if( Interactor->GetLaserPointer( LaserPointerStart, LaserPointerEnd ) )
 			{
-				// Only allow clicks to our own widget components
-				UWidgetComponent* WidgetComponent = Cast<UWidgetComponent>( HitResult.GetComponent() );
-				if( WidgetComponent != nullptr && IsWidgetAnEditorUIWidget( WidgetComponent ) )
+				FHitResult HitResult = Interactor->GetHitResultFromLaserPointer();
+				if( HitResult.Actor.IsValid() )
 				{
-					FVector2D LastLocalHitLocation = WidgetComponent->GetLastLocalHitLocation();
-
-					FVector2D LocalHitLocation;
-					WidgetComponent->GetLocalHitLocation( HitResult.ImpactPoint, LocalHitLocation );
-
-					// If we weren't already hovering over this widget, then we'll reset the last hit location
-					if( WidgetComponent != VREditorInteractor->GetHoveringOverWidgetComponent() )
+					// Only allow clicks to our own widget components
+					UWidgetComponent* WidgetComponent = Cast<UWidgetComponent>( HitResult.GetComponent() );
+					if( WidgetComponent != nullptr && IsWidgetAnEditorUIWidget( WidgetComponent ) )
 					{
-						LastLocalHitLocation = LocalHitLocation;
+						FVector2D LastLocalHitLocation = WidgetComponent->GetLastLocalHitLocation();
 
-						if ( UVREditorWidgetComponent* VRWidgetComponent = Cast<UVREditorWidgetComponent>( VREditorInteractor->GetHoveringOverWidgetComponent() ) )
-						{
-							VRWidgetComponent->SetIsHovering(false);
-						}
-					}
+						FVector2D LocalHitLocation;
+						WidgetComponent->GetLocalHitLocation( HitResult.ImpactPoint, LocalHitLocation );
 
-					// @todo vreditor UI: There is a CursorRadius optional arg that we migth want to make use of wherever we call GetHitWidgetPath()!
-					FWidgetPath WidgetPathUnderFinger = FWidgetPath( WidgetComponent->GetHitWidgetPath( HitResult.ImpactPoint, /*bIgnoreEnabledStatus*/ false ) );
-					if ( WidgetPathUnderFinger.IsValid() )
-					{
-						TSet<FKey> PressedButtons;
-						FPointerEvent PointerEvent(
-							1 + (uint8)VREditorInteractor->GetControllerSide(),
-							LocalHitLocation,
-							LastLocalHitLocation,
-							LocalHitLocation - LastLocalHitLocation,
-							PressedButtons,
-							FModifierKeysState());
-
-						FSlateApplication::Get().RoutePointerMoveEvent(WidgetPathUnderFinger, PointerEvent, false);
-
-						bWasHandled = true;
-						HoverImpactPoint = HitResult.ImpactPoint;
-						VREditorInteractor->SetHoveringOverWidgetComponent( WidgetComponent );
-						VREditorInteractor->SetIsHoveringOverUI( true );
-
-						if ( UVREditorWidgetComponent* VRWidgetComponent = Cast<UVREditorWidgetComponent>(WidgetComponent) )
+						// If we weren't already hovering over this widget, then we'll reset the last hit location
+						if( WidgetComponent != VREditorInteractor->GetHoveringOverWidgetComponent() )
 						{
-							VRWidgetComponent->SetIsHovering(true);
+							LastLocalHitLocation = LocalHitLocation;
+
+							if ( UVREditorWidgetComponent* VRWidgetComponent = Cast<UVREditorWidgetComponent>( VREditorInteractor->GetHoveringOverWidgetComponent() ) )
+							{
+								VRWidgetComponent->SetIsHovering(false);
+							}
 						}
 
-						// Route the mouse scrolling
-						if ( VREditorInteractor->IsTrackpadPositionValid( 1 ) )
+						// @todo vreditor UI: There is a CursorRadius optional arg that we migth want to make use of wherever we call GetHitWidgetPath()!
+						FWidgetPath WidgetPathUnderFinger = FWidgetPath( WidgetComponent->GetHitWidgetPath( HitResult.ImpactPoint, /*bIgnoreEnabledStatus*/ false ) );
+						if ( WidgetPathUnderFinger.IsValid() )
 						{
-							const bool bIsAbsolute = ( GetOwner().GetHMDDeviceType() == EHMDDeviceType::DT_SteamVR );
+							TSet<FKey> PressedButtons;
+							FPointerEvent PointerEvent(
+								1 + (uint8)VREditorInteractor->GetControllerSide(),
+								LocalHitLocation,
+								LastLocalHitLocation,
+								LocalHitLocation - LastLocalHitLocation,
+								PressedButtons,
+								FModifierKeysState());
 
-							float ScrollDelta = 0.0f;
-							if( VREditorInteractor->IsTouchingTrackpad() || !bIsAbsolute )
-							{
-								if( bIsAbsolute )
-								{
-									const float ScrollSpeed = VREd::UIAbsoluteScrollSpeed->GetFloat();
-									ScrollDelta = ( VREditorInteractor->GetTrackpadPosition().Y - VREditorInteractor->GetLastTrackpadPosition().Y ) * ScrollSpeed;
-								}
-								else
-								{
-									const float ScrollSpeed = VREd::UIRelativeScrollSpeed->GetFloat();
-									ScrollDelta = VREditorInteractor->GetTrackpadPosition().Y * ScrollSpeed;
-								}
-							}
+							FSlateApplication::Get().RoutePointerMoveEvent(WidgetPathUnderFinger, PointerEvent, false);
+
+							bWasHandled = true;
+							HoverImpactPoint = HitResult.ImpactPoint;
+							VREditorInteractor->SetHoveringOverWidgetComponent( WidgetComponent );
+							VREditorInteractor->SetIsHoveringOverUI( true );
 
-							// If using a trackpad (Vive), invert scroll direction so that it feels more like scrolling on a mobile device
-							if( GetOwner().GetHMDDeviceType() == EHMDDeviceType::DT_SteamVR )
+							if ( UVREditorWidgetComponent* VRWidgetComponent = Cast<UVREditorWidgetComponent>(WidgetComponent) )
 							{
-								ScrollDelta *= -1.0f;
+								VRWidgetComponent->SetIsHovering(true);
 							}
 
-							if( !FMath::IsNearlyZero( ScrollDelta ) )
+							// Route the mouse scrolling
+							if ( VREditorInteractor->IsTrackpadPositionValid( 1 ) )
 							{
-								FPointerEvent MouseWheelEvent(
-									1 + (uint8)VREditorInteractor->GetControllerSide(),
-									LocalHitLocation,
-									LastLocalHitLocation,
-									PressedButtons,
-									EKeys::MouseWheelAxis,
-									ScrollDelta,
-									FModifierKeysState() );
+								const bool bIsAbsolute = ( GetOwner().GetHMDDeviceType() == EHMDDeviceType::DT_SteamVR );
 
-								FSlateApplication::Get().RouteMouseWheelOrGestureEvent(WidgetPathUnderFinger, MouseWheelEvent, nullptr);
-
-								VREditorInteractor->SetUIScrollVelocity( 0.0f );
-
-								// Don't apply inertia unless the user dragged a decent amount this frame
-								if( bIsAbsolute && FMath::Abs( ScrollDelta ) >= VREd::MinUIScrollDeltaForInertia->GetFloat() )
+								float ScrollDelta = 0.0f;
+								if( VREditorInteractor->IsTouchingTrackpad() || !bIsAbsolute )
 								{
-									// Don't apply inertia if our data is sort of old
-									const FTimespan CurrentTime = FTimespan::FromSeconds( FPlatformTime::Seconds() );
-									if( CurrentTime - VREditorInteractor->GetLastTrackpadPositionUpdateTime() < FTimespan::FromSeconds( 1.0f / 30.0f ) )
+									if( bIsAbsolute )
+									{
+										const float ScrollSpeed = VREd::UIAbsoluteScrollSpeed->GetFloat();
+										ScrollDelta = ( VREditorInteractor->GetTrackpadPosition().Y - VREditorInteractor->GetLastTrackpadPosition().Y ) * ScrollSpeed;
+									}
+									else
 									{
-										//GWarn->Logf( TEXT( "INPUT: UIScrollVelocity=%0.2f" ), ScrollDelta );
-										VREditorInteractor->SetUIScrollVelocity( ScrollDelta );
+										const float ScrollSpeed = VREd::UIRelativeScrollSpeed->GetFloat();
+										ScrollDelta = VREditorInteractor->GetTrackpadPosition().Y * ScrollSpeed;
 									}
 								}
-							}
-						}
-						else
-						{
-							if( !FMath::IsNearlyZero( VREditorInteractor->GetUIScrollVelocity() ) )
-							{
-								// Apply UI scrolling inertia
-								const float ScrollDelta = VREditorInteractor->GetUIScrollVelocity();
+
+								// If using a trackpad (Vive), invert scroll direction so that it feels more like scrolling on a mobile device
+								if( GetOwner().GetHMDDeviceType() == EHMDDeviceType::DT_SteamVR )
+								{
+									ScrollDelta *= -1.0f;
+								}
+
+								if( !FMath::IsNearlyZero( ScrollDelta ) )
 								{
 									FPointerEvent MouseWheelEvent(
 										1 + (uint8)VREditorInteractor->GetControllerSide(),
@@ -451,47 +419,81 @@ void UVREditorUISystem::OnVRHoverUpdate( FEditorViewportClient& ViewportClient,
 										ScrollDelta,
 										FModifierKeysState() );
 
-									FSlateApplication::Get().RouteMouseWheelOrGestureEvent( WidgetPathUnderFinger, MouseWheelEvent, nullptr );
-								}
+									FSlateApplication::Get().RouteMouseWheelOrGestureEvent(WidgetPathUnderFinger, MouseWheelEvent, nullptr);
+
+									VREditorInteractor->SetUIScrollVelocity( 0.0f );
 
-								// Apply damping
-								FVector ScrollVelocityVector( VREditorInteractor->GetUIScrollVelocity(), 0.0f, 0.0f );
-								const bool bVelocitySensitive = false;
-								GetOwner().GetWorldInteraction().ApplyVelocityDamping( ScrollVelocityVector, bVelocitySensitive );
-								VREditorInteractor->SetUIScrollVelocity( ScrollVelocityVector.X );
+									// Don't apply inertia unless the user dragged a decent amount this frame
+									if( bIsAbsolute && FMath::Abs( ScrollDelta ) >= VREd::MinUIScrollDeltaForInertia->GetFloat() )
+									{
+										// Don't apply inertia if our data is sort of old
+										const FTimespan CurrentTime = FTimespan::FromSeconds( FPlatformTime::Seconds() );
+										if( CurrentTime - VREditorInteractor->GetLastTrackpadPositionUpdateTime() < FTimespan::FromSeconds( 1.0f / 30.0f ) )
+										{
+											//GWarn->Logf( TEXT( "INPUT: UIScrollVelocity=%0.2f" ), ScrollDelta );
+											VREditorInteractor->SetUIScrollVelocity( ScrollDelta );
+										}
+									}
+								}
 							}
 							else
 							{
-								VREditorInteractor->SetUIScrollVelocity( 0.0f );
+								if( !FMath::IsNearlyZero( VREditorInteractor->GetUIScrollVelocity() ) )
+								{
+									// Apply UI scrolling inertia
+									const float ScrollDelta = VREditorInteractor->GetUIScrollVelocity();
+									{
+										FPointerEvent MouseWheelEvent(
+											1 + (uint8)VREditorInteractor->GetControllerSide(),
+											LocalHitLocation,
+											LastLocalHitLocation,
+											PressedButtons,
+											EKeys::MouseWheelAxis,
+											ScrollDelta,
+											FModifierKeysState() );
+
+										FSlateApplication::Get().RouteMouseWheelOrGestureEvent( WidgetPathUnderFinger, MouseWheelEvent, nullptr );
+									}
+
+									// Apply damping
+									FVector ScrollVelocityVector( VREditorInteractor->GetUIScrollVelocity(), 0.0f, 0.0f );
+									const bool bVelocitySensitive = false;
+									GetOwner().GetWorldInteraction().ApplyVelocityDamping( ScrollVelocityVector, bVelocitySensitive );
+									VREditorInteractor->SetUIScrollVelocity( ScrollVelocityVector.X );
+								}
+								else
+								{
+									VREditorInteractor->SetUIScrollVelocity( 0.0f );
+								}
 							}
 						}
 					}
 				}
 			}
 		}
-	}
 
-	// If nothing was hovered, make sure we tell Slate about that
-	if( !bWasHandled && VREditorInteractor->GetHoveringOverWidgetComponent() != nullptr )
-	{
-		if ( UVREditorWidgetComponent* VRWidgetComponent = Cast<UVREditorWidgetComponent>( VREditorInteractor->GetHoveringOverWidgetComponent() ) )
+		// If nothing was hovered, make sure we tell Slate about that
+		if( !bWasHandled && VREditorInteractor->GetHoveringOverWidgetComponent() != nullptr )
 		{
-			VRWidgetComponent->SetIsHovering( false );
-		}
+			if ( UVREditorWidgetComponent* VRWidgetComponent = Cast<UVREditorWidgetComponent>( VREditorInteractor->GetHoveringOverWidgetComponent() ) )
+			{
+				VRWidgetComponent->SetIsHovering( false );
+			}
 
-		const FVector2D LastLocalHitLocation = VREditorInteractor->GetHoveringOverWidgetComponent()->GetLastLocalHitLocation();
-		VREditorInteractor->SetHoveringOverWidgetComponent( nullptr );
+			const FVector2D LastLocalHitLocation = VREditorInteractor->GetHoveringOverWidgetComponent()->GetLastLocalHitLocation();
+			VREditorInteractor->SetHoveringOverWidgetComponent( nullptr );
 
-		TSet<FKey> PressedButtons;
-		FPointerEvent PointerEvent(
-			1 + (uint8)VREditorInteractor->GetControllerSide(),
-			LastLocalHitLocation,
-			LastLocalHitLocation,
-			FVector2D::ZeroVector,
-			PressedButtons,
-			FModifierKeysState() );
+			TSet<FKey> PressedButtons;
+			FPointerEvent PointerEvent(
+				1 + (uint8)VREditorInteractor->GetControllerSide(),
+				LastLocalHitLocation,
+				LastLocalHitLocation,
+				FVector2D::ZeroVector,
+				PressedButtons,
+				FModifierKeysState() );
 
-		FSlateApplication::Get().RoutePointerMoveEvent( FWidgetPath(), PointerEvent, false );
+			FSlateApplication::Get().RoutePointerMoveEvent( FWidgetPath(), PointerEvent, false );
+		}
 	}
 }
 
@@ -521,24 +523,26 @@ void UVREditorUISystem::Tick( FEditorViewportClient* ViewportClient, const float
 			bool bShowQuickMenuOnArm = false;
 
 			UViewportInteractor* OtherInteractor = Interactor->GetOtherInteractor();
-
-			// @todo vreditor tweak: Weird to hard code this here. Probably should be an accessor on the hand itself, and based on the actual device type
-			const FTransform UICapsuleTransform = OtherInteractor->GetTransform();
-			const FVector UICapsuleStart = FVector( -9.0f, 0.0f, 0.0f ) * WorldScaleFactor;
-			const FVector UICapsuleEnd = FVector( -18.0f, 0.0f, 0.0f ) * WorldScaleFactor;
-			const float UICapsuleLocalRadius = 6.0f * WorldScaleFactor;
-			const float MinDistanceToUICapsule = 8.0f * WorldScaleFactor; // @todo vreditor tweak
-			const FVector UIForwardVector = FVector::UpVector;
-			const float MinDotForAimingAtOtherHand = 0.25f;	// @todo vreditor tweak
-
-			if( GetOwner().IsHandAimingTowardsCapsule( Interactor, UICapsuleTransform, UICapsuleStart, UICapsuleEnd, UICapsuleLocalRadius, MinDistanceToUICapsule, UIForwardVector, MinDotForAimingAtOtherHand ) )
+			if( OtherInteractor != nullptr )
 			{
-				bShowQuickMenuOnArm = true;
-			}
+				// @todo vreditor tweak: Weird to hard code this here. Probably should be an accessor on the hand itself, and based on the actual device type
+				const FTransform UICapsuleTransform = OtherInteractor->GetTransform();
+				const FVector UICapsuleStart = FVector( -9.0f, 0.0f, 0.0f ) * WorldScaleFactor;
+				const FVector UICapsuleEnd = FVector( -18.0f, 0.0f, 0.0f ) * WorldScaleFactor;
+				const float UICapsuleLocalRadius = 6.0f * WorldScaleFactor;
+				const float MinDistanceToUICapsule = 8.0f * WorldScaleFactor; // @todo vreditor tweak
+				const FVector UIForwardVector = FVector::UpVector;
+				const float MinDotForAimingAtOtherHand = 0.25f;	// @todo vreditor tweak
+
+				if( GetOwner().IsHandAimingTowardsCapsule( Interactor, UICapsuleTransform, UICapsuleStart, UICapsuleEnd, UICapsuleLocalRadius, MinDistanceToUICapsule, UIForwardVector, MinDotForAimingAtOtherHand ) )
+				{
+					bShowQuickMenuOnArm = true;
+				}
 
-			if( bShowQuickMenuOnArm )
-			{
-				HandInteractorThatNeedsQuickMenu = OtherInteractor;
+				if( bShowQuickMenuOnArm )
+				{
+					HandInteractorThatNeedsQuickMenu = OtherInteractor;
+				}
 			}
 		}
 		
@@ -753,7 +757,9 @@ void UVREditorUISystem::CreateUIs()
 
 		// Create the tutorial dockable window
 		{
-			const bool bShowAtStart = !bIsVREditorDemo;
+			const bool bShowAtStart = 
+				!bIsVREditorDemo && 
+				GetOwner().IsActuallyUsingVR();	// Don't show tutorial at startup when in "forced VR" mode, because you can't interact with UI using the mouse currently
 
 			AVREditorFloatingUI* TutorialUI = GetOwner().SpawnTransientSceneActor< AVREditorDockableWindow >( TEXT( "TutorialUI" ), false );
 			TutorialUI->SetUMGWidget( *this, TutorialWidgetClass, FIntPoint( VREd::TutorialUIResolutionX->GetFloat(), VREd::TutorialUIResolutionY->GetFloat() ), VREd::TutorialUISize->GetFloat(), bShowAtStart ? AVREditorFloatingUI::EDockedTo::Room : AVREditorFloatingUI::EDockedTo::Nothing );
diff --git a/Engine/Source/Editor/ViewportInteraction/MouseCursorInteractor.cpp b/Engine/Source/Editor/ViewportInteraction/MouseCursorInteractor.cpp
new file mode 100644
--- /dev/null
+++ b/Engine/Source/Editor/ViewportInteraction/MouseCursorInteractor.cpp
@@ -0,0 +1,57 @@
+// Copyright 1998-2015 Epic Games, Inc. All Rights Reserved.
+
+#include "ViewportInteractionModule.h"
+#include "MouseCursorInteractor.h"
+#include "ViewportWorldInteraction.h"
+
+
+UMouseCursorInteractor::UMouseCursorInteractor( const FObjectInitializer& Initializer ) 
+	: UViewportInteractor( Initializer )
+{
+}
+
+
+void UMouseCursorInteractor::Init()
+{
+	KeyToActionMap.Reset();
+
+	// Setup keys
+	{
+		AddKeyAction( EKeys::LeftMouseButton, FViewportActionKeyInput( ViewportWorldActionTypes::SelectAndMove ) );
+		AddKeyAction( EKeys::MiddleMouseButton, FViewportActionKeyInput( ViewportWorldActionTypes::SelectAndMove_LightlyPressed ) );
+	}
+}
+
+
+void UMouseCursorInteractor::PollInput()
+{
+	FEditorViewportClient& ViewportClient = *WorldInteraction->GetViewportClient();
+
+	InteractorData.LastTransform = InteractorData.Transform;
+	InteractorData.LastRoomSpaceTransform = InteractorData.RoomSpaceTransform;
+
+	// Only if we're not tracking (RMB looking)
+	if( !ViewportClient.IsTracking() )
+	{
+		FViewport* Viewport = ViewportClient.Viewport;
+
+		// Make sure we have a valid viewport, otherwise we won't be able to construct an FSceneView.  The first time we're ticked we might not be properly setup. (@todo mesheditor)
+		if( Viewport != nullptr && Viewport->GetSizeXY().GetMin() > 0 )
+		{
+			const int32 ViewportInteractX = Viewport->GetMouseX();
+			const int32 ViewportInteractY = Viewport->GetMouseY();
+
+			FSceneViewFamilyContext ViewFamily( FSceneViewFamily::ConstructionValues(
+				Viewport,
+				ViewportClient.GetScene(),
+				ViewportClient.EngineShowFlags )
+				.SetRealtimeUpdate( ViewportClient.IsRealtime() ) );
+			FSceneView* SceneView = ViewportClient.CalcSceneView( &ViewFamily );
+
+			const FViewportCursorLocation MouseViewportRay( SceneView, &ViewportClient, ViewportInteractX, ViewportInteractY );
+
+			InteractorData.Transform = FTransform( MouseViewportRay.GetDirection().ToOrientationQuat(), MouseViewportRay.GetOrigin(), FVector( 1.0f ) );
+		}
+	}
+	InteractorData.RoomSpaceTransform = InteractorData.Transform * WorldInteraction->GetRoomTransform().Inverse();
+}
diff --git a/Engine/Source/Editor/ViewportInteraction/Public/MouseCursorInteractor.h b/Engine/Source/Editor/ViewportInteraction/Public/MouseCursorInteractor.h
new file mode 100644
--- /dev/null
+++ b/Engine/Source/Editor/ViewportInteraction/Public/MouseCursorInteractor.h
@@ -0,0 +1,27 @@
+// Copyright 1998-2015 Epic Games, Inc. All Rights Reserved.
+
+#pragma once
+
+#include "ViewportInteractor.h"
+#include "MouseCursorInteractor.generated.h"
+
+/**
+ * Viewport interactor for a mouse cursor
+ */
+UCLASS()
+class VIEWPORTINTERACTION_API UMouseCursorInteractor : public UViewportInteractor
+{
+	GENERATED_UCLASS_BODY()
+
+public:
+
+	/** Initialize default values */
+	void Init();
+
+	// ViewportInteractor overrides
+	virtual void PollInput() override;
+
+private:
+
+	// ...
+};
\ No newline at end of file
diff --git a/Engine/Source/Editor/ViewportInteraction/ViewportInteractor.cpp b/Engine/Source/Editor/ViewportInteraction/ViewportInteractor.cpp
--- a/Engine/Source/Editor/ViewportInteraction/ViewportInteractor.cpp
+++ b/Engine/Source/Editor/ViewportInteraction/ViewportInteractor.cpp
@@ -103,14 +103,13 @@ bool UViewportInteractor::HandleInputKey( const FKey Key, const EInputEvent Even
 			if ( ( Action->ActionType == ViewportWorldActionTypes::SelectAndMove ||
 				Action->ActionType == ViewportWorldActionTypes::SelectAndMove_LightlyPressed ) )
 			{
-				FViewportInteractorData& OtherInteractorData = OtherInteractor->GetInteractorData();
-
 				if ( Event == IE_Pressed )
 				{
 					// No clicking while we're dragging the world.  (No laser pointers are visible, anyway.)
 					const bool bIsDraggingWorldWithTwoHands =
-						( InteractorData.DraggingMode == EViewportInteractionDraggingMode::World && GetOtherInteractor()->GetInteractorData().DraggingMode == EViewportInteractionDraggingMode::AssistingDrag ) ||
-						( GetOtherInteractor()->GetInteractorData().DraggingMode == EViewportInteractionDraggingMode::World && InteractorData.DraggingMode == EViewportInteractionDraggingMode::AssistingDrag );
+						OtherInteractor != nullptr &&
+						( ( InteractorData.DraggingMode == EViewportInteractionDraggingMode::World && GetOtherInteractor()->GetInteractorData().DraggingMode == EViewportInteractionDraggingMode::AssistingDrag ) ||
+						  ( GetOtherInteractor()->GetInteractorData().DraggingMode == EViewportInteractionDraggingMode::World && InteractorData.DraggingMode == EViewportInteractionDraggingMode::AssistingDrag ) );
 
 					FHitResult HitResult = GetHitResultFromLaserPointer();
 					if ( !bIsDraggingWorldWithTwoHands && !bHandled && HitResult.Actor.IsValid() )
@@ -145,21 +144,28 @@ bool UViewportInteractor::HandleInputKey( const FKey Key, const EInputEvent Even
 							{
 								bHandled = true;
 
+								FViewportInteractorData* OtherInteractorData = nullptr;
+								if( OtherInteractor != nullptr )
+								{
+									OtherInteractorData = &OtherInteractor->GetInteractorData();
+								}
+
 								// Is the other hand already dragging this stuff?
-								if ( OtherInteractorData.DraggingMode == EViewportInteractionDraggingMode::ActorsWithGizmo ||
-									OtherInteractorData.DraggingMode == EViewportInteractionDraggingMode::ActorsFreely )
+								if ( OtherInteractorData != nullptr &&
+									 ( OtherInteractorData->DraggingMode == EViewportInteractionDraggingMode::ActorsWithGizmo ||
+									   OtherInteractorData->DraggingMode == EViewportInteractionDraggingMode::ActorsFreely ) )
 								{
 									// Only if they clicked on one of the objects we're moving
 									if ( WorldInteraction->IsTransformingActor( Actor ) )
 									{
 										// If we were dragging with a gizmo, we'll need to stop doing that and instead drag freely.
-										if ( OtherInteractorData.DraggingMode == EViewportInteractionDraggingMode::ActorsWithGizmo )
+										if ( OtherInteractorData->DraggingMode == EViewportInteractionDraggingMode::ActorsWithGizmo )
 										{
 											WorldInteraction->StopDragging( this );
 
 											FViewportActionKeyInput OtherInteractorAction( ViewportWorldActionTypes::SelectAndMove );
 											const bool bIsPlacingActors = false;
-											WorldInteraction->StartDraggingActors( OtherInteractor, OtherInteractorAction, HitResult.GetComponent(), OtherInteractorData.HoverLocation, bIsPlacingActors ); //@todo ViewportInteraction
+											WorldInteraction->StartDraggingActors( OtherInteractor, OtherInteractorAction, HitResult.GetComponent(), OtherInteractorData->HoverLocation, bIsPlacingActors ); //@todo ViewportInteraction
 										}
 
 										InteractorData.DraggingMode = InteractorData.LastDraggingMode = EViewportInteractionDraggingMode::AssistingDrag;
@@ -176,8 +182,8 @@ bool UViewportInteractor::HandleInputKey( const FKey Key, const EInputEvent Even
 										InteractorData.TransformGizmoInteractionType = ETransformGizmoInteractionType::None;
 										InteractorData.bIsDrivingVelocityOfSimulatedTransformables = false;
 
-										InteractorData.GizmoStartTransform = OtherInteractorData.GizmoStartTransform;
-										InteractorData.GizmoStartLocalBounds = OtherInteractorData.GizmoStartLocalBounds;
+										InteractorData.GizmoStartTransform = OtherInteractorData->GizmoStartTransform;
+										InteractorData.GizmoStartLocalBounds = OtherInteractorData->GizmoStartLocalBounds;
 										InteractorData.GizmoSpaceFirstDragUpdateOffsetAlongAxis = FVector::ZeroVector;	// Will be determined on first update
 										InteractorData.GizmoSpaceDragDeltaFromStartOffset = FVector::ZeroVector;	// Set every frame while dragging
 
@@ -194,11 +200,11 @@ bool UViewportInteractor::HandleInputKey( const FKey Key, const EInputEvent Even
 										// @todo vreditor: We don't currently support selection/dragging separate objects with either hand
 									}
 								}
-								else if ( OtherInteractorData.DraggingMode == EViewportInteractionDraggingMode::ActorsWithGizmo )
+								else if ( OtherInteractorData != nullptr && OtherInteractorData->DraggingMode == EViewportInteractionDraggingMode::ActorsWithGizmo )
 								{
 									// We don't support dragging objects with the gizmo using two hands.  Just absorb it.
 								}
-								else if ( OtherInteractorData.DraggingMode == EViewportInteractionDraggingMode::ActorsAtLaserImpact )
+								else if ( OtherInteractorData != nullptr && OtherInteractorData->DraggingMode == EViewportInteractionDraggingMode::ActorsAtLaserImpact )
 								{
 									// Doesn't work with two hands.  Just absorb it.
 								}
@@ -215,9 +221,10 @@ bool UViewportInteractor::HandleInputKey( const FKey Key, const EInputEvent Even
 									// If the other hand is holding down the button after selecting an object, we'll allow this hand to toggle selection
 									// of additional objects (multi select)
 									{
-										if ( OtherInteractorData.ClickingOnComponent.IsValid() &&		// Other hand is clicking on something
+										if( OtherInteractorData != nullptr && 
+											OtherInteractorData->ClickingOnComponent.IsValid() &&		// Other hand is clicking on something
 											Actor != WorldInteraction->GetTransformGizmoActor() &&					// Not clicking on a gizmo
-											OtherInteractorData.ClickingOnComponent.Get()->GetOwner() != Actor )	// Not clicking on same actor
+											OtherInteractorData->ClickingOnComponent.Get()->GetOwner() != Actor )	// Not clicking on same actor
 										{
 											// OK, the other hand is holding down the "select and move" button.
 											SelectionModification = ESelectionModification::Toggle;
@@ -261,9 +268,10 @@ bool UViewportInteractor::HandleInputKey( const FKey Key, const EInputEvent Even
 
 										// If the user did a full press to select an actor that wasn't selected, allow them to drag it right away
 										const bool bOtherHandTryingToDrag =
-											OtherInteractorData.ClickingOnComponent.IsValid() &&
-											OtherInteractorData.ClickingOnComponent.Get()->GetOwner()->IsSelected() &&
-											OtherInteractorData.ClickingOnComponent.Get()->GetOwner() == HitResult.GetComponent()->GetOwner();	// Trying to drag same actor
+											OtherInteractorData != nullptr &&
+											OtherInteractorData->ClickingOnComponent.IsValid() &&
+											OtherInteractorData->ClickingOnComponent.Get()->GetOwner()->IsSelected() &&
+											OtherInteractorData->ClickingOnComponent.Get()->GetOwner() == HitResult.GetComponent()->GetOwner();	// Trying to drag same actor
 										if ( ( !bWasSelected || bOtherHandTryingToDrag ) &&
 											Actor->IsSelected() &&
 											( Action->ActionType == ViewportWorldActionTypes::SelectAndMove || bOtherHandTryingToDrag ) )
@@ -299,8 +307,9 @@ bool UViewportInteractor::HandleInputKey( const FKey Key, const EInputEvent Even
 
 					const bool bIsDraggingWorld = InteractorData.DraggingMode == EViewportInteractionDraggingMode::World;
 					const bool bIsDraggingWorldWithTwoHands =
-						( InteractorData.DraggingMode == EViewportInteractionDraggingMode::World && GetOtherInteractor()->InteractorData.DraggingMode == EViewportInteractionDraggingMode::AssistingDrag ) ||
-						( GetOtherInteractor()->InteractorData.DraggingMode == EViewportInteractionDraggingMode::World && InteractorData.DraggingMode == EViewportInteractionDraggingMode::AssistingDrag );
+						OtherInteractor != nullptr &&
+						( ( InteractorData.DraggingMode == EViewportInteractionDraggingMode::World && GetOtherInteractor()->InteractorData.DraggingMode == EViewportInteractionDraggingMode::AssistingDrag ) ||
+						  ( GetOtherInteractor()->InteractorData.DraggingMode == EViewportInteractionDraggingMode::World && InteractorData.DraggingMode == EViewportInteractionDraggingMode::AssistingDrag ) );
 
 					// Don't allow the trigger to cancel our drag on release if we're dragging the world. 
 					if ( InteractorData.DraggingMode != EViewportInteractionDraggingMode::Nothing &&
@@ -309,9 +318,8 @@ bool UViewportInteractor::HandleInputKey( const FKey Key, const EInputEvent Even
 					{
 						WorldInteraction->StopDragging( this );
 						Action->bIsInputCaptured = false;
+						bHandled = true;
 					}
-
-					bHandled = true;
 				}
 			}
 			// World Movement
@@ -331,11 +339,14 @@ bool UViewportInteractor::HandleInputKey( const FKey Key, const EInputEvent Even
 						InteractorData.DraggingMode = InteractorData.LastDraggingMode = EViewportInteractionDraggingMode::World;
 						InteractorData.bWasAssistingDrag = false;
 
-						// Starting a new drag, so make sure the other hand doesn't think it's assisting us
-						OtherInteractor->GetInteractorData().bWasAssistingDrag = false;
+						if( OtherInteractor != nullptr )
+						{
+							// Starting a new drag, so make sure the other hand doesn't think it's assisting us
+							OtherInteractor->GetInteractorData().bWasAssistingDrag = false;
 
-						// Stop any interia from the other hand's previous movements -- we've grabbed the world and it needs to stick!
-						OtherInteractor->GetInteractorData().DragTranslationVelocity = FVector::ZeroVector;
+							// Stop any interia from the other hand's previous movements -- we've grabbed the world and it needs to stick!
+							OtherInteractor->GetInteractorData().DragTranslationVelocity = FVector::ZeroVector;
+						}
 					}
 
 					InteractorData.bIsFirstDragUpdate = true;
@@ -442,8 +453,9 @@ bool UViewportInteractor::GetLaserPointer( FVector& LaserPointerStart, FVector&
 {
 	// If we're currently grabbing the world with both hands, then the laser pointer is not available
 	if ( /*bHaveMotionController &&*/ //@todo ViewportInteraction
-		!( InteractorData.DraggingMode == EViewportInteractionDraggingMode::World && GetOtherInteractor()->GetInteractorData().DraggingMode == EViewportInteractionDraggingMode::AssistingDrag ) &&
-		!( InteractorData.DraggingMode == EViewportInteractionDraggingMode::AssistingDrag && GetOtherInteractor()->GetInteractorData().DraggingMode == EViewportInteractionDraggingMode::World ) )
+		 OtherInteractor == nullptr ||
+		 ( !( InteractorData.DraggingMode == EViewportInteractionDraggingMode::World && GetOtherInteractor()->GetInteractorData().DraggingMode == EViewportInteractionDraggingMode::AssistingDrag ) &&
+		   !( InteractorData.DraggingMode == EViewportInteractionDraggingMode::AssistingDrag && GetOtherInteractor()->GetInteractorData().DraggingMode == EViewportInteractionDraggingMode::World ) ) )
 	{
 		// If we have UI attached to us, don't allow a laser pointer
 		if ( bEvenIfBlocked || !GetIsLaserBlocked() )
@@ -539,7 +551,10 @@ FHitResult UViewportInteractor::GetHitResultFromLaserPointer( TArray<AActor*>* O
 
 bool UViewportInteractor::GetTransformAndForwardVector( FTransform& OutHandTransform, FVector& OutForwardVector )
 {
-	return false;
+	OutHandTransform = InteractorData.Transform;
+	OutForwardVector = OutHandTransform.GetRotation().Vector();
+
+	return true;
 }
 
 FVector UViewportInteractor::GetHoverLocation() const 
diff --git a/Engine/Source/Editor/ViewportInteraction/ViewportInteractorMotionController.cpp b/Engine/Source/Editor/ViewportInteraction/ViewportInteractorMotionController.cpp
--- a/Engine/Source/Editor/ViewportInteraction/ViewportInteractorMotionController.cpp
+++ b/Engine/Source/Editor/ViewportInteraction/ViewportInteractorMotionController.cpp
@@ -175,25 +175,20 @@ void UViewportInteractorMotionController::PollInput()
 	InteractorData.LastTransform = InteractorData.Transform;
 	InteractorData.LastRoomSpaceTransform = InteractorData.RoomSpaceTransform;
 
-	//@todo ViewportInteraction: VREditorMode::PollInputFromMotionControllers mouse implementation 
-
 	// Generic motion controllers
-	if ( !bHaveMotionController )
+	TArray<IMotionController*> MotionControllers = IModularFeatures::Get().GetModularFeatureImplementations<IMotionController>( IMotionController::GetModularFeatureName() );
+	for ( auto MotionController : MotionControllers )	// @todo viewportinteraction: Needs support for multiple pairs of motion controllers
 	{
-		TArray<IMotionController*> MotionControllers = IModularFeatures::Get().GetModularFeatureImplementations<IMotionController>( IMotionController::GetModularFeatureName() );
-		for ( auto MotionController : MotionControllers )
+		if ( MotionController != nullptr && !bHaveMotionController )
 		{
-			if ( MotionController != nullptr && !bHaveMotionController )
-			{
-				FVector Location = FVector::ZeroVector;
-				FRotator Rotation = FRotator::ZeroRotator;
+			FVector Location = FVector::ZeroVector;
+			FRotator Rotation = FRotator::ZeroRotator;
 
-				if ( MotionController->GetControllerOrientationAndPosition( WorldInteraction->GetMotionControllerID(), ControllerHandSide, /* Out */ Rotation, /* Out */ Location ) )
-				{
-					bHaveMotionController = true;
-					InteractorData.RoomSpaceTransform = FTransform( Rotation.Quaternion(), Location, FVector( 1.0f ) );
-					InteractorData.Transform = InteractorData.RoomSpaceTransform * WorldInteraction->GetRoomTransform();
-				}
+			if ( MotionController->GetControllerOrientationAndPosition( WorldInteraction->GetMotionControllerID(), ControllerHandSide, /* Out */ Rotation, /* Out */ Location ) )
+			{
+				bHaveMotionController = true;
+				InteractorData.RoomSpaceTransform = FTransform( Rotation.Quaternion(), Location, FVector( 1.0f ) );
+				InteractorData.Transform = InteractorData.RoomSpaceTransform * WorldInteraction->GetRoomTransform();
 			}
 		}
 	}
diff --git a/Engine/Source/Editor/ViewportInteraction/ViewportWorldInteraction.cpp b/Engine/Source/Editor/ViewportInteraction/ViewportWorldInteraction.cpp
--- a/Engine/Source/Editor/ViewportInteraction/ViewportWorldInteraction.cpp
+++ b/Engine/Source/Editor/ViewportInteraction/ViewportWorldInteraction.cpp
@@ -650,8 +650,6 @@ void UViewportWorldInteraction::InteractionTick( FEditorViewportClient* Viewport
 				OtherInteractor = AssistingDragWithInteractor.GetValue();
 			}
 
-			FViewportInteractorData& OtherInteractorData = OtherInteractor->GetInteractorData();
-
 			UViewportInteractor* OtherInteractorThatWasAssistingDrag = GetOtherInteractorIntertiaContribute( Interactor );
 
 			FVector DraggedTo = InteractorData.Transform.GetLocation();
@@ -754,6 +752,8 @@ void UViewportWorldInteraction::InteractionTick( FEditorViewportClient* Viewport
 					FVector OtherHandLaserPointerStart, OtherHandLaserPointerEnd;
 					if( OtherInteractor->GetLaserPointer( /* Out */ OtherHandLaserPointerStart, /* Out */ OtherHandLaserPointerEnd ) )
 					{
+						FViewportInteractorData& OtherInteractorData = OtherInteractor->GetInteractorData();
+
 						const FVector OtherHandLaserPointerDirection = ( OtherHandLaserPointerEnd - OtherHandLaserPointerStart ).GetSafeNormal();
 						OtherHandDraggedTo = OtherHandLaserPointerStart + OtherHandLaserPointerDirection * OtherInteractorData.DragRayLength;
 
@@ -797,6 +797,8 @@ void UViewportWorldInteraction::InteractionTick( FEditorViewportClient* Viewport
 				// Two handed?
 				if( OtherInteractor != nullptr )
 				{
+					FViewportInteractorData& OtherInteractorData = OtherInteractor->GetInteractorData();
+
 					const FVector OtherHandRoomSpaceUnscaledLocation = ( OtherInteractorData.RoomSpaceTransform.GetLocation() / WorldToMetersScale ) * LastWorldToMetersScale;
 					const FVector OtherHandRoomSpaceUnscaledHandDelta = ( OtherHandRoomSpaceUnscaledLocation - OtherInteractorData.LastRoomSpaceTransform.GetLocation() );
 
@@ -943,7 +945,6 @@ void UViewportWorldInteraction::InteractionTick( FEditorViewportClient* Viewport
 				const FVector DragDeltaFromStart = DraggedTo - InteractorData.LaserPointerImpactAtDragStart;
 
 				UViewportInteractor* OtherInteractorThatWasAssistingDrag = GetOtherInteractorIntertiaContribute( Interactor );
-				FViewportInteractorData& OtherInteractorThatWasAssistingDragData = OtherInteractorThatWasAssistingDrag->GetInteractorData();
 
 				const bool bWithTwoHands = ( OtherInteractorThatWasAssistingDrag != nullptr );
 
@@ -952,6 +953,7 @@ void UViewportWorldInteraction::InteractionTick( FEditorViewportClient* Viewport
 				FVector OtherHandDraggedTo = FVector::ZeroVector;
 				if( bWithTwoHands )
 				{
+					FViewportInteractorData& OtherInteractorThatWasAssistingDragData = OtherInteractorThatWasAssistingDrag->GetInteractorData();
 					OtherHandDragDelta = OtherInteractorThatWasAssistingDragData.DragTranslationVelocity;
 					OtherHandDraggedTo = OtherInteractorThatWasAssistingDragData.LastDragToLocation + OtherHandDragDelta;
 					OtherHandDragDeltaFromStart = OtherHandDraggedTo - OtherInteractorThatWasAssistingDragData.LaserPointerImpactAtDragStart;
@@ -990,6 +992,7 @@ void UViewportWorldInteraction::InteractionTick( FEditorViewportClient* Viewport
 
 				if( OtherInteractorThatWasAssistingDrag != nullptr )
 				{
+					FViewportInteractorData& OtherInteractorThatWasAssistingDragData = OtherInteractorThatWasAssistingDrag->GetInteractorData();
 					OtherInteractorThatWasAssistingDragData.LastDragToLocation = OtherHandDraggedTo;
 					ApplyVelocityDamping( OtherInteractorThatWasAssistingDragData.DragTranslationVelocity, bVelocitySensitive );
 				}
@@ -2098,8 +2101,11 @@ void UViewportWorldInteraction::StartDraggingActors( UViewportInteractor* Intera
 				( bUsingGizmo ? EViewportInteractionDraggingMode::ActorsWithGizmo : EViewportInteractionDraggingMode::ActorsFreely );
 
 			// Starting a new drag, so make sure the other hand doesn't think it's assisting us
-			FViewportInteractorData& OtherInteractorData = Interactor->GetOtherInteractor()->GetInteractorData();
-			OtherInteractorData.bWasAssistingDrag = false;
+			if( Interactor->GetOtherInteractor() != nullptr )
+			{
+				FViewportInteractorData& OtherInteractorData = Interactor->GetOtherInteractor()->GetInteractorData();
+				OtherInteractorData.bWasAssistingDrag = false;
+			}
 
 			InteractorData.bIsFirstDragUpdate = true;
 			InteractorData.bWasAssistingDrag = false;
@@ -2165,15 +2171,20 @@ void UViewportWorldInteraction::StopDragging( UViewportInteractor* Interactor )
 		// If the other hand started dragging after we started, allow that hand to "take over" the drag, so the user
 		// doesn't have to click again to continue their action.  Inertial effects of the hand that stopped dragging
 		// will still be in effect.
-		FViewportInteractorData& OtherInteractorData = Interactor->GetOtherInteractor()->GetInteractorData();
-		if ( OtherInteractorData.DraggingMode == EViewportInteractionDraggingMode::AssistingDrag )
+		FViewportInteractorData* OtherInteractorData = nullptr;
+		if( Interactor->GetOtherInteractor() != nullptr )
+		{
+			OtherInteractorData = &Interactor->GetOtherInteractor()->GetInteractorData();
+		}
+		
+		if ( OtherInteractorData != nullptr && OtherInteractorData->DraggingMode == EViewportInteractionDraggingMode::AssistingDrag )
 		{
 			// The other hand takes over whatever this hand was doing
-			OtherInteractorData.DraggingMode = OtherInteractorData.LastDraggingMode = InteractorData.DraggingMode;
-			OtherInteractorData.bIsDrivingVelocityOfSimulatedTransformables = InteractorData.bIsDrivingVelocityOfSimulatedTransformables;
+			OtherInteractorData->DraggingMode = OtherInteractorData->LastDraggingMode = InteractorData.DraggingMode;
+			OtherInteractorData->bIsDrivingVelocityOfSimulatedTransformables = InteractorData.bIsDrivingVelocityOfSimulatedTransformables;
 
 			// The other hand is no longer assisting, as it's now the primary interacting hand.
-			OtherInteractorData.bWasAssistingDrag = false;
+			OtherInteractorData->bWasAssistingDrag = false;
 
 			// The hand that stopped dragging will be remembered as "assisting" the other hand, so that its
 			// inertia will still influence interactions
@@ -2668,15 +2679,18 @@ UViewportInteractor* UViewportWorldInteraction::GetOtherInteractorIntertiaContri
 	UViewportInteractor* OtherInteractorThatWasAssistingDrag = nullptr;
 	{
 		UViewportInteractor* OtherInteractor = Interactor->GetOtherInteractor();
-		FViewportInteractorData& OtherHandInteractorData = OtherInteractor->GetInteractorData();
-
-		// If the other hand isn't doing anything, but the last thing it was doing was assisting a drag, then allow it
-		// to contribute inertia!
-		if ( OtherHandInteractorData.DraggingMode == EViewportInteractionDraggingMode::Nothing && OtherHandInteractorData.bWasAssistingDrag )
+		if( OtherInteractor != nullptr  )
 		{
-			if ( !OtherHandInteractorData.DragTranslationVelocity.IsNearlyZero( KINDA_SMALL_NUMBER ) )
+			FViewportInteractorData& OtherHandInteractorData = OtherInteractor->GetInteractorData();
+
+			// If the other hand isn't doing anything, but the last thing it was doing was assisting a drag, then allow it
+			// to contribute inertia!
+			if ( OtherHandInteractorData.DraggingMode == EViewportInteractionDraggingMode::Nothing && OtherHandInteractorData.bWasAssistingDrag )
 			{
-				OtherInteractorThatWasAssistingDrag = OtherInteractor;
+				if ( !OtherHandInteractorData.DragTranslationVelocity.IsNearlyZero( KINDA_SMALL_NUMBER ) )
+				{
+					OtherInteractorThatWasAssistingDrag = OtherInteractor;
+				}
 			}
 		}
 	}
