VREditor: Fixed assertion when hovering over model components (BSP)

diff --git a/Engine/Source/Editor/VREditor/VREditorInteractor.cpp b/Engine/Source/Editor/VREditor/VREditorInteractor.cpp
--- a/Engine/Source/Editor/VREditor/VREditorInteractor.cpp
+++ b/Engine/Source/Editor/VREditor/VREditorInteractor.cpp
@@ -409,7 +409,8 @@ void UVREditorInteractor::Tick( const float DeltaTime )
 FHitResult UVREditorInteractor::GetHitResultFromLaserPointer( TArray<AActor*>* OptionalListOfIgnoredActors /*= nullptr*/, const bool bIgnoreGizmos /*= false*/,
 	TArray<UClass*>* ObjectsInFrontOfGizmo /*= nullptr */, const bool bEvenIfBlocked /*= false */, const float LaserLengthOverride /*= 0.0f */ )
 {
-	TArray<AActor*> IgnoredActors;
+	static TArray<AActor*> IgnoredActors;
+	IgnoredActors.Reset();
 	if ( OptionalListOfIgnoredActors == nullptr )
 	{
 		OptionalListOfIgnoredActors = &IgnoredActors;
@@ -426,7 +427,8 @@ FHitResult UVREditorInteractor::GetHitResultFromLaserPointer( TArray<AActor*>* O
 		}
 	}
 
-	TArray<UClass*> PriorityOverGizmoObjects;
+	static TArray<UClass*> PriorityOverGizmoObjects;
+	PriorityOverGizmoObjects.Reset();
 	if ( ObjectsInFrontOfGizmo == nullptr )
 	{
 		ObjectsInFrontOfGizmo = &PriorityOverGizmoObjects;
diff --git a/Engine/Source/Editor/ViewportInteraction/ViewportInteractor.cpp b/Engine/Source/Editor/ViewportInteraction/ViewportInteractor.cpp
--- a/Engine/Source/Editor/ViewportInteraction/ViewportInteractor.cpp
+++ b/Engine/Source/Editor/ViewportInteraction/ViewportInteractor.cpp
@@ -528,7 +528,9 @@ FHitResult UViewportInteractor::GetHitResultFromLaserPointer( TArray<AActor*>* O
 					for ( UClass* CurrentClass : *ObjectsInFrontOfGizmo )
 					{
 						bool bClassHasPriority = false;
-						bClassHasPriority = HitResult.GetActor()->IsA( CurrentClass );
+						bClassHasPriority =
+							( HitResult.GetComponent() != nullptr && HitResult.GetComponent()->IsA( CurrentClass ) ) ||
+							( HitResult.GetActor() != nullptr && HitResult.GetActor()->IsA( CurrentClass ) );
 
 						if ( bClassHasPriority )
 						{
