VREditor: - Moved ViewportMotionControllerInteractor to VREditor module to solve inheritance issue
- Fixed laser visibility (UE-32529 & UE-32523)
- Fixed clearing delegate in Teleport

diff --git a/Engine/Source/Editor/VREditor/Public/VREditorInteractor.h b/Engine/Source/Editor/VREditor/Public/VREditorInteractor.h
--- a/Engine/Source/Editor/VREditor/Public/VREditorInteractor.h
+++ b/Engine/Source/Editor/VREditor/Public/VREditorInteractor.h
@@ -2,14 +2,14 @@
 
 #pragma once
 
-#include "Editor/ViewportInteraction/Public/ViewportInteractorMotionController.h"
+#include "Editor/ViewportInteraction/Public/ViewportInteractor.h"
 #include "VREditorInteractor.generated.h"
 
 /**
  * VREditor default interactor
  */
 UCLASS()
-class VREDITOR_API UVREditorInteractor : public UViewportInteractorMotionController
+class VREDITOR_API UVREditorInteractor : public UViewportInteractor
 {
 	GENERATED_UCLASS_BODY()
 
@@ -18,10 +18,7 @@ public:
 	virtual ~UVREditorInteractor();
 
 	/** Initialize default values */
-	void Init( const EControllerHand InHandSide, class FVREditorMode* InVRMode );
-
-	/** Sets up all componentes defined in UViewportInteractorMotionController with the passed Actor as owner */
-	void SetupComponent( AActor* OwningActor );
+	virtual void Init( class FVREditorMode* InVRMode );
 
 	/** Gets the owner of this system */
 	class FVREditorMode& GetVRMode()
@@ -41,17 +38,16 @@ public:
 	virtual FHitResult GetHitResultFromLaserPointer( TArray<AActor*>* OptionalListOfIgnoredActors = nullptr, const bool bIgnoreGizmos = false,
 		TArray<UClass*>* ObjectsInFrontOfGizmo = nullptr, const bool bEvenIfBlocked = false, const float LaserLengthOverride = 0.0f ) override;
 	virtual void ResetHoverState( const float DeltaTime ) override;
-	virtual void CalculateDragRay() override;
 	virtual void OnStartDragging( UActorComponent* ClickedComponent, const FVector& HitLocation, const bool bIsPlacingActors ) override;
 
+	/** Returns the slide delta for pushing and pulling objects. Needs to be implemented by derived classes (e.g. touchpad for vive controller or scrollweel for mouse ) */
+	virtual float GetSlideDelta();
+
 	//
 	// Getters and setters
 	//
 
-	/** Gets the Y delta of the trackpad */
-	float GetTrackpadSlideDelta();
-
-	/** Gets if*/
+	/** Gets if this interactor is hovering over UI */
 	bool IsHoveringOverUI() const;
 	
 	/** Sets if there is a docked window on this interactor */
@@ -66,26 +62,11 @@ public:
 	/** Check if the quick menu is on this interactor */
 	bool HasUIOnForearm() const;
 
-	/** Check if the touchpad is currently touched */
-	bool IsTouchingTrackpad() const;
-
-	/** Get the current position of the trackpad or analog stick */
-	FVector2D GetTrackpadPosition() const;
-
-	/** Get the last position of the trackpad or analog stick */
-	FVector2D GetLastTrackpadPosition() const;
-
 	/** Gets the current hovered widget component if any */
 	UWidgetComponent* GetHoveringOverWidgetComponent() const;
 
 	/** Sets the current hovered widget component */
 	void SetHoveringOverWidgetComponent( UWidgetComponent* NewHoveringOverWidgetComponent );
-	
-	/** If the trackpad values are valid */
-	bool IsTrackpadPositionValid( const int32 AxisIndex ) const;
-
-	/** Get when the last time the trackpad position was updated */
-	FTimespan& GetLastTrackpadPositionUpdateTime();
 
 	/** If the modifier key is currently pressed */
 	bool IsModifierPressed() const;
@@ -116,43 +97,29 @@ public:
 
 	/** Gets the UI scroll velocity */
 	float GetUIScrollVelocity() const;
+
+	/** Gets the trigger value */
+	float GetSelectAndMoveTriggerValue() const;
 	
 	/* ViewportInteractor overrides, checks if the laser is blocked by UI */
 	virtual bool GetIsLaserBlocked() override;
-
+	
 protected:
 
-	// ViewportInteractor overrides
-	virtual void HandleInputKey( FViewportActionKeyInput& Action, const FKey Key, const EInputEvent Event, bool& bOutWasHandled ) override;
-	virtual void HandleInputAxis( FViewportActionKeyInput& Action, const FKey Key, const float Delta, const float DeltaTime, bool& bOutWasHandled ) override;
-
-private:
-
-	/** Changes the color of the buttons on the handmesh */
-	void ApplyButtonPressColors( const FViewportActionKeyInput& Action );
-
-	/** Set the visuals for a button on the motion controller */
-	void SetMotionControllerButtonPressedVisuals( const EInputEvent Event, const FName& ParameterName, const float PressStrength );
-	
-	/** Pops up some help text labels for the controller in the specified hand, or hides it, if requested */
-	void ShowHelpForHand( const bool bShowIt );
-	
-	/** Called every frame to update the position of any floating help labels */
-	void UpdateHelpLabels();
+	/** The mode that owns this interactor */
+	class FVREditorMode* VRMode;
 
-	/** Given a mesh and a key name, tries to find a socket on the mesh that matches a supported key */
-	UStaticMeshSocket* FindMeshSocketForKey( UStaticMesh* StaticMesh, const FKey Key );
+	//
+	// General input @todo: VREditor: Should this be general (non-UI) in interactordata ?
+	//
 
-	/** Sets the visuals of the LaserPointer */
-	void SetLaserVisuals( const FLinearColor& NewColor, const float CrawlFade, const float CrawlSpeed );
+	/** Is the Modifier button held down? */
+	bool bIsModifierPressed;
 
-	/** Updates the radial menu */
-	void UpdateRadialMenuInput( const float DeltaTime );
+	/** Current trigger pressed amount for 'select and move' (0.0 - 1.0) */
+	float SelectAndMoveTriggerValue;
 
 private:
-	
-	/** The mode that created this interactor */
-	class FVREditorMode* VRMode;
 
 	//
 	// UI
@@ -177,35 +144,11 @@ private:
 	/** Inertial scrolling -- how fast to scroll the mousewheel over UI */
 	float UIScrollVelocity;	
 
-	//
-	// Trackpad support
-	//
-
-	/** True if the trackpad is actively being touched */
-	bool bIsTouchingTrackpad;
-
-	/** Position of the touched trackpad */
-	FVector2D TrackpadPosition;
-
-	/** Last position of the touched trackpad */
-	FVector2D LastTrackpadPosition;
-
-	/** True if we have a valid trackpad position (for each axis) */
-	bool bIsTrackpadPositionValid[ 2 ];
-
-	/** Real time that the last trackpad position was last updated.  Used to filter out stale previous data. */
-	FTimespan LastTrackpadPositionUpdateTime;
-
-	//
-	// General input @todo: VREditor: Should this be general (non-UI) in interactordata ?
-	//
-
-	/** Is the Modifier button held down? */
-	bool bIsModifierPressed;
-
 	/** Last real time that we released the 'SelectAndMove' button on UI.  This is used to detect double-clicks. */
 	double LastClickReleaseTime;
-	
+
+protected:
+
 	//
 	// Help
 	//
diff --git a/Engine/Source/Editor/VREditor/Teleporter/VREditorTeleporter.cpp b/Engine/Source/Editor/VREditor/Teleporter/VREditorTeleporter.cpp
--- a/Engine/Source/Editor/VREditor/Teleporter/VREditorTeleporter.cpp
+++ b/Engine/Source/Editor/VREditor/Teleporter/VREditorTeleporter.cpp
@@ -38,6 +38,8 @@ void UVREditorTeleporter::Init( IVREditorMode* InMode )
 
 void UVREditorTeleporter::Shutdown()
 {
+	VRMode->OnTickHandle().RemoveAll( this );
+	VRMode->GetWorldInteraction().OnViewportInteractionInputAction().RemoveAll( this );
 	TeleportSound = nullptr;
 }
 
@@ -57,7 +59,7 @@ void UVREditorTeleporter::Tick( const float DeltaTime )
 void UVREditorTeleporter::StartTeleport( UViewportInteractor* Interactor )
 {
 	UVREditorInteractor* VREditorInteractor = Cast<UVREditorInteractor>( Interactor );
-	if ( VREditorInteractor == nullptr || !VREditorInteractor->IsHoveringOverUI() )
+	if ( VREditorInteractor && !VREditorInteractor->IsHoveringOverUI() )
 	{
 		FVector LaserPointerStart, LaserPointerEnd;
 		if ( Interactor->GetLaserPointer( LaserPointerStart, LaserPointerEnd ) )
diff --git a/Engine/Source/Editor/VREditor/VREditorDockableWindow.cpp b/Engine/Source/Editor/VREditor/VREditorDockableWindow.cpp
--- a/Engine/Source/Editor/VREditor/VREditorDockableWindow.cpp
+++ b/Engine/Source/Editor/VREditor/VREditorDockableWindow.cpp
@@ -426,7 +426,7 @@ void UDockableWindowDragOperation::ExecuteDrag( UViewportInteractor* Interactor,
 		UVREditorUISystem& UISystem = DockableWindow->GetOwner();
 
 		const bool bIsAbsolute = ( UISystem.GetOwner().GetHMDDeviceType() == EHMDDeviceType::DT_SteamVR );
-		float SlideDeltaY = VRInteractor->GetTrackpadSlideDelta();
+		float SlideDeltaY = VRInteractor->GetSlideDelta();
 
 		float NewUIScale = DockableWindow->GetScale() + SlideDeltaY;
 		if ( NewUIScale <= UISystem.GetMinDockWindowSize() )
diff --git a/Engine/Source/Editor/VREditor/VREditorInteractor.cpp b/Engine/Source/Editor/VREditor/VREditorInteractor.cpp
--- a/Engine/Source/Editor/VREditor/VREditorInteractor.cpp
+++ b/Engine/Source/Editor/VREditor/VREditorInteractor.cpp
@@ -5,67 +5,24 @@
 #include "VREditorWorldInteraction.h"
 
 #include "VREditorMode.h"
-#include "VREditorUISystem.h"
 #include "VREditorFloatingText.h"
-#include "MotionControllerComponent.h"
 
 #include "Interactables/VREditorButton.h"
 #include "VREditorDockableWindow.h"
 
-namespace VREd
-{
-	static FAutoConsoleVariable HelpLabelFadeDuration( TEXT( "VREd.HelpLabelFadeDuration" ), 0.4f, TEXT( "Duration to fade controller help labels in and out" ) );
-	static FAutoConsoleVariable HelpLabelFadeDistance( TEXT( "VREd.HelpLabelFadeDistance" ), 30.0f, TEXT( "Distance at which controller help labels should appear (in cm)" ) );
-	
-	static FAutoConsoleVariable LaserPointerRadius( TEXT( "VREd.LaserPointerRadius" ), 0.5f, TEXT( "Radius of the laser pointer line" ) );
-	static FAutoConsoleVariable LaserPointerHoverBallRadius( TEXT( "VREd.LaserPointerHoverBallRadius" ), 1.5f, TEXT( "Radius of the visual cue for a hovered object along the laser pointer ray" ) );
-	static FAutoConsoleVariable LaserPointerLightPullBackDistance( TEXT( "VREd.LaserPointerLightPullBackDistance" ), 2.5f, TEXT( "How far to pull back our little hover light from the impact surface" ) );
-	static FAutoConsoleVariable LaserPointerLightRadius( TEXT( "VREd.LaserPointLightRadius" ), 20.0f, TEXT( "How big our hover light is" ) );
-	static FAutoConsoleVariable OculusLaserPointerStartOffset( TEXT( "VREd.OculusLaserPointerStartOffset" ), 2.8f, TEXT( "How far to offset the start of the laser pointer to avoid overlapping the hand mesh geometry (Oculus)" ) );
-	static FAutoConsoleVariable ViveLaserPointerStartOffset( TEXT( "VREd.ViveLaserPointerStartOffset" ), 1.25f /* 8.5f */, TEXT( "How far to offset the start of the laser pointer to avoid overlapping the hand mesh geometry (Vive)" ) );
-
-	static FAutoConsoleVariable TrackpadAbsoluteDragSpeed( TEXT( "VREd.TrackpadAbsoluteDragSpeed" ), 40.0f, TEXT( "How fast objects move toward or away when you drag on the touchpad while carrying them" ) );
-	static FAutoConsoleVariable TrackpadRelativeDragSpeed( TEXT( "VREd.TrackpadRelativeDragSpeed" ), 8.0f, TEXT( "How fast objects move toward or away when you hold a direction on an analog stick while carrying them" ) );
-	static FAutoConsoleVariable InvertTrackpadVertical( TEXT( "VREd.InvertTrackpadVertical" ), 1, TEXT( "Toggles inverting the touch pad vertical axis" ) );
-	static FAutoConsoleVariable MinVelocityForInertia( TEXT( "VREd.MinVelocityForInertia" ), 1.0f, TEXT( "Minimum velocity (in cm/frame in unscaled room space) before inertia will kick in when releasing objects (or the world)" ) );
-	static FAutoConsoleVariable MinJoystickOffsetBeforeRadialMenu( TEXT( "VREd.MinJoystickOffsetBeforeRadialMenu" ), 0.15f, TEXT( "Toggles inverting the touch pad vertical axis" ) );
-	static FAutoConsoleVariable TriggerLightlyPressedLockTime( TEXT( "VREd.TriggerLightlyPressedLockTime" ), 0.15f, TEXT( "If the trigger remains lightly pressed for longer than this, we'll continue to treat it as a light press in some cases" ) );
-}
-
-namespace VREditorKeyNames
-{
-	// @todo vreditor input: Ideally these would not be needed, but SteamVR fires off it's "trigger pressed" event
-	// well before the trigger is fully down (*click*)
-	static const FName MotionController_Left_FullyPressedTriggerAxis( "MotionController_Left_FullyPressedTriggerAxis" );
-	static const FName MotionController_Right_FullyPressedTriggerAxis( "MotionController_Right_FullyPressedTriggerAxis" );
-
-	static const FName MotionController_Left_LightlyPressedTriggerAxis( "MotionController_Left_LightlyPressedTriggerAxis" );
-	static const FName MotionController_Right_LightlyPressedTriggerAxis( "MotionController_Right_LightlyPressedTriggerAxis" );
-}
-
-namespace SteamVRControllerKeyNames
-{
-	static const FGamepadKeyNames::Type Touch0( "Steam_Touch_0" );
-	static const FGamepadKeyNames::Type Touch1( "Steam_Touch_1" );
-}
-
 #define LOCTEXT_NAMESPACE "VREditor"
 
 UVREditorInteractor::UVREditorInteractor( const FObjectInitializer& Initializer ) :
-	UViewportInteractorMotionController( Initializer ),
+	Super( Initializer ),
 	VRMode( nullptr ),
+	bIsModifierPressed( false ),
+	SelectAndMoveTriggerValue( 0.0f ),
 	bHasUIInFront( false ),
 	bHasUIOnForearm( false ),
 	bIsClickingOnUI( false ),
 	bIsRightClickingOnUI( false ),
 	bIsHoveringOverUI( false ),
 	UIScrollVelocity( 0.0f ),
-	bIsTouchingTrackpad( false ),
-	TrackpadPosition( FVector2D::ZeroVector ),
-	LastTrackpadPosition( FVector2D::ZeroVector ),
-	bIsTrackpadPositionValid { false, false },
-	LastTrackpadPositionUpdateTime( FTimespan::MinValue() ),
-	bIsModifierPressed( false ),
 	LastClickReleaseTime( 0.0f ),
 	bWantHelpLabels( false ),
 	HelpLabelShowOrHideStartTime( FTimespan::MinValue() )
@@ -78,169 +35,10 @@ UVREditorInteractor::~UVREditorInteractor()
 	Shutdown();
 }
 
-void UVREditorInteractor::Init( const EControllerHand InHandSide, FVREditorMode* InVRMode )
+void UVREditorInteractor::Init( FVREditorMode* InVRMode )
 {
 	VRMode = InVRMode;
-	ControllerHandSide = InHandSide;
-	bHaveMotionController = true;
 	KeyToActionMap.Reset();
-
-	// Setup keys
-	if ( InHandSide == EControllerHand::Left )
-	{
-		AddKeyAction( EKeys::MotionController_Left_Grip1, FViewportActionKeyInput( ViewportWorldActionTypes::WorldMovement ) );
-		AddKeyAction( UViewportInteractorMotionController::MotionController_Left_FullyPressedTriggerAxis, FViewportActionKeyInput( ViewportWorldActionTypes::SelectAndMove ) );
-		AddKeyAction( UViewportInteractorMotionController::MotionController_Left_LightlyPressedTriggerAxis, FViewportActionKeyInput( ViewportWorldActionTypes::SelectAndMove_LightlyPressed ) );
-		AddKeyAction( EKeys::MotionController_Left_Thumbstick, FViewportActionKeyInput( VRActionTypes::ConfirmRadialSelection ) );
-
-		AddKeyAction( SteamVRControllerKeyNames::Touch0, FViewportActionKeyInput( VRActionTypes::Touch ) );
-		AddKeyAction( EKeys::MotionController_Left_TriggerAxis, FViewportActionKeyInput( UViewportInteractorMotionController::TriggerAxis ) );
-		AddKeyAction( EKeys::MotionController_Left_Thumbstick_X, FViewportActionKeyInput( UViewportInteractorMotionController::TrackpadPositionX ) );
-		AddKeyAction( EKeys::MotionController_Left_Thumbstick_Y, FViewportActionKeyInput( UViewportInteractorMotionController::TrackpadPositionY ) );	
-		AddKeyAction( EKeys::MotionController_Left_Thumbstick, FViewportActionKeyInput( VRActionTypes::ConfirmRadialSelection ) );
-
-		if ( GetVRMode().GetHMDDeviceType() == EHMDDeviceType::DT_SteamVR )
-		{
-			AddKeyAction( EKeys::MotionController_Left_Shoulder, FViewportActionKeyInput( VRActionTypes::Modifier ) );
-		}
-		else
-		{
-			AddKeyAction( EKeys::MotionController_Left_FaceButton1, FViewportActionKeyInput( VRActionTypes::Modifier ) );
-		}
-	}
-	else if ( InHandSide == EControllerHand::Right )
-	{
-		AddKeyAction( EKeys::MotionController_Right_Grip1, FViewportActionKeyInput( ViewportWorldActionTypes::WorldMovement ) );
-		AddKeyAction( UViewportInteractorMotionController::MotionController_Right_FullyPressedTriggerAxis, FViewportActionKeyInput( ViewportWorldActionTypes::SelectAndMove ) );
-		AddKeyAction( UViewportInteractorMotionController::MotionController_Right_LightlyPressedTriggerAxis, FViewportActionKeyInput( ViewportWorldActionTypes::SelectAndMove_LightlyPressed ) );
-		AddKeyAction( EKeys::MotionController_Right_Thumbstick, FViewportActionKeyInput( VRActionTypes::ConfirmRadialSelection ) );
-
-		AddKeyAction( SteamVRControllerKeyNames::Touch1, FViewportActionKeyInput( VRActionTypes::Touch ) );
-		AddKeyAction( EKeys::MotionController_Right_TriggerAxis, FViewportActionKeyInput( UViewportInteractorMotionController::TriggerAxis ) );
-		AddKeyAction( EKeys::MotionController_Right_Thumbstick_X, FViewportActionKeyInput( UViewportInteractorMotionController::TrackpadPositionX ) );
-		AddKeyAction( EKeys::MotionController_Right_Thumbstick_Y, FViewportActionKeyInput( UViewportInteractorMotionController::TrackpadPositionY ) );
-		AddKeyAction( EKeys::MotionController_Right_Thumbstick, FViewportActionKeyInput( VRActionTypes::ConfirmRadialSelection ) );
-
-		if ( GetVRMode().GetHMDDeviceType() == EHMDDeviceType::DT_SteamVR )
-		{
-			AddKeyAction( EKeys::MotionController_Right_Shoulder, FViewportActionKeyInput( VRActionTypes::Modifier ) );
-		}
-		else
-		{
-			AddKeyAction( EKeys::MotionController_Right_FaceButton1, FViewportActionKeyInput( VRActionTypes::Modifier ) );
-		}
-	}
-}
-
-void UVREditorInteractor::SetupComponent( AActor* OwningActor )
-{
-	// Setup a motion controller component.  This allows us to take advantage of late frame updates, so
-	// our motion controllers won't lag behind the HMD
-	{
-		MotionControllerComponent = NewObject<UMotionControllerComponent>( OwningActor );
-		OwningActor->AddOwnedComponent( MotionControllerComponent );
-		MotionControllerComponent->SetupAttachment( OwningActor->GetRootComponent() );
-		MotionControllerComponent->RegisterComponent();
-
-		MotionControllerComponent->SetMobility( EComponentMobility::Movable );
-		MotionControllerComponent->SetCollisionEnabled( ECollisionEnabled::NoCollision );
-
-		MotionControllerComponent->Hand = ControllerHandSide;
-
-		MotionControllerComponent->bDisableLowLatencyUpdate = false;
-	}
-
-	// Hand mesh
-	{
-		HandMeshComponent = NewObject<UStaticMeshComponent>( OwningActor );
-		OwningActor->AddOwnedComponent( HandMeshComponent );
-		HandMeshComponent->SetupAttachment( MotionControllerComponent );
-		HandMeshComponent->RegisterComponent();
-
-		// @todo vreditor extensibility: We need this to be able to be overridden externally, or simply based on the HMD name (but allowing external folders)
-		FString MeshName;
-		FString MaterialName;
-		if ( GetVRMode().GetHMDDeviceType() == EHMDDeviceType::DT_SteamVR )
-		{
-			MeshName = TEXT( "/Engine/VREditor/Devices/Vive/VivePreControllerMesh" );
-			MaterialName = TEXT( "/Engine/VREditor/Devices/Vive/VivePreControllerMaterial_Inst" );
-		}
-		else
-		{
-			MeshName = TEXT( "/Engine/VREditor/Devices/Oculus/OculusControllerMesh" );
-			MaterialName = TEXT( "/Engine/VREditor/Devices/Oculus/OculusControllerMaterial_Inst" );
-		}
-
-		UStaticMesh* HandMesh = LoadObject<UStaticMesh>( nullptr, *MeshName );
-		check( HandMesh != nullptr );
-
-		HandMeshComponent->SetStaticMesh( HandMesh );
-		HandMeshComponent->SetMobility( EComponentMobility::Movable );
-		HandMeshComponent->SetCollisionEnabled( ECollisionEnabled::NoCollision );
-
-		UMaterialInstance* HandMeshMaterialInst = LoadObject<UMaterialInstance>( nullptr, *MaterialName );
-		check( HandMeshMaterialInst != nullptr );
-		HandMeshMID = UMaterialInstanceDynamic::Create( HandMeshMaterialInst, GetTransientPackage() );
-		check( HandMeshMID != nullptr );
-		HandMeshComponent->SetMaterial( 0, HandMeshMID );
-	}
-
-	// Laser pointer
-	{
-		LaserPointerMeshComponent = NewObject<UStaticMeshComponent>( OwningActor );
-		OwningActor->AddOwnedComponent( LaserPointerMeshComponent );
-		LaserPointerMeshComponent->SetupAttachment( MotionControllerComponent );
-		LaserPointerMeshComponent->RegisterComponent();
-
-		UStaticMesh* LaserPointerMesh = LoadObject<UStaticMesh>( nullptr, TEXT( "/Engine/VREditor/LaserPointer/LaserPointerMesh" ) );
-		check( LaserPointerMesh != nullptr );
-
-		LaserPointerMeshComponent->SetStaticMesh( LaserPointerMesh );
-		LaserPointerMeshComponent->SetMobility( EComponentMobility::Movable );
-		LaserPointerMeshComponent->SetCollisionEnabled( ECollisionEnabled::NoCollision );
-
-		UMaterialInstance* LaserPointerMaterialInst = LoadObject<UMaterialInstance>( nullptr, TEXT( "/Engine/VREditor/LaserPointer/LaserPointerMaterialInst" ) );
-		check( LaserPointerMaterialInst != nullptr );
-		LaserPointerMID = UMaterialInstanceDynamic::Create( LaserPointerMaterialInst, GetTransientPackage() );
-		check( LaserPointerMID != nullptr );
-		LaserPointerMeshComponent->SetMaterial( 0, LaserPointerMID );
-
-		UMaterialInstance* TranslucentLaserPointerMaterialInst = LoadObject<UMaterialInstance>( nullptr, TEXT( "/Engine/VREditor/LaserPointer/TranslucentLaserPointerMaterialInst" ) );
-		check( TranslucentLaserPointerMaterialInst != nullptr );
-		TranslucentLaserPointerMID = UMaterialInstanceDynamic::Create( TranslucentLaserPointerMaterialInst, GetTransientPackage() );
-		check( TranslucentLaserPointerMID != nullptr );
-		LaserPointerMeshComponent->SetMaterial( 1, TranslucentLaserPointerMID );
-	}
-
-	// Hover cue for laser pointer
-	{
-		HoverMeshComponent = NewObject<UStaticMeshComponent>( OwningActor );
-		OwningActor->AddOwnedComponent( HoverMeshComponent );
-		HoverMeshComponent->SetupAttachment( OwningActor->GetRootComponent() );
-		HoverMeshComponent->RegisterComponent();
-
-		UStaticMesh* HoverMesh = LoadObject<UStaticMesh>( nullptr, TEXT( "/Engine/VREditor/LaserPointer/HoverMesh" ) );
-		check( HoverMesh != nullptr );
-		HoverMeshComponent->SetStaticMesh( HoverMesh );
-		HoverMeshComponent->SetMobility( EComponentMobility::Movable );
-		HoverMeshComponent->SetCollisionEnabled( ECollisionEnabled::NoCollision );
-
-		// Add a light!
-		{
-			HoverPointLightComponent = NewObject<UPointLightComponent>( OwningActor );
-			OwningActor->AddOwnedComponent( HoverPointLightComponent );
-			HoverPointLightComponent->SetupAttachment( HoverMeshComponent );
-			HoverPointLightComponent->RegisterComponent();
-
-			HoverPointLightComponent->SetLightColor( FLinearColor::Red );
-			//Hand.HoverPointLightComponent->SetLightColor( FLinearColor( 0.0f, 1.0f, 0.2f, 1.0f ) );
-			HoverPointLightComponent->SetIntensity( 30.0f );	// @todo: VREditor tweak
-			HoverPointLightComponent->SetMobility( EComponentMobility::Movable );
-			HoverPointLightComponent->SetAttenuationRadius( VREd::LaserPointerLightRadius->GetFloat() );
-			HoverPointLightComponent->bUseInverseSquaredFalloff = false;
-			HoverPointLightComponent->SetCastShadows( false );
-		}
-	}
 }
 
 void UVREditorInteractor::Shutdown()
@@ -261,145 +59,6 @@ void UVREditorInteractor::Shutdown()
 void UVREditorInteractor::Tick( const float DeltaTime )
 {
 	Super::Tick( DeltaTime );
-	
-	UpdateRadialMenuInput( DeltaTime );
-
-	{
-		const float WorldScaleFactor = WorldInteraction->GetWorldScaleFactor();
-
-		// Don't bother drawing hands if we're not currently tracking them.
-		if ( bHaveMotionController )
-		{
-			HandMeshComponent->SetVisibility( true );
-		}
-		else
-		{
-			HandMeshComponent->SetVisibility( false );
-		}
-
-		FVector LaserPointerStart, LaserPointerEnd;
-		if ( GetLaserPointer( /* Out */ LaserPointerStart, /* Out */ LaserPointerEnd ) )
-		{
-			// NOTE: We don't need to set the laser pointer location and rotation, as the MotionControllerComponent will do
-			// that later in the frame.  
-
-			// If we're actively dragging something around, then we'll crop the laser length to the hover impact
-			// point.  Otherwise, we always want the laser to protrude through hovered objects, so that you can
-			// interact with translucent gizmo handles that are occluded by geometry
-			FVector LaserPointerImpactPoint = LaserPointerEnd;
-			if ( GetInteractorData().bIsHovering && ( GetInteractorData().DraggingMode != EViewportInteractionDraggingMode::Nothing || bIsHoveringOverUI ) )
-			{
-				LaserPointerImpactPoint = GetInteractorData().HoverLocation;
-			}
-
-			// Apply rotation offset to the laser direction
-			const float LaserPointerRotationOffset = 0.0f;
-			LaserPointerMeshComponent->SetRelativeRotation( FRotator( LaserPointerRotationOffset, 0.0f, 0.0f ) );
-
-			const FVector LaserPointerDirection = ( LaserPointerImpactPoint - LaserPointerStart ).GetSafeNormal();
-
-			// Offset the beginning of the laser pointer a bit, so that it doesn't overlap the hand mesh
-			const float LaserPointerStartOffset =
-				WorldScaleFactor *
-				( GetVRMode().GetHMDDeviceType() == EHMDDeviceType::DT_OculusRift ? VREd::OculusLaserPointerStartOffset->GetFloat() : VREd::ViveLaserPointerStartOffset->GetFloat() );
-
-			const float LaserPointerLength = FMath::Max( 0.000001f, ( LaserPointerImpactPoint - LaserPointerStart ).Size() - LaserPointerStartOffset );
-
-			LaserPointerMeshComponent->SetRelativeLocation( FRotator( LaserPointerRotationOffset, 0.0f, 0.0f ).RotateVector( FVector( LaserPointerStartOffset, 0.0f, 0.0f ) ) );
-
-			// The laser pointer needs to stay the same size relative to our tracking space, so we inverse compensate for world to meters scale here
-			float LaserPointerRadius = VREd::LaserPointerRadius->GetFloat() * WorldScaleFactor;
-			float HoverMeshRadius = VREd::LaserPointerHoverBallRadius->GetFloat() * WorldScaleFactor;
-
-			// If we're hovering over something really close to the camera, go ahead and shrink the effect
-			// @todo vreditor: Can we make this actually just sized based on distance automatically?  The beam and impact point are basically a cone.
-			if ( bIsHoveringOverUI )
-			{
-				LaserPointerRadius *= 0.35f;	// @todo vreditor tweak
-				HoverMeshRadius *= 0.35f;	// @todo vreditor tweak
-			}
-
-			LaserPointerMeshComponent->SetRelativeScale3D( FVector( LaserPointerLength, LaserPointerRadius * 2.0f, LaserPointerRadius * 2.0f ) );
-
-			if ( IsHovering() )
-			{
-				// The hover effect needs to stay the same size relative to our tracking space, so we inverse compensate for world to meters scale here
-				HoverMeshComponent->SetRelativeScale3D( FVector( HoverMeshRadius * 2.0f ) );
-				HoverMeshComponent->SetVisibility( true );
-				HoverMeshComponent->SetWorldLocation( GetHoverLocation() );
-
-				// Show the light too, unless it's on top of UI.  It looks too distracting on top of UI.
-				HoverPointLightComponent->SetVisibility( !bIsHoveringOverUI );
-
-				// Update radius for world scaling
-				HoverPointLightComponent->SetAttenuationRadius( VREd::LaserPointerLightRadius->GetFloat() * WorldScaleFactor );
-
-				// Pull hover light back a bit from the end of the ray
-				const float PullBackAmount = VREd::LaserPointerLightPullBackDistance->GetFloat() * WorldInteraction->GetWorldScaleFactor();
-				HoverPointLightComponent->SetWorldLocation( GetHoverLocation() - PullBackAmount * LaserPointerDirection );
-			}
-			else
-			{
-				HoverMeshComponent->SetVisibility( false );
-				HoverPointLightComponent->SetVisibility( false );
-			}
-		}
-		else
-		{
-			LaserPointerMeshComponent->SetVisibility( false );
-			HoverMeshComponent->SetVisibility( false );
-			HoverPointLightComponent->SetVisibility( false );
-		}
-	}
-
-	// Updating laser colors for both hands
-	{
-		FVREditorMode::EColors ResultColorType = FVREditorMode::EColors::DefaultColor;
-		float CrawlSpeed = 0.0f;
-		float CrawlFade = 0.0f;
-			
-		const EViewportInteractionDraggingMode DraggingMode = GetDraggingMode();
-
-		//@todo VREditor
-		const bool bIsDraggingWorldWithTwoHands = false;
-			( DraggingMode == EViewportInteractionDraggingMode::AssistingDrag && GetOtherInteractor()->GetDraggingMode() == EViewportInteractionDraggingMode::World ) ||
-			( DraggingMode == EViewportInteractionDraggingMode::World && GetOtherInteractor()->GetDraggingMode() == EViewportInteractionDraggingMode::AssistingDrag );
-
-		if ( bIsDraggingWorldWithTwoHands )
-		{
-			ResultColorType = FVREditorMode::EColors::WorldDraggingColor_TwoHands;
-		}
-		else if ( DraggingMode == EViewportInteractionDraggingMode::World )
-		{
-			//@todo VREditor find a way to make the teleporter do it
-// 			if ( GetVRMode().GetWorldInteraction().IsTeleporting() ) //@todo VREditor: Make teleport an action
-// 			{
-// 				ResultColorType = FVREditorMode::EColors::TeleportColor;
-// 			}
-// 			else
-// 			{
-				// We can teleport in this mode, so animate the laser a bit
-				CrawlFade = 1.0f;
-				CrawlSpeed = 5.0f;
-				ResultColorType = FVREditorMode::EColors::WorldDraggingColor_OneHand;
-//			}
-		}
-		else if ( DraggingMode == EViewportInteractionDraggingMode::ActorsAtLaserImpact ||
-			DraggingMode == EViewportInteractionDraggingMode::Material ||
-			DraggingMode == EViewportInteractionDraggingMode::ActorsFreely ||
-			DraggingMode == EViewportInteractionDraggingMode::ActorsWithGizmo ||
-			DraggingMode == EViewportInteractionDraggingMode::AssistingDrag ||
-			DraggingMode == EViewportInteractionDraggingMode::Interactable ||
-			( GetVRMode().GetUISystem().IsInteractorDraggingDockUI( this ) && GetVRMode().GetUISystem().IsDraggingDockUI() ) )
-		{
-			ResultColorType = FVREditorMode::EColors::SelectionColor;
-		}
-		
-		const FLinearColor ResultColor = GetVRMode().GetColor( ResultColorType );
-		SetLaserVisuals( ResultColor, CrawlFade, CrawlSpeed );
-	}
-
-	UpdateHelpLabels();
 }
 
 FHitResult UVREditorInteractor::GetHitResultFromLaserPointer( TArray<AActor*>* OptionalListOfIgnoredActors /*= nullptr*/, const bool bIgnoreGizmos /*= false*/,
@@ -441,69 +100,6 @@ void UVREditorInteractor::ResetHoverState( const float DeltaTime )
 	bIsHoveringOverUI = false;
 }
 
-void UVREditorInteractor::CalculateDragRay()
-{
-	const FTimespan CurrentTime = FTimespan::FromSeconds( FPlatformTime::Seconds() );
-	const float WorldScaleFactor = WorldInteraction->GetWorldScaleFactor();
-
-	// If we're dragging an object, go ahead and slide the object along the ray based on how far they slide their touch
-	// Make sure they are touching the trackpad, otherwise we get bad data
-	if ( bIsTrackpadPositionValid[ 1 ] )
-	{
-		const bool bIsAbsolute = ( GetVRMode().GetHMDDeviceType() == EHMDDeviceType::DT_SteamVR );
-		float SlideDelta = GetTrackpadSlideDelta() * WorldScaleFactor;
-
-		if ( !FMath::IsNearlyZero( SlideDelta ) )
-		{
-			InteractorData.DragRayLength += SlideDelta;
-
-			InteractorData.DragRayLengthVelocity = 0.0f;
-
-			// Don't apply inertia unless the user dragged a decent amount this frame
-			if ( bIsAbsolute && FMath::Abs( SlideDelta ) >= VREd::MinVelocityForInertia->GetFloat() * WorldScaleFactor )
-			{
-				// Don't apply inertia if our data is sort of old
-				if ( CurrentTime - LastTrackpadPositionUpdateTime <= FTimespan::FromSeconds( 1.0f / 30.0f ) )
-				{
-					InteractorData.DragRayLengthVelocity = SlideDelta;
-				}
-			}
-
-			// Don't go too far
-			if ( InteractorData.DragRayLength < 0.0f )
-			{
-				InteractorData.DragRayLength = 0.0f;
-				InteractorData.DragRayLengthVelocity = 0.0f;
-			}
-		}
-	}
-	else
-	{
-		if ( !FMath::IsNearlyZero( InteractorData.DragRayLengthVelocity ) )
-		{
-			// Apply drag ray length inertia
-			InteractorData.DragRayLength += InteractorData.DragRayLengthVelocity;
-
-			// Don't go too far!
-			if ( InteractorData.DragRayLength < 0.0f )
-			{
-				InteractorData.DragRayLength = 0.0f;
-				InteractorData.DragRayLengthVelocity = 0.0f;
-			}
-
-			// Apply damping
-			FVector RayVelocityVector( InteractorData.DragRayLengthVelocity, 0.0f, 0.0f );
-			const bool bVelocitySensitive = true;
-			WorldInteraction->ApplyVelocityDamping( RayVelocityVector, bVelocitySensitive );
-			InteractorData.DragRayLengthVelocity = RayVelocityVector.X;
-		}
-		else
-		{
-			InteractorData.DragRayLengthVelocity = 0.0f;
-		}
-	}
-}
-
 void UVREditorInteractor::OnStartDragging( UActorComponent* ClickedComponent, const FVector& HitLocation, const bool bIsPlacingActors )
 {
 	// If the user is holding down the modifier key, go ahead and duplicate the selection first.  This needs to
@@ -515,23 +111,9 @@ void UVREditorInteractor::OnStartDragging( UActorComponent* ClickedComponent, co
 	}
 }
 
-float UVREditorInteractor::GetTrackpadSlideDelta()
+float UVREditorInteractor::GetSlideDelta()
 {
-	const bool bIsAbsolute = ( GetVRMode().GetHMDDeviceType() == EHMDDeviceType::DT_SteamVR );
-	float SlideDelta = 0.0f;
-	if ( bIsTouchingTrackpad || !bIsAbsolute )
-	{
-		if ( bIsAbsolute )
-		{
-			SlideDelta = ( ( TrackpadPosition.Y - LastTrackpadPosition.Y ) * VREd::TrackpadAbsoluteDragSpeed->GetFloat() );
-		}
-		else
-		{
-			SlideDelta = ( TrackpadPosition.Y * VREd::TrackpadRelativeDragSpeed->GetFloat() );
-		}
-	}
-
-	return SlideDelta;
+	return 0.0f;
 }
 
 bool UVREditorInteractor::IsHoveringOverUI() const
@@ -559,21 +141,6 @@ bool UVREditorInteractor::HasUIOnForearm() const
 	return bHasUIOnForearm;
 }
 
-bool UVREditorInteractor::IsTouchingTrackpad() const
-{
-	return bIsTouchingTrackpad;
-}
-
-FVector2D UVREditorInteractor::GetTrackpadPosition() const
-{
-	return TrackpadPosition;
-}
-
-FVector2D UVREditorInteractor::GetLastTrackpadPosition() const
-{
-	return LastTrackpadPosition;
-}
-
 UWidgetComponent* UVREditorInteractor::GetHoveringOverWidgetComponent() const
 {
 	return InteractorData.HoveringOverWidgetComponent;
@@ -584,16 +151,6 @@ void UVREditorInteractor::SetHoveringOverWidgetComponent( UWidgetComponent* NewH
 	InteractorData.HoveringOverWidgetComponent = NewHoveringOverWidgetComponent;
 }
 
-bool UVREditorInteractor::IsTrackpadPositionValid( const int32 AxisIndex ) const
-{
-	return bIsTrackpadPositionValid[ AxisIndex ];
-}
-
-FTimespan& UVREditorInteractor::GetLastTrackpadPositionUpdateTime()
-{
-	return LastTrackpadPositionUpdateTime;
-}
-
 bool UVREditorInteractor::IsModifierPressed() const
 {
 	return bIsModifierPressed;
@@ -644,116 +201,9 @@ float UVREditorInteractor::GetUIScrollVelocity() const
 	return UIScrollVelocity;
 }
 
-void UVREditorInteractor::HandleInputKey( FViewportActionKeyInput& Action, const FKey Key, const EInputEvent Event, bool& bOutWasHandled )
+float UVREditorInteractor::GetSelectAndMoveTriggerValue() const
 {
-	if( !bOutWasHandled )
-	{
-		bool bShouldAbsorbEvent = false;
-
-		// If "SelectAndMove" was pressed, we need to make sure that "SelectAndMove_LightlyPressed" is no longer
-		// in effect before we start handling "SelectAndMove".
-		if ( Action.ActionType == ViewportWorldActionTypes::SelectAndMove && IsTriggerAtLeastLightlyPressed() )
-		{
-			if ( Event == IE_Pressed )
-			{
-				// How long since the trigger was lightly pressed?
-				const float TimeSinceLightlyPressed = ( float ) ( FPlatformTime::Seconds() - GetRealTimeTriggerWasLightlyPressed() );
-				if ( !bIsClickingOnUI &&	// @todo vreditor: UI clicks happen with light presses; we never want to convert to a hard press!
-					GetDraggingMode() != EViewportInteractionDraggingMode::Material &&	// @todo vreditor: Material dragging happens with light presses, don't convert to a hard press!
-					GetDraggingMode() != EViewportInteractionDraggingMode::ActorsAtLaserImpact &&	// @todo vreditor: Actor dragging happens with light presses, don't convert to a hard press!
-					( ! AllowTriggerLightPressLocking() || TimeSinceLightlyPressed < VREd::TriggerLightlyPressedLockTime->GetFloat() ) )
-				{
-					SetTriggerAtLeastLightlyPressed( false );
-					SetTriggerBeenReleasedSinceLastPress( false );
-
-					// Synthesize an input key for releasing the light press
-					// NOTE: Intentionally re-entrant call here.
-					const EInputEvent InputEvent = IE_Released;
-					const bool bWasLightReleaseHandled = UViewportInteractor::HandleInputKey( GetControllerSide() == EControllerHand::Left ? VREditorKeyNames::MotionController_Left_LightlyPressedTriggerAxis : VREditorKeyNames::MotionController_Right_LightlyPressedTriggerAxis, InputEvent );
-				}
-				else
-				{
-					// The button was held in a "lightly pressed" state for a long enough time that we should just continue to
-					// treat it like a light press.  This makes it much easier to hold it in this state when you need to!
-					bShouldAbsorbEvent = true;
-				}
-			}
-			else
-			{
-				// Absorb the release of the SelectAndMove event if we are still in a lightly pressed state
-				bShouldAbsorbEvent = true;
-			}
-		}
-
-		if ( bShouldAbsorbEvent )
-		{
-			bOutWasHandled = true;
-		}
-	}
-
-	if ( !bOutWasHandled )
-	{
-		// Update touch state
-		if ( Action.ActionType == VRActionTypes::Touch )
-		{
-			if ( Event == IE_Pressed )
-			{
-				bIsTouchingTrackpad = true;
-			}
-			else if ( Event == IE_Released )
-			{
-				bIsTouchingTrackpad = false;
-				bIsTrackpadPositionValid[ 0 ] = false;
-				bIsTrackpadPositionValid[ 1 ] = false;
-			}
-		}
-
-		// Update modifier state
-		if ( Action.ActionType == VRActionTypes::Modifier )
-		{
-			bOutWasHandled = true;
-			if ( Event == IE_Pressed )
-			{
-				bIsModifierPressed = true;
-			}
-			else if ( Event == IE_Released )
-			{
-				bIsModifierPressed = false;
-			}
-		}
-	}
-
-	ApplyButtonPressColors( Action );
-}
-
-void UVREditorInteractor::HandleInputAxis( FViewportActionKeyInput& Action, const FKey Key, const float Delta, const float DeltaTime, bool& bOutWasHandled )
-{
-	Super::HandleInputAxis( Action, Key, Delta, DeltaTime, bOutWasHandled );
-
-	if ( !bOutWasHandled )
-	{
-		if ( Action.ActionType == VRActionTypes::TrackpadPositionX )
-		{
-			LastTrackpadPosition.X = bIsTrackpadPositionValid[ 0 ] ? TrackpadPosition.X : Delta;
-			LastTrackpadPositionUpdateTime = FTimespan::FromSeconds( FPlatformTime::Seconds() );
-			TrackpadPosition.X = Delta;
-			bIsTrackpadPositionValid[ 0 ] = true;
-		}
-			
-		if ( Action.ActionType == VRActionTypes::TrackpadPositionY )
-		{
-			float DeltaAxis = Delta;
-			if ( VREd::InvertTrackpadVertical->GetInt() != 0 )
-			{
-				DeltaAxis = -DeltaAxis;	// Y axis is inverted from HMD
-			}
-
-			LastTrackpadPosition.Y = bIsTrackpadPositionValid[ 1 ] ? TrackpadPosition.Y : DeltaAxis;
-			LastTrackpadPositionUpdateTime = FTimespan::FromSeconds( FPlatformTime::Seconds() );
-			TrackpadPosition.Y = DeltaAxis;
-			bIsTrackpadPositionValid[ 1 ] = true;
-		}
-	}
+	return SelectAndMoveTriggerValue;
 }
 
 bool UVREditorInteractor::GetIsLaserBlocked()
@@ -761,370 +211,4 @@ bool UVREditorInteractor::GetIsLaserBlocked()
 	return bHasUIInFront;
 }
 
-/** Changes the color of the buttons on the handmesh */
-void UVREditorInteractor::ApplyButtonPressColors( const FViewportActionKeyInput& Action )
-{
-	const float PressStrength = 10.0f;
-	const FName ActionType = Action.ActionType;
-	const EInputEvent Event = Action.Event;
-
-	// Trigger
-	if ( ActionType == ViewportWorldActionTypes::SelectAndMove|| 
-		ActionType == ViewportWorldActionTypes::SelectAndMove_LightlyPressed )
-	{
-		static FName StaticTriggerParameter( "B1" );
-		SetMotionControllerButtonPressedVisuals( Event, StaticTriggerParameter, PressStrength );
-	}
-
-	// Shoulder button
-	if ( ActionType == ViewportWorldActionTypes::WorldMovement )
-	{
-		static FName StaticShoulderParameter( "B2" );
-		SetMotionControllerButtonPressedVisuals( Event, StaticShoulderParameter, PressStrength );
-	}
-
-	// Trackpad
-	if ( ActionType == VRActionTypes::ConfirmRadialSelection )
-	{
-		static FName StaticTrackpadParameter( "B3" );
-		SetMotionControllerButtonPressedVisuals( Event, StaticTrackpadParameter, PressStrength );
-	}
-
-	// Modifier
-	if ( ActionType == VRActionTypes::Modifier )
-	{
-		static FName StaticModifierParameter( "B4" );
-		SetMotionControllerButtonPressedVisuals( Event, StaticModifierParameter, PressStrength );
-	}
-}
-
-void UVREditorInteractor::SetMotionControllerButtonPressedVisuals( const EInputEvent Event, const FName& ParameterName, const float PressStrength )
-{
-	if ( Event == IE_Pressed )
-	{
-		HandMeshMID->SetScalarParameterValue( ParameterName, PressStrength );
-	}
-	else if ( Event == IE_Released )
-	{
-		HandMeshMID->SetScalarParameterValue( ParameterName, 0.0f );
-	}
-}
-
-void UVREditorInteractor::ShowHelpForHand( const bool bShowIt )
-{
-	if ( bShowIt != bWantHelpLabels )
-	{
-		bWantHelpLabels = bShowIt;
-
-		const FTimespan CurrentTime = FTimespan::FromSeconds( FApp::GetCurrentTime() );
-		const FTimespan TimeSinceStartedFadingOut = CurrentTime - HelpLabelShowOrHideStartTime;
-		const FTimespan HelpLabelFadeDuration = FTimespan::FromSeconds( VREd::HelpLabelFadeDuration->GetFloat() );
-
-		// If we were already fading, account for that here
-		if ( TimeSinceStartedFadingOut < HelpLabelFadeDuration )
-		{
-			// We were already fading, so we'll reverse the time value so it feels continuous
-			HelpLabelShowOrHideStartTime = CurrentTime - ( HelpLabelFadeDuration - TimeSinceStartedFadingOut );
-		}
-		else
-		{
-			HelpLabelShowOrHideStartTime = FTimespan::FromSeconds( FApp::GetCurrentTime() );
-		}
-
-		if ( bShowIt && HelpLabels.Num() == 0 )
-		{
-			for ( const auto& KeyToAction : KeyToActionMap )
-			{
-				const FKey Key = KeyToAction.Key;
-				const FViewportActionKeyInput& Action = KeyToAction.Value;
-
-				UStaticMeshSocket* Socket = FindMeshSocketForKey( HandMeshComponent->StaticMesh, Key );
-				if ( Socket != nullptr )
-				{
-					FText LabelText;
-					FString ComponentName;
-					
-					if ( Action.ActionType == VRActionTypes::Modifier )
-					{
-						LabelText = LOCTEXT( "ModifierHelp", "Modifier" );
-						ComponentName = TEXT( "ModifierHelp" );
-					}
-					else if ( Action.ActionType == ViewportWorldActionTypes::WorldMovement )
-					{
-						LabelText = LOCTEXT( "WorldMovementHelp", "Move World" );
-						ComponentName = TEXT( "WorldMovementHelp" );
-					}
-					else if ( Action.ActionType == ViewportWorldActionTypes::SelectAndMove )
-					{
-						LabelText = LOCTEXT( "SelectAndMoveHelp", "Select & Move" );
-						ComponentName = TEXT( "SelectAndMoveHelp" );
-					}
-					else if ( Action.ActionType == ViewportWorldActionTypes::SelectAndMove_LightlyPressed )
-					{
-						LabelText = LOCTEXT( "SelectAndMove_LightlyPressedHelp", "Select & Move" );
-						ComponentName = TEXT( "SelectAndMove_LightlyPressedHelp" );
-					}
-					else if ( Action.ActionType == VRActionTypes::Touch )
-					{
-						LabelText = LOCTEXT( "TouchHelp", "Slide" );
-						ComponentName = TEXT( "TouchHelp" );
-					}
-					else if ( Action.ActionType == ViewportWorldActionTypes::Undo )
-					{
-						LabelText = LOCTEXT( "UndoHelp", "Undo" );
-						ComponentName = TEXT( "UndoHelp" );
-					}
-					else if ( Action.ActionType == ViewportWorldActionTypes::Redo )
-					{
-						LabelText = LOCTEXT( "RedoHelp", "Redo" );
-						ComponentName = TEXT( "RedoHelp" );
-					}
-					else if ( Action.ActionType == ViewportWorldActionTypes::Delete )
-					{
-						LabelText = LOCTEXT( "DeleteHelp", "Delete" );
-						ComponentName = TEXT( "DeleteHelp" );
-					}
-					else if ( Action.ActionType == VRActionTypes::ConfirmRadialSelection )
-					{
-						LabelText = LOCTEXT( "ConfirmRadialSelectionHelp", "Radial Menu" );
-						ComponentName = TEXT( "ConfirmRadialSelectionHelp" );
-					}
-
-					const bool bWithSceneComponent = false;	// Nope, we'll spawn our own inside AFloatingText
-					check( VRMode );
-					AFloatingText* FloatingText = GetVRMode().SpawnTransientSceneActor<AFloatingText>( ComponentName, bWithSceneComponent );
-					FloatingText->SetText( LabelText );
-
-					HelpLabels.Add( Key, FloatingText );
-				}
-			}
-		}
-	}
-}
-
-
-void UVREditorInteractor::UpdateHelpLabels()
-{
-	const FTimespan HelpLabelFadeDuration = FTimespan::FromSeconds( VREd::HelpLabelFadeDuration->GetFloat() );
-
-	const FTransform HeadTransform = GetVRMode().GetHeadTransform();
-
-	// Only show help labels if the hand is pretty close to the face
-	const float DistanceToHead = ( GetTransform().GetLocation() - HeadTransform.GetLocation() ).Size();
-	const float MinDistanceToHeadForHelp = VREd::HelpLabelFadeDistance->GetFloat() * GetVRMode().GetWorldScaleFactor();	// (in cm)
-	bool bShowHelp = DistanceToHead <= MinDistanceToHeadForHelp;
-
-	// Don't show help if a UI is summoned on that hand
-	if ( bHasUIInFront || bHasUIOnForearm || GetVRMode().GetUISystem().IsShowingRadialMenu( this ) )
-	{
-		bShowHelp = false;
-	}
-
-	ShowHelpForHand( bShowHelp );
-
-	// Have the labels finished fading out?  If so, we'll kill their actors!
-	const FTimespan CurrentTime = FTimespan::FromSeconds( FApp::GetCurrentTime() );
-	const FTimespan TimeSinceStartedFadingOut = CurrentTime - HelpLabelShowOrHideStartTime;
-	if ( !bWantHelpLabels && ( TimeSinceStartedFadingOut > HelpLabelFadeDuration ) )
-	{
-		// Get rid of help text
-		for ( auto& KeyAndValue : HelpLabels )
-		{
-			AFloatingText* FloatingText = KeyAndValue.Value;
-			GetVRMode().DestroyTransientActor( FloatingText );
-		}
-		HelpLabels.Reset();
-	}
-	else
-	{
-		// Update fading state
-		float FadeAlpha = FMath::Clamp( ( float ) TimeSinceStartedFadingOut.GetTotalSeconds() / ( float ) HelpLabelFadeDuration.GetTotalSeconds(), 0.0f, 1.0f );
-		if ( !bWantHelpLabels )
-		{
-			FadeAlpha = 1.0f - FadeAlpha;
-		}
-
-		// Exponential falloff, so the fade is really obvious (gamma/HDR)
-		FadeAlpha = FMath::Pow( FadeAlpha, 3.0f );
-
-		for ( auto& KeyAndValue : HelpLabels )
-		{
-			const FKey Key = KeyAndValue.Key;
-			AFloatingText* FloatingText = KeyAndValue.Value;
-
-			UStaticMeshSocket* Socket = FindMeshSocketForKey( HandMeshComponent->StaticMesh, Key );
-			check( Socket != nullptr );
-			FTransform SocketRelativeTransform( Socket->RelativeRotation, Socket->RelativeLocation, Socket->RelativeScale );
-
-			// Oculus has asymmetrical controllers, so we the sock transform horizontally
-			if ( ControllerHandSide == EControllerHand::Right &&
-				GetVRMode().GetHMDDeviceType() == EHMDDeviceType::DT_OculusRift )
-			{
-				const FVector Scale3D = SocketRelativeTransform.GetLocation();
-				SocketRelativeTransform.SetLocation( FVector( Scale3D.X, -Scale3D.Y, Scale3D.Z ) );
-			}
-
-			// Make sure the labels stay the same size even when the world is scaled
-			FTransform HandTransformWithWorldToMetersScaling = GetTransform();
-			HandTransformWithWorldToMetersScaling.SetScale3D( HandTransformWithWorldToMetersScaling.GetScale3D() * FVector( GetVRMode().GetWorldScaleFactor() ) );
-
-			// Position right on top of the controller itself
-			FTransform FloatingTextTransform = SocketRelativeTransform * HandTransformWithWorldToMetersScaling;
-			FloatingText->SetActorTransform( FloatingTextTransform );
-
-			// Orientate it toward the viewer
-			FloatingText->Update( HeadTransform.GetLocation() );
-
-			// Update fade state
-			FloatingText->SetOpacity( FadeAlpha );
-		}
-	}
-}
-
-
-UStaticMeshSocket* UVREditorInteractor::FindMeshSocketForKey( UStaticMesh* StaticMesh, const FKey Key )
-{
-	// @todo vreditor: Hard coded mapping of socket names (e.g. "Shoulder") to expected names of sockets in the static mesh
-	FName SocketName = NAME_None;
-	if ( Key == EKeys::MotionController_Left_Shoulder || Key == EKeys::MotionController_Right_Shoulder )
-	{
-		static FName ShoulderSocketName( "Shoulder" );
-		SocketName = ShoulderSocketName;
-	}
-	else if ( Key == EKeys::MotionController_Left_Trigger || Key == EKeys::MotionController_Right_Trigger ||
-		Key == VREditorKeyNames::MotionController_Left_FullyPressedTriggerAxis || Key == VREditorKeyNames::MotionController_Right_FullyPressedTriggerAxis ||
-		Key == VREditorKeyNames::MotionController_Left_LightlyPressedTriggerAxis || Key == VREditorKeyNames::MotionController_Right_LightlyPressedTriggerAxis )
-	{
-		static FName TriggerSocketName( "Trigger" );
-		SocketName = TriggerSocketName;
-	}
-	else if ( Key == EKeys::MotionController_Left_Grip1 || Key == EKeys::MotionController_Right_Grip1 )
-	{
-		static FName GripSocketName( "Grip" );
-		SocketName = GripSocketName;
-	}
-	else if ( Key == EKeys::MotionController_Left_Thumbstick || Key == EKeys::MotionController_Right_Thumbstick )
-	{
-		static FName ThumbstickSocketName( "Thumbstick" );
-		SocketName = ThumbstickSocketName;
-	}
-	else if ( Key == SteamVRControllerKeyNames::Touch0 || Key == SteamVRControllerKeyNames::Touch1 )
-	{
-		static FName TouchSocketName( "Touch" );
-		SocketName = TouchSocketName;
-	}
-	else if ( Key == EKeys::MotionController_Left_Thumbstick_Down || Key == EKeys::MotionController_Right_Thumbstick_Down )
-	{
-		static FName DownSocketName( "Down" );
-		SocketName = DownSocketName;
-	}
-	else if ( Key == EKeys::MotionController_Left_Thumbstick_Up || Key == EKeys::MotionController_Right_Thumbstick_Up )
-	{
-		static FName UpSocketName( "Up" );
-		SocketName = UpSocketName;
-	}
-	else if ( Key == EKeys::MotionController_Left_Thumbstick_Left || Key == EKeys::MotionController_Right_Thumbstick_Left )
-	{
-		static FName LeftSocketName( "Left" );
-		SocketName = LeftSocketName;
-	}
-	else if ( Key == EKeys::MotionController_Left_Thumbstick_Right || Key == EKeys::MotionController_Right_Thumbstick_Right )
-	{
-		static FName RightSocketName( "Right" );
-		SocketName = RightSocketName;
-	}
-	else if ( Key == EKeys::MotionController_Left_FaceButton1 || Key == EKeys::MotionController_Right_FaceButton1 )
-	{
-		static FName FaceButton1SocketName( "FaceButton1" );
-		SocketName = FaceButton1SocketName;
-	}
-	else if ( Key == EKeys::MotionController_Left_FaceButton2 || Key == EKeys::MotionController_Right_FaceButton2 )
-	{
-		static FName FaceButton2SocketName( "FaceButton2" );
-		SocketName = FaceButton2SocketName;
-	}
-	else if ( Key == EKeys::MotionController_Left_FaceButton3 || Key == EKeys::MotionController_Right_FaceButton3 )
-	{
-		static FName FaceButton3SocketName( "FaceButton3" );
-		SocketName = FaceButton3SocketName;
-	}
-	else if ( Key == EKeys::MotionController_Left_FaceButton4 || Key == EKeys::MotionController_Right_FaceButton4 )
-	{
-		static FName FaceButton4SocketName( "FaceButton4" );
-		SocketName = FaceButton4SocketName;
-	}
-	else
-	{
-		// Not a key that we care about
-	}
-
-	if ( SocketName != NAME_None )
-	{
-		UStaticMeshSocket* Socket = StaticMesh->FindSocket( SocketName );
-		if ( Socket != nullptr )
-		{
-			return Socket;
-		}
-	}
-
-	return nullptr;
-};
-
-void UVREditorInteractor::SetLaserVisuals( const FLinearColor& NewColor, const float CrawlFade, const float CrawlSpeed )
-{
-	static FName StaticLaserColorParameterName( "LaserColor" );
-	LaserPointerMID->SetVectorParameterValue( StaticLaserColorParameterName, NewColor );
-	TranslucentLaserPointerMID->SetVectorParameterValue( StaticLaserColorParameterName, NewColor );
-
-	static FName StaticCrawlParameterName( "Crawl" );
-	LaserPointerMID->SetScalarParameterValue( StaticCrawlParameterName, CrawlFade );
-	TranslucentLaserPointerMID->SetScalarParameterValue( StaticCrawlParameterName, CrawlFade );
-
-	static FName StaticCrawlSpeedParameterName( "CrawlSpeed" );
-	LaserPointerMID->SetScalarParameterValue( StaticCrawlSpeedParameterName, CrawlSpeed );
-	TranslucentLaserPointerMID->SetScalarParameterValue( StaticCrawlSpeedParameterName, CrawlSpeed );
-
-	static FName StaticHandTrimColorParameter( "TrimGlowColor" );
-	HandMeshMID->SetVectorParameterValue( StaticHandTrimColorParameter, NewColor );
-
-	HoverPointLightComponent->SetLightColor( NewColor );
-}
-
-void UVREditorInteractor::UpdateRadialMenuInput( const float DeltaTime )
-{
-	UVREditorUISystem& UISystem = GetVRMode().GetUISystem();
-
-	//Update the radial menu
-	EViewportInteractionDraggingMode DraggingMode = GetDraggingMode();
-	if ( !bHasUIInFront && 
-		( bIsTrackpadPositionValid[ 0 ] && bIsTrackpadPositionValid[ 1 ] ) &&
-		DraggingMode != EViewportInteractionDraggingMode::ActorsWithGizmo &&
-		DraggingMode != EViewportInteractionDraggingMode::ActorsFreely &&
-		DraggingMode != EViewportInteractionDraggingMode::ActorsAtLaserImpact &&
-		DraggingMode != EViewportInteractionDraggingMode::AssistingDrag && 
-		!UISystem.IsInteractorDraggingDockUI( this ) &&
-		!UISystem.IsShowingRadialMenu( Cast<UVREditorInteractor>( OtherInteractor ) ) )
-	{
-		const EHMDDeviceType::Type HMDDevice = GetHMDDeviceType();
-
-		// Spawn the radial menu if we are using the touchpad for steamvr or the analog stick for the oculus
-		if ( ( HMDDevice == EHMDDeviceType::DT_SteamVR && bIsTouchingTrackpad ) ||
-			 ( HMDDevice == EHMDDeviceType::DT_OculusRift && TrackpadPosition.Size() > VREd::MinJoystickOffsetBeforeRadialMenu->GetFloat() ) )
-		{
-			UISystem.TryToSpawnRadialMenu( this );
-		}
-		else
-		{
-			// Hide it if we are not using the touchpad or analog stick
-			UISystem.HideRadialMenu( this );
-		}
-
-		// Update the radial menu if we are already showing the radial menu
-		if ( UISystem.IsShowingRadialMenu( this ) )
-		{
-			UISystem.UpdateRadialMenu( this );
-		}
-	}
-}
-
 #undef LOCTEXT_NAMESPACE
\ No newline at end of file
diff --git a/Engine/Source/Editor/VREditor/VREditorMode.cpp b/Engine/Source/Editor/VREditor/VREditorMode.cpp
--- a/Engine/Source/Editor/VREditor/VREditorMode.cpp
+++ b/Engine/Source/Editor/VREditor/VREditorMode.cpp
@@ -26,6 +26,7 @@
 #include "Editor/ViewportInteraction/Public/ViewportInteraction.h"
 #include "VREditorInteractor.h"
 #include "MouseCursorInteractor.h"
+#include "VREditorMotionControllerInteractor.h"
 
 #include "Interactables/VREditorButton.h"
 
@@ -333,12 +334,14 @@ void FVREditorMode::Enter()
 
 		// Motion controllers
 		{
-			LeftHandInteractor = NewObject<UVREditorInteractor>( WorldInteraction );
-			LeftHandInteractor->Init( EControllerHand::Left, this );
+			LeftHandInteractor = NewObject<UVREditorMotionControllerInteractor>( WorldInteraction );
+			LeftHandInteractor->SetControllerHandSide( EControllerHand::Left );
+			LeftHandInteractor->Init( this );
 			WorldInteraction->AddInteractor( LeftHandInteractor );
 
-			RightHandInteractor = NewObject<UVREditorInteractor>( WorldInteraction );
-			RightHandInteractor->Init( EControllerHand::Right, this );
+			RightHandInteractor = NewObject<UVREditorMotionControllerInteractor>( WorldInteraction );
+			RightHandInteractor->SetControllerHandSide( EControllerHand::Right );
+			RightHandInteractor->Init( this );
 			WorldInteraction->AddInteractor( RightHandInteractor );
 
 			WorldInteraction->PairInteractors( LeftHandInteractor, RightHandInteractor );
@@ -383,6 +386,8 @@ void FVREditorMode::Exit()
 		WorldMovementPostProcessMaterial = nullptr;
 	}
 
+	check( WorldInteraction->GetViewportClient() != nullptr );
+
 	{
 		DestroyTransientActor( AvatarMeshActor );
 		AvatarMeshActor = nullptr;
@@ -474,13 +479,6 @@ void FVREditorMode::Exit()
 		UISystem = nullptr;
 	}
 
-	if ( WorldInteraction != nullptr )
-	{
-		WorldInteraction->Shutdown();
-		WorldInteraction->MarkPendingKill();
-		WorldInteraction = nullptr;
-	}
-
 	if ( TeleporterSystem != nullptr )
 	{
 		TeleporterSystem->Shutdown();
@@ -488,6 +486,13 @@ void FVREditorMode::Exit()
 		TeleporterSystem = nullptr;
 	}
 
+	if ( WorldInteraction != nullptr )
+	{
+		WorldInteraction->Shutdown();
+		WorldInteraction->MarkPendingKill();
+		WorldInteraction = nullptr;
+	}
+
 	// @todo vreditor urgent: Disable global editor hacks for VR Editor mode
 	GEnableVREditorHacks = false;
 }
@@ -1136,7 +1141,6 @@ const SLevelViewport& FVREditorMode::GetLevelViewportPossessedForVR() const
 	return *VREditorLevelViewportWeakPtr.Pin();
 }
 
-
 SLevelViewport& FVREditorMode::GetLevelViewportPossessedForVR()
 {
 	return *VREditorLevelViewportWeakPtr.Pin();
@@ -1213,26 +1217,30 @@ void FVREditorMode::CleanUpActorsBeforeMapChangeOrSimulate()
 
 void FVREditorMode::ToggleFlashlight( UVREditorInteractor* Interactor )
 {
-	if ( FlashlightComponent == nullptr )
+	UVREditorMotionControllerInteractor* MotionControllerInteractor = Cast<UVREditorMotionControllerInteractor>( Interactor );
+	if ( MotionControllerInteractor )
 	{
-		FlashlightComponent = NewObject<USpotLightComponent>(AvatarMeshActor);
-		AvatarMeshActor->AddOwnedComponent(FlashlightComponent);
-		FlashlightComponent->RegisterComponent();
-		FlashlightComponent->SetMobility(EComponentMobility::Movable);
-		FlashlightComponent->SetCastShadows(false);
-		FlashlightComponent->bUseInverseSquaredFalloff = false;
-		//@todo vreditor tweak
-		FlashlightComponent->SetLightFalloffExponent(8.0f);
-		FlashlightComponent->SetIntensity(20.0f);
-		FlashlightComponent->SetOuterConeAngle(25.0f);
-		FlashlightComponent->SetInnerConeAngle(25.0f);
-		
-	}
+		if ( FlashlightComponent == nullptr )
+		{
+			FlashlightComponent = NewObject<USpotLightComponent>( AvatarMeshActor );
+			AvatarMeshActor->AddOwnedComponent( FlashlightComponent );
+			FlashlightComponent->RegisterComponent();
+			FlashlightComponent->SetMobility( EComponentMobility::Movable );
+			FlashlightComponent->SetCastShadows( false );
+			FlashlightComponent->bUseInverseSquaredFalloff = false;
+			//@todo vreditor tweak
+			FlashlightComponent->SetLightFalloffExponent( 8.0f );
+			FlashlightComponent->SetIntensity( 20.0f );
+			FlashlightComponent->SetOuterConeAngle( 25.0f );
+			FlashlightComponent->SetInnerConeAngle( 25.0f );
 
-	const FAttachmentTransformRules AttachmentTransformRules = FAttachmentTransformRules( EAttachmentRule::KeepRelative, true );
-	FlashlightComponent->AttachToComponent( Interactor->GetMotionControllerComponent(), AttachmentTransformRules );
-	bIsFlashlightOn = !bIsFlashlightOn;
-	FlashlightComponent->SetVisibility(bIsFlashlightOn);
+		}
+
+		const FAttachmentTransformRules AttachmentTransformRules = FAttachmentTransformRules( EAttachmentRule::KeepRelative, true );
+		FlashlightComponent->AttachToComponent( MotionControllerInteractor->GetMotionControllerComponent(), AttachmentTransformRules );
+		bIsFlashlightOn = !bIsFlashlightOn;
+		FlashlightComponent->SetVisibility( bIsFlashlightOn );
+	}
 }
 
 void FVREditorMode::CycleTransformGizmoHandleType()
diff --git a/Engine/Source/Editor/VREditor/VREditorMode.h b/Engine/Source/Editor/VREditor/VREditorMode.h
--- a/Engine/Source/Editor/VREditor/VREditorMode.h
+++ b/Engine/Source/Editor/VREditor/VREditorMode.h
@@ -353,10 +353,10 @@ protected:
 	class UMouseCursorInteractor* MouseCursorInteractor;
 
 	/** The right motion controller */
-	class UVREditorInteractor* LeftHandInteractor; //@todo vreditor: Hardcoded interactors
+	class UVREditorMotionControllerInteractor* LeftHandInteractor; //@todo vreditor: Hardcoded interactors
 	
 	/** The right motion controller */
-	class UVREditorInteractor* RightHandInteractor; 
+	class UVREditorMotionControllerInteractor* RightHandInteractor; 
 
 	//
 	// Colors
diff --git a/Engine/Source/Editor/VREditor/VREditorMotionControllerInteractor.cpp b/Engine/Source/Editor/VREditor/VREditorMotionControllerInteractor.cpp
new file mode 100644
--- /dev/null
+++ b/Engine/Source/Editor/VREditor/VREditorMotionControllerInteractor.cpp
@@ -0,0 +1,1293 @@
+// Copyright 1998-2015 Epic Games, Inc. All Rights Reserved.
+
+#include "VREditorModule.h"
+#include "VREditorMotionControllerInteractor.h"
+#include "ViewportWorldInteraction.h"
+#include "VREditorMode.h"
+#include "VREditorUISystem.h"
+#include "VREditorFloatingText.h"
+
+#include "IMotionController.h"
+#include "IHeadMountedDisplay.h"
+#include "MotionControllerComponent.h"
+#include "Features/IModularFeatures.h"
+
+namespace VREd
+{
+	//Laser
+	static FAutoConsoleVariable OculusLaserPointerRotationOffset( TEXT( "VI.OculusLaserPointerRotationOffset" ), 0.0f, TEXT( "How much to rotate the laser pointer (pitch) relative to the forward vector of the controller (Oculus)" ) );
+	static FAutoConsoleVariable ViveLaserPointerRotationOffset( TEXT( "VI.ViveLaserPointerRotationOffset" ), /* -57.8f */ 0.0f, TEXT( "How much to rotate the laser pointer (pitch) relative to the forward vector of the controller (Vive)" ) );
+	static FAutoConsoleVariable OculusLaserPointerStartOffset( TEXT( "VI.OculusLaserPointerStartOffset" ), 2.8f, TEXT( "How far to offset the start of the laser pointer to avoid overlapping the hand mesh geometry (Oculus)" ) );
+	static FAutoConsoleVariable ViveLaserPointerStartOffset( TEXT( "VI.ViveLaserPointerStartOffset" ), 1.25f /* 8.5f */, TEXT( "How far to offset the start of the laser pointer to avoid overlapping the hand mesh geometry (Vive)" ) );
+
+	// Laser visuals
+	static FAutoConsoleVariable LaserPointerRadius( TEXT( "VREd.LaserPointerRadius" ), 0.5f, TEXT( "Radius of the laser pointer line" ) );
+	static FAutoConsoleVariable LaserPointerHoverBallRadius( TEXT( "VREd.LaserPointerHoverBallRadius" ), 1.5f, TEXT( "Radius of the visual cue for a hovered object along the laser pointer ray" ) );
+	static FAutoConsoleVariable LaserPointerLightPullBackDistance( TEXT( "VREd.LaserPointerLightPullBackDistance" ), 2.5f, TEXT( "How far to pull back our little hover light from the impact surface" ) );
+	static FAutoConsoleVariable LaserPointerLightRadius( TEXT( "VREd.LaserPointLightRadius" ), 20.0f, TEXT( "How big our hover light is" ) );
+
+	//Trigger
+	static FAutoConsoleVariable TriggerLightlyPressedThreshold( TEXT( "VI.TriggerLightlyPressedThreshold" ), 0.03f, TEXT( "Minimum trigger threshold before we consider the trigger at least 'lightly pressed'" ) );
+	static FAutoConsoleVariable TriggerDeadZone( TEXT( "VI.TriggerDeadZone" ), 0.01f, TEXT( "Trigger dead zone.  The trigger must be fully released before we'll trigger a new 'light press'" ) );
+	static FAutoConsoleVariable TriggerFullyPressedThreshold( TEXT( "VI.TriggerFullyPressedThreshold" ), 0.95f, TEXT( "Minimum trigger threshold before we consider the trigger 'fully pressed'" ) );
+	static FAutoConsoleVariable TriggerFullyPressedReleaseThreshold( TEXT( "VI.TriggerFullyPressedReleaseThreshold" ), 0.8f, TEXT( "After fully pressing the trigger, if the axis falls below this threshold we no longer consider it fully pressed" ) );
+
+	// Haptic feedback
+	static FAutoConsoleVariable SleepForRiftHaptics( TEXT( "VI.SleepForRiftHaptics" ), 1, TEXT( "When enabled, we'll sleep the game thread mid-frame to wait for haptic effects to finish.  This can be devasting to performance!" ) );
+	static FAutoConsoleVariable MinHapticTimeForRift( TEXT( "VREd.MinHapticTimeForRift" ), 0.005f, TEXT( "How long to play haptic effects on the Rift" ) );
+
+	static FAutoConsoleVariable TrackpadAbsoluteDragSpeed( TEXT( "VREd.TrackpadAbsoluteDragSpeed" ), 40.0f, TEXT( "How fast objects move toward or away when you drag on the touchpad while carrying them" ) );
+	static FAutoConsoleVariable TrackpadRelativeDragSpeed( TEXT( "VREd.TrackpadRelativeDragSpeed" ), 8.0f, TEXT( "How fast objects move toward or away when you hold a direction on an analog stick while carrying them" ) );
+	static FAutoConsoleVariable InvertTrackpadVertical( TEXT( "VREd.InvertTrackpadVertical" ), 1, TEXT( "Toggles inverting the touch pad vertical axis" ) );
+	static FAutoConsoleVariable TriggerLightlyPressedLockTime( TEXT( "VREd.TriggerLightlyPressedLockTime" ), 0.15f, TEXT( "If the trigger remains lightly pressed for longer than this, we'll continue to treat it as a light press in some cases" ) );
+	static FAutoConsoleVariable MinVelocityForInertia( TEXT( "VREd.MinVelocityForInertia" ), 1.0f, TEXT( "Minimum velocity (in cm/frame in unscaled room space) before inertia will kick in when releasing objects (or the world)" ) );
+	static FAutoConsoleVariable MinJoystickOffsetBeforeRadialMenu( TEXT( "VREd.MinJoystickOffsetBeforeRadialMenu" ), 0.15f, TEXT( "Toggles inverting the touch pad vertical axis" ) );
+
+	static FAutoConsoleVariable HelpLabelFadeDuration( TEXT( "VREd.HelpLabelFadeDuration" ), 0.4f, TEXT( "Duration to fade controller help labels in and out" ) );
+	static FAutoConsoleVariable HelpLabelFadeDistance( TEXT( "VREd.HelpLabelFadeDistance" ), 30.0f, TEXT( "Distance at which controller help labels should appear (in cm)" ) );
+}
+
+const FName UVREditorMotionControllerInteractor::TrackpadPositionX = FName( "TrackpadPositionX" );
+const FName UVREditorMotionControllerInteractor::TrackpadPositionY = FName( "TrackpadPositionY" );
+const FName UVREditorMotionControllerInteractor::TriggerAxis = FName( "TriggerAxis" );
+const FName UVREditorMotionControllerInteractor::MotionController_Left_FullyPressedTriggerAxis = FName( "MotionController_Left_FullyPressedTriggerAxis" );
+const FName UVREditorMotionControllerInteractor::MotionController_Right_FullyPressedTriggerAxis = FName( "MotionController_Right_FullyPressedTriggerAxis" );
+const FName UVREditorMotionControllerInteractor::MotionController_Left_LightlyPressedTriggerAxis = FName( "MotionController_Left_LightlyPressedTriggerAxis" );
+const FName UVREditorMotionControllerInteractor::MotionController_Right_LightlyPressedTriggerAxis = FName( "MotionController_Right_LightlyPressedTriggerAxis" );
+
+namespace VREditorKeyNames
+{
+	// @todo vreditor input: Ideally these would not be needed, but SteamVR fires off it's "trigger pressed" event
+	// well before the trigger is fully down (*click*)
+	static const FName MotionController_Left_FullyPressedTriggerAxis( "MotionController_Left_FullyPressedTriggerAxis" );
+	static const FName MotionController_Right_FullyPressedTriggerAxis( "MotionController_Right_FullyPressedTriggerAxis" );
+
+	static const FName MotionController_Left_LightlyPressedTriggerAxis( "MotionController_Left_LightlyPressedTriggerAxis" );
+	static const FName MotionController_Right_LightlyPressedTriggerAxis( "MotionController_Right_LightlyPressedTriggerAxis" );
+}
+
+namespace SteamVRControllerKeyNames
+{
+	static const FGamepadKeyNames::Type Touch0( "Steam_Touch_0" );
+	static const FGamepadKeyNames::Type Touch1( "Steam_Touch_1" );
+}
+
+
+#define LOCTEXT_NAMESPACE "VREditor"
+
+UVREditorMotionControllerInteractor::UVREditorMotionControllerInteractor( const FObjectInitializer& Initializer ) :
+	Super( Initializer ),
+	MotionControllerComponent( nullptr ),
+	HandMeshComponent( nullptr ),
+	LaserPointerMeshComponent( nullptr ),
+	LaserPointerMID( nullptr ),
+	TranslucentLaserPointerMID( nullptr ),
+	HoverMeshComponent( nullptr ),
+	HoverPointLightComponent( nullptr ),
+	HandMeshMID( nullptr ),
+	ControllerHandSide( EControllerHand::Pad ),
+	bHaveMotionController( false ),
+	bIsTriggerFullyPressed( false ),
+	bIsTriggerAtLeastLightlyPressed( false ),
+	RealTimeTriggerWasLightlyPressed( 0.0f ),
+	bHasTriggerBeenReleasedSinceLastPress( true ),
+	LastHapticTime( 0.0f ),
+	bIsTouchingTrackpad( false ),
+	TrackpadPosition( FVector2D::ZeroVector ),
+	LastTrackpadPosition( FVector2D::ZeroVector ),
+	bIsTrackpadPositionValid { false, false },
+	LastTrackpadPositionUpdateTime( FTimespan::MinValue() )
+{
+}
+
+
+UVREditorMotionControllerInteractor::~UVREditorMotionControllerInteractor()
+{
+	Shutdown();
+}
+
+
+void UVREditorMotionControllerInteractor::Init( class FVREditorMode* InVRMode )
+{
+	Super::Init( InVRMode );
+	bHaveMotionController = true;
+
+	// Setup keys
+	if ( ControllerHandSide == EControllerHand::Left )
+	{
+		AddKeyAction( EKeys::MotionController_Left_Grip1, FViewportActionKeyInput( ViewportWorldActionTypes::WorldMovement ) );
+		AddKeyAction( UVREditorMotionControllerInteractor::MotionController_Left_FullyPressedTriggerAxis, FViewportActionKeyInput( ViewportWorldActionTypes::SelectAndMove ) );
+		AddKeyAction( UVREditorMotionControllerInteractor::MotionController_Left_LightlyPressedTriggerAxis, FViewportActionKeyInput( ViewportWorldActionTypes::SelectAndMove_LightlyPressed ) );
+		AddKeyAction( EKeys::MotionController_Left_Thumbstick, FViewportActionKeyInput( VRActionTypes::ConfirmRadialSelection ) );
+
+		AddKeyAction( SteamVRControllerKeyNames::Touch0, FViewportActionKeyInput( VRActionTypes::Touch ) );
+		AddKeyAction( EKeys::MotionController_Left_TriggerAxis, FViewportActionKeyInput( UVREditorMotionControllerInteractor::TriggerAxis ) );
+		AddKeyAction( EKeys::MotionController_Left_Thumbstick_X, FViewportActionKeyInput( UVREditorMotionControllerInteractor::TrackpadPositionX ) );
+		AddKeyAction( EKeys::MotionController_Left_Thumbstick_Y, FViewportActionKeyInput( UVREditorMotionControllerInteractor::TrackpadPositionY ) );
+		AddKeyAction( EKeys::MotionController_Left_Thumbstick, FViewportActionKeyInput( VRActionTypes::ConfirmRadialSelection ) );
+
+		if ( GetVRMode().GetHMDDeviceType() == EHMDDeviceType::DT_SteamVR )
+		{
+			AddKeyAction( EKeys::MotionController_Left_Shoulder, FViewportActionKeyInput( VRActionTypes::Modifier ) );
+		}
+		else
+		{
+			AddKeyAction( EKeys::MotionController_Left_FaceButton1, FViewportActionKeyInput( VRActionTypes::Modifier ) );
+		}
+	}
+	else if ( ControllerHandSide == EControllerHand::Right )
+	{
+		AddKeyAction( EKeys::MotionController_Right_Grip1, FViewportActionKeyInput( ViewportWorldActionTypes::WorldMovement ) );
+		AddKeyAction( UVREditorMotionControllerInteractor::MotionController_Right_FullyPressedTriggerAxis, FViewportActionKeyInput( ViewportWorldActionTypes::SelectAndMove ) );
+		AddKeyAction( UVREditorMotionControllerInteractor::MotionController_Right_LightlyPressedTriggerAxis, FViewportActionKeyInput( ViewportWorldActionTypes::SelectAndMove_LightlyPressed ) );
+		AddKeyAction( EKeys::MotionController_Right_Thumbstick, FViewportActionKeyInput( VRActionTypes::ConfirmRadialSelection ) );
+
+		AddKeyAction( SteamVRControllerKeyNames::Touch1, FViewportActionKeyInput( VRActionTypes::Touch ) );
+		AddKeyAction( EKeys::MotionController_Right_TriggerAxis, FViewportActionKeyInput( UVREditorMotionControllerInteractor::TriggerAxis ) );
+		AddKeyAction( EKeys::MotionController_Right_Thumbstick_X, FViewportActionKeyInput( UVREditorMotionControllerInteractor::TrackpadPositionX ) );
+		AddKeyAction( EKeys::MotionController_Right_Thumbstick_Y, FViewportActionKeyInput( UVREditorMotionControllerInteractor::TrackpadPositionY ) );
+		AddKeyAction( EKeys::MotionController_Right_Thumbstick, FViewportActionKeyInput( VRActionTypes::ConfirmRadialSelection ) );
+
+		if ( GetVRMode().GetHMDDeviceType() == EHMDDeviceType::DT_SteamVR )
+		{
+			AddKeyAction( EKeys::MotionController_Right_Shoulder, FViewportActionKeyInput( VRActionTypes::Modifier ) );
+		}
+		else
+		{
+			AddKeyAction( EKeys::MotionController_Right_FaceButton1, FViewportActionKeyInput( VRActionTypes::Modifier ) );
+		}
+	}
+}
+
+
+float UVREditorMotionControllerInteractor::GetSlideDelta()
+{
+	return GetTrackpadSlideDelta();
+}
+
+
+void UVREditorMotionControllerInteractor::SetupComponent( AActor* OwningActor )
+{
+	// Setup a motion controller component.  This allows us to take advantage of late frame updates, so
+	// our motion controllers won't lag behind the HMD
+	{
+		MotionControllerComponent = NewObject<UMotionControllerComponent>( OwningActor );
+		OwningActor->AddOwnedComponent( MotionControllerComponent );
+		MotionControllerComponent->SetupAttachment( OwningActor->GetRootComponent() );
+		MotionControllerComponent->RegisterComponent();
+
+		MotionControllerComponent->SetMobility( EComponentMobility::Movable );
+		MotionControllerComponent->SetCollisionEnabled( ECollisionEnabled::NoCollision );
+
+		MotionControllerComponent->Hand = ControllerHandSide;
+
+		MotionControllerComponent->bDisableLowLatencyUpdate = false;
+	}
+
+	// Hand mesh
+	{
+		HandMeshComponent = NewObject<UStaticMeshComponent>( OwningActor );
+		OwningActor->AddOwnedComponent( HandMeshComponent );
+		HandMeshComponent->SetupAttachment( MotionControllerComponent );
+		HandMeshComponent->RegisterComponent();
+
+		// @todo vreditor extensibility: We need this to be able to be overridden externally, or simply based on the HMD name (but allowing external folders)
+		FString MeshName;
+		FString MaterialName;
+		if ( GetVRMode().GetHMDDeviceType() == EHMDDeviceType::DT_SteamVR )
+		{
+			MeshName = TEXT( "/Engine/VREditor/Devices/Vive/VivePreControllerMesh" );
+			MaterialName = TEXT( "/Engine/VREditor/Devices/Vive/VivePreControllerMaterial_Inst" );
+		}
+		else
+		{
+			MeshName = TEXT( "/Engine/VREditor/Devices/Oculus/OculusControllerMesh" );
+			MaterialName = TEXT( "/Engine/VREditor/Devices/Oculus/OculusControllerMaterial_Inst" );
+		}
+
+		UStaticMesh* HandMesh = LoadObject<UStaticMesh>( nullptr, *MeshName );
+		check( HandMesh != nullptr );
+
+		HandMeshComponent->SetStaticMesh( HandMesh );
+		HandMeshComponent->SetMobility( EComponentMobility::Movable );
+		HandMeshComponent->SetCollisionEnabled( ECollisionEnabled::NoCollision );
+
+		UMaterialInstance* HandMeshMaterialInst = LoadObject<UMaterialInstance>( nullptr, *MaterialName );
+		check( HandMeshMaterialInst != nullptr );
+		HandMeshMID = UMaterialInstanceDynamic::Create( HandMeshMaterialInst, GetTransientPackage() );
+		check( HandMeshMID != nullptr );
+		HandMeshComponent->SetMaterial( 0, HandMeshMID );
+	}
+
+	// Laser pointer
+	{
+		LaserPointerMeshComponent = NewObject<UStaticMeshComponent>( OwningActor );
+		OwningActor->AddOwnedComponent( LaserPointerMeshComponent );
+		LaserPointerMeshComponent->SetupAttachment( MotionControllerComponent );
+		LaserPointerMeshComponent->RegisterComponent();
+
+		UStaticMesh* LaserPointerMesh = LoadObject<UStaticMesh>( nullptr, TEXT( "/Engine/VREditor/LaserPointer/LaserPointerMesh" ) );
+		check( LaserPointerMesh != nullptr );
+
+		LaserPointerMeshComponent->SetStaticMesh( LaserPointerMesh );
+		LaserPointerMeshComponent->SetMobility( EComponentMobility::Movable );
+		LaserPointerMeshComponent->SetCollisionEnabled( ECollisionEnabled::NoCollision );
+
+		UMaterialInstance* LaserPointerMaterialInst = LoadObject<UMaterialInstance>( nullptr, TEXT( "/Engine/VREditor/LaserPointer/LaserPointerMaterialInst" ) );
+		check( LaserPointerMaterialInst != nullptr );
+		LaserPointerMID = UMaterialInstanceDynamic::Create( LaserPointerMaterialInst, GetTransientPackage() );
+		check( LaserPointerMID != nullptr );
+		LaserPointerMeshComponent->SetMaterial( 0, LaserPointerMID );
+
+		UMaterialInstance* TranslucentLaserPointerMaterialInst = LoadObject<UMaterialInstance>( nullptr, TEXT( "/Engine/VREditor/LaserPointer/TranslucentLaserPointerMaterialInst" ) );
+		check( TranslucentLaserPointerMaterialInst != nullptr );
+		TranslucentLaserPointerMID = UMaterialInstanceDynamic::Create( TranslucentLaserPointerMaterialInst, GetTransientPackage() );
+		check( TranslucentLaserPointerMID != nullptr );
+		LaserPointerMeshComponent->SetMaterial( 1, TranslucentLaserPointerMID );
+	}
+
+	// Hover cue for laser pointer
+	{
+		HoverMeshComponent = NewObject<UStaticMeshComponent>( OwningActor );
+		OwningActor->AddOwnedComponent( HoverMeshComponent );
+		HoverMeshComponent->SetupAttachment( OwningActor->GetRootComponent() );
+		HoverMeshComponent->RegisterComponent();
+
+		UStaticMesh* HoverMesh = LoadObject<UStaticMesh>( nullptr, TEXT( "/Engine/VREditor/LaserPointer/HoverMesh" ) );
+		check( HoverMesh != nullptr );
+		HoverMeshComponent->SetStaticMesh( HoverMesh );
+		HoverMeshComponent->SetMobility( EComponentMobility::Movable );
+		HoverMeshComponent->SetCollisionEnabled( ECollisionEnabled::NoCollision );
+
+		// Add a light!
+		{
+			HoverPointLightComponent = NewObject<UPointLightComponent>( OwningActor );
+			OwningActor->AddOwnedComponent( HoverPointLightComponent );
+			HoverPointLightComponent->SetupAttachment( HoverMeshComponent );
+			HoverPointLightComponent->RegisterComponent();
+
+			HoverPointLightComponent->SetLightColor( FLinearColor::Red );
+			//Hand.HoverPointLightComponent->SetLightColor( FLinearColor( 0.0f, 1.0f, 0.2f, 1.0f ) );
+			HoverPointLightComponent->SetIntensity( 30.0f );	// @todo: VREditor tweak
+			HoverPointLightComponent->SetMobility( EComponentMobility::Movable );
+			HoverPointLightComponent->SetAttenuationRadius( VREd::LaserPointerLightRadius->GetFloat() );
+			HoverPointLightComponent->bUseInverseSquaredFalloff = false;
+			HoverPointLightComponent->SetCastShadows( false );
+		}
+	}
+}
+
+
+void UVREditorMotionControllerInteractor::SetControllerHandSide( const EControllerHand InControllerHandSide )
+{
+	ControllerHandSide = InControllerHandSide;
+}
+
+
+void UVREditorMotionControllerInteractor::Shutdown()
+{
+	Super::Shutdown();
+
+	MotionControllerComponent = nullptr;
+	HandMeshComponent = nullptr;
+	LaserPointerMeshComponent = nullptr;
+	LaserPointerMID = nullptr;
+	TranslucentLaserPointerMID = nullptr;
+	HoverMeshComponent = nullptr;
+	HoverPointLightComponent = nullptr;
+	HandMeshComponent = nullptr;
+}
+
+void UVREditorMotionControllerInteractor::Tick( const float DeltaTime )
+{
+	Super::Tick( DeltaTime );
+
+	{
+		const float WorldScaleFactor = WorldInteraction->GetWorldScaleFactor();
+
+		// @todo vreditor: Manually ticking motion controller components
+		MotionControllerComponent->TickComponent( DeltaTime, ELevelTick::LEVELTICK_PauseTick, nullptr );
+
+		// The hands need to stay the same size relative to our tracking space, so we inverse compensate for world to meters scale here
+		// NOTE: We don't need to set the hand mesh location and rotation, as the MotionControllerComponent does that itself
+		if ( ControllerHandSide == EControllerHand::Right &&
+			GetHMDDeviceType() == EHMDDeviceType::DT_OculusRift )	// Oculus has asymmetrical controllers, so we mirror the mesh horizontally
+		{
+			HandMeshComponent->SetRelativeScale3D( FVector( WorldScaleFactor, -WorldScaleFactor, WorldScaleFactor ) );
+		}
+		else
+		{
+			HandMeshComponent->SetRelativeScale3D( FVector( WorldScaleFactor ) );
+		}
+	}
+
+	UpdateRadialMenuInput( DeltaTime );
+
+	{
+		const float WorldScaleFactor = WorldInteraction->GetWorldScaleFactor();
+
+		// Don't bother drawing hands if we're not currently tracking them.
+		if ( bHaveMotionController )
+		{
+			HandMeshComponent->SetVisibility( true );
+		}
+		else
+		{
+			HandMeshComponent->SetVisibility( false );
+		}
+
+		FVector LaserPointerStart, LaserPointerEnd;
+		if ( GetLaserPointer( /* Out */ LaserPointerStart, /* Out */ LaserPointerEnd ) )
+		{
+			// Only show the laser if we're actually in VR
+			LaserPointerMeshComponent->SetVisibility( GetVRMode().IsActuallyUsingVR() );
+
+			// NOTE: We don't need to set the laser pointer location and rotation, as the MotionControllerComponent will do
+			// that later in the frame.  
+
+			// If we're actively dragging something around, then we'll crop the laser length to the hover impact
+			// point.  Otherwise, we always want the laser to protrude through hovered objects, so that you can
+			// interact with translucent gizmo handles that are occluded by geometry
+			FVector LaserPointerImpactPoint = LaserPointerEnd;
+			if ( GetInteractorData().bIsHovering && ( GetInteractorData().DraggingMode != EViewportInteractionDraggingMode::Nothing || IsHoveringOverUI() ) )
+			{
+				LaserPointerImpactPoint = GetInteractorData().HoverLocation;
+			}
+
+			// Apply rotation offset to the laser direction
+			const float LaserPointerRotationOffset = 0.0f;
+			LaserPointerMeshComponent->SetRelativeRotation( FRotator( LaserPointerRotationOffset, 0.0f, 0.0f ) );
+
+			const FVector LaserPointerDirection = ( LaserPointerImpactPoint - LaserPointerStart ).GetSafeNormal();
+
+			// Offset the beginning of the laser pointer a bit, so that it doesn't overlap the hand mesh
+			const float LaserPointerStartOffset =
+				WorldScaleFactor *
+				( GetVRMode().GetHMDDeviceType() == EHMDDeviceType::DT_OculusRift ? VREd::OculusLaserPointerStartOffset->GetFloat() : VREd::ViveLaserPointerStartOffset->GetFloat() );
+
+			const float LaserPointerLength = FMath::Max( 0.000001f, ( LaserPointerImpactPoint - LaserPointerStart ).Size() - LaserPointerStartOffset );
+
+			LaserPointerMeshComponent->SetRelativeLocation( FRotator( LaserPointerRotationOffset, 0.0f, 0.0f ).RotateVector( FVector( LaserPointerStartOffset, 0.0f, 0.0f ) ) );
+
+			// The laser pointer needs to stay the same size relative to our tracking space, so we inverse compensate for world to meters scale here
+			float LaserPointerRadius = VREd::LaserPointerRadius->GetFloat() * WorldScaleFactor;
+			float HoverMeshRadius = VREd::LaserPointerHoverBallRadius->GetFloat() * WorldScaleFactor;
+
+			// If we're hovering over something really close to the camera, go ahead and shrink the effect
+			// @todo vreditor: Can we make this actually just sized based on distance automatically?  The beam and impact point are basically a cone.
+			if ( IsHoveringOverUI() )
+			{
+				LaserPointerRadius *= 0.35f;	// @todo vreditor tweak
+				HoverMeshRadius *= 0.35f;	// @todo vreditor tweak
+			}
+
+			LaserPointerMeshComponent->SetRelativeScale3D( FVector( LaserPointerLength, LaserPointerRadius * 2.0f, LaserPointerRadius * 2.0f ) );
+
+			if ( IsHovering() )
+			{
+				// The hover effect needs to stay the same size relative to our tracking space, so we inverse compensate for world to meters scale here
+				HoverMeshComponent->SetRelativeScale3D( FVector( HoverMeshRadius * 2.0f ) );
+				HoverMeshComponent->SetVisibility( true );
+				HoverMeshComponent->SetWorldLocation( GetHoverLocation() );
+
+				GLog->Logf( TEXT( " IsHoveringOverUI: %d " ), IsHoveringOverUI() );
+
+				// Show the light too, unless it's on top of UI.  It looks too distracting on top of UI.
+				HoverPointLightComponent->SetVisibility( !IsHoveringOverUI() );
+
+				// Update radius for world scaling
+				HoverPointLightComponent->SetAttenuationRadius( VREd::LaserPointerLightRadius->GetFloat() * WorldScaleFactor );
+
+				// Pull hover light back a bit from the end of the ray
+				const float PullBackAmount = VREd::LaserPointerLightPullBackDistance->GetFloat() * WorldInteraction->GetWorldScaleFactor();
+				HoverPointLightComponent->SetWorldLocation( GetHoverLocation() - PullBackAmount * LaserPointerDirection );
+			}
+			else
+			{
+				HoverMeshComponent->SetVisibility( false );
+				HoverPointLightComponent->SetVisibility( false );
+			}
+		}
+		else
+		{
+			LaserPointerMeshComponent->SetVisibility( false );
+			HoverMeshComponent->SetVisibility( false );
+			HoverPointLightComponent->SetVisibility( false );
+		}
+	}
+
+	// Updating laser colors for both hands
+	{
+		FVREditorMode::EColors ResultColorType = FVREditorMode::EColors::DefaultColor;
+		float CrawlSpeed = 0.0f;
+		float CrawlFade = 0.0f;
+
+		const EViewportInteractionDraggingMode DraggingMode = GetDraggingMode();
+
+		//@todo VREditor
+		const bool bIsDraggingWorldWithTwoHands = false;
+		( DraggingMode == EViewportInteractionDraggingMode::AssistingDrag && GetOtherInteractor()->GetDraggingMode() == EViewportInteractionDraggingMode::World ) ||
+			( DraggingMode == EViewportInteractionDraggingMode::World && GetOtherInteractor()->GetDraggingMode() == EViewportInteractionDraggingMode::AssistingDrag );
+
+		if ( bIsDraggingWorldWithTwoHands )
+		{
+			ResultColorType = FVREditorMode::EColors::WorldDraggingColor_TwoHands;
+		}
+		else if ( DraggingMode == EViewportInteractionDraggingMode::World )
+		{
+			//@todo VREditor find a way to make the teleporter do it
+			// 			if ( GetVRMode().GetWorldInteraction().IsTeleporting() ) //@todo VREditor: Make teleport an action
+			// 			{
+			// 				ResultColorType = FVREditorMode::EColors::TeleportColor;
+			// 			}
+			// 			else
+			// 			{
+			// We can teleport in this mode, so animate the laser a bit
+			CrawlFade = 1.0f;
+			CrawlSpeed = 5.0f;
+			ResultColorType = FVREditorMode::EColors::WorldDraggingColor_OneHand;
+			//			}
+		}
+		else if ( DraggingMode == EViewportInteractionDraggingMode::ActorsAtLaserImpact ||
+			DraggingMode == EViewportInteractionDraggingMode::Material ||
+			DraggingMode == EViewportInteractionDraggingMode::ActorsFreely ||
+			DraggingMode == EViewportInteractionDraggingMode::ActorsWithGizmo ||
+			DraggingMode == EViewportInteractionDraggingMode::AssistingDrag ||
+			DraggingMode == EViewportInteractionDraggingMode::Interactable ||
+			( GetVRMode().GetUISystem().IsInteractorDraggingDockUI( this ) && GetVRMode().GetUISystem().IsDraggingDockUI() ) )
+		{
+			ResultColorType = FVREditorMode::EColors::SelectionColor;
+		}
+
+		const FLinearColor ResultColor = GetVRMode().GetColor( ResultColorType );
+		SetLaserVisuals( ResultColor, CrawlFade, CrawlSpeed );
+	}
+
+	UpdateHelpLabels();
+}
+
+void UVREditorMotionControllerInteractor::CalculateDragRay()
+{
+	const FTimespan CurrentTime = FTimespan::FromSeconds( FPlatformTime::Seconds() );
+	const float WorldScaleFactor = WorldInteraction->GetWorldScaleFactor();
+
+	// If we're dragging an object, go ahead and slide the object along the ray based on how far they slide their touch
+	// Make sure they are touching the trackpad, otherwise we get bad data
+	if ( bIsTrackpadPositionValid[ 1 ] )
+	{
+		const bool bIsAbsolute = ( GetVRMode().GetHMDDeviceType() == EHMDDeviceType::DT_SteamVR );
+		float SlideDelta = GetTrackpadSlideDelta() * WorldScaleFactor;
+
+		if ( !FMath::IsNearlyZero( SlideDelta ) )
+		{
+			InteractorData.DragRayLength += SlideDelta;
+
+			InteractorData.DragRayLengthVelocity = 0.0f;
+
+			// Don't apply inertia unless the user dragged a decent amount this frame
+			if ( bIsAbsolute && FMath::Abs( SlideDelta ) >= VREd::MinVelocityForInertia->GetFloat() * WorldScaleFactor )
+			{
+				// Don't apply inertia if our data is sort of old
+				if ( CurrentTime - LastTrackpadPositionUpdateTime <= FTimespan::FromSeconds( 1.0f / 30.0f ) )
+				{
+					InteractorData.DragRayLengthVelocity = SlideDelta;
+				}
+			}
+
+			// Don't go too far
+			if ( InteractorData.DragRayLength < 0.0f )
+			{
+				InteractorData.DragRayLength = 0.0f;
+				InteractorData.DragRayLengthVelocity = 0.0f;
+			}
+		}
+	}
+	else
+	{
+		if ( !FMath::IsNearlyZero( InteractorData.DragRayLengthVelocity ) )
+		{
+			// Apply drag ray length inertia
+			InteractorData.DragRayLength += InteractorData.DragRayLengthVelocity;
+
+			// Don't go too far!
+			if ( InteractorData.DragRayLength < 0.0f )
+			{
+				InteractorData.DragRayLength = 0.0f;
+				InteractorData.DragRayLengthVelocity = 0.0f;
+			}
+
+			// Apply damping
+			FVector RayVelocityVector( InteractorData.DragRayLengthVelocity, 0.0f, 0.0f );
+			const bool bVelocitySensitive = true;
+			WorldInteraction->ApplyVelocityDamping( RayVelocityVector, bVelocitySensitive );
+			InteractorData.DragRayLengthVelocity = RayVelocityVector.X;
+		}
+		else
+		{
+			InteractorData.DragRayLengthVelocity = 0.0f;
+		}
+	}
+}
+
+
+EHMDDeviceType::Type UVREditorMotionControllerInteractor::GetHMDDeviceType() const
+{
+	return GEngine->HMDDevice.IsValid() ? GEngine->HMDDevice->GetHMDDeviceType() : EHMDDeviceType::DT_SteamVR; //@todo: ViewportInteraction, assumption that it's steamvr ??
+}
+
+
+void UVREditorMotionControllerInteractor::HandleInputKey( FViewportActionKeyInput& Action, const FKey Key, const EInputEvent Event, bool& bOutWasHandled )
+{
+	if ( !bOutWasHandled )
+	{
+		bool bShouldAbsorbEvent = false;
+
+		// If "SelectAndMove" was pressed, we need to make sure that "SelectAndMove_LightlyPressed" is no longer
+		// in effect before we start handling "SelectAndMove".
+		if ( Action.ActionType == ViewportWorldActionTypes::SelectAndMove && IsTriggerAtLeastLightlyPressed() )
+		{
+			if ( Event == IE_Pressed )
+			{
+				// How long since the trigger was lightly pressed?
+				const float TimeSinceLightlyPressed = ( float ) ( FPlatformTime::Seconds() - GetRealTimeTriggerWasLightlyPressed() );
+				if ( !IsClickingOnUI() &&	// @todo vreditor: UI clicks happen with light presses; we never want to convert to a hard press!
+					GetDraggingMode() != EViewportInteractionDraggingMode::Material &&	// @todo vreditor: Material dragging happens with light presses, don't convert to a hard press!
+					GetDraggingMode() != EViewportInteractionDraggingMode::ActorsAtLaserImpact &&	// @todo vreditor: Actor dragging happens with light presses, don't convert to a hard press!
+					( !AllowTriggerLightPressLocking() || TimeSinceLightlyPressed < VREd::TriggerLightlyPressedLockTime->GetFloat() ) )
+				{
+					SetTriggerAtLeastLightlyPressed( false );
+					SetTriggerBeenReleasedSinceLastPress( false );
+
+					// Synthesize an input key for releasing the light press
+					// NOTE: Intentionally re-entrant call here.
+					const EInputEvent InputEvent = IE_Released;
+					const bool bWasLightReleaseHandled = UViewportInteractor::HandleInputKey( GetControllerSide() == EControllerHand::Left ? VREditorKeyNames::MotionController_Left_LightlyPressedTriggerAxis : VREditorKeyNames::MotionController_Right_LightlyPressedTriggerAxis, InputEvent );
+				}
+				else
+				{
+					// The button was held in a "lightly pressed" state for a long enough time that we should just continue to
+					// treat it like a light press.  This makes it much easier to hold it in this state when you need to!
+					bShouldAbsorbEvent = true;
+				}
+			}
+			else
+			{
+				// Absorb the release of the SelectAndMove event if we are still in a lightly pressed state
+				bShouldAbsorbEvent = true;
+			}
+		}
+
+		if ( bShouldAbsorbEvent )
+		{
+			bOutWasHandled = true;
+		}
+	}
+
+	if ( !bOutWasHandled )
+	{
+		// Update touch state
+		if ( Action.ActionType == VRActionTypes::Touch )
+		{
+			if ( Event == IE_Pressed )
+			{
+				bIsTouchingTrackpad = true;
+			}
+			else if ( Event == IE_Released )
+			{
+				bIsTouchingTrackpad = false;
+				bIsTrackpadPositionValid[ 0 ] = false;
+				bIsTrackpadPositionValid[ 1 ] = false;
+			}
+		}
+
+		// Update modifier state
+		if ( Action.ActionType == VRActionTypes::Modifier )
+		{
+			bOutWasHandled = true;
+			if ( Event == IE_Pressed )
+			{
+				bIsModifierPressed = true;
+			}
+			else if ( Event == IE_Released )
+			{
+				bIsModifierPressed = false;
+			}
+		}
+	}
+
+	ApplyButtonPressColors( Action );
+}
+
+void UVREditorMotionControllerInteractor::HandleInputAxis( FViewportActionKeyInput& Action, const FKey Key, const float Delta, const float DeltaTime, bool& bOutWasHandled )
+{
+	if ( Action.ActionType == TriggerAxis )
+	{
+		// Synthesize "lightly pressed" events for the trigger
+		{
+			// Store latest trigger value amount
+			SelectAndMoveTriggerValue = Delta;
+
+			FViewportActionKeyInput* OptionalAction = GetActionWithName( ViewportWorldActionTypes::SelectAndMove );
+			const bool bIsFullPressAlreadyCapturing = OptionalAction && OptionalAction->bIsInputCaptured;
+				
+			if ( !bIsFullPressAlreadyCapturing &&		// Don't fire if we're already capturing input for a full press
+				!bIsTriggerAtLeastLightlyPressed &&	// Don't fire unless we are already pressed
+				bHasTriggerBeenReleasedSinceLastPress &&	// Only if we've been fully released since the last time we fired
+				Delta >= VREd::TriggerLightlyPressedThreshold->GetFloat() )
+			{
+				bIsTriggerAtLeastLightlyPressed = true;
+				SetAllowTriggerLightPressLocking( true );
+				RealTimeTriggerWasLightlyPressed = FPlatformTime::Seconds();
+				bHasTriggerBeenReleasedSinceLastPress = false;
+
+				// Synthesize an input key for this light press
+				const EInputEvent InputEvent = IE_Pressed;
+				const bool bWasLightPressHandled = UViewportInteractor::HandleInputKey( ControllerHandSide == EControllerHand::Left ? MotionController_Left_LightlyPressedTriggerAxis : MotionController_Right_LightlyPressedTriggerAxis, InputEvent );
+			}
+			else if ( bIsTriggerAtLeastLightlyPressed && Delta < VREd::TriggerLightlyPressedThreshold->GetFloat() )
+			{
+				bIsTriggerAtLeastLightlyPressed = false;
+
+				// Synthesize an input key for this light press
+				const EInputEvent InputEvent = IE_Released;
+				const bool bWasLightReleaseHandled = UViewportInteractor::HandleInputKey( ControllerHandSide == EControllerHand::Left ? MotionController_Left_LightlyPressedTriggerAxis : MotionController_Right_LightlyPressedTriggerAxis, InputEvent );
+			}
+		}
+
+		if ( Delta < VREd::TriggerDeadZone->GetFloat() )
+		{
+			bHasTriggerBeenReleasedSinceLastPress = true;
+		}
+
+		// Synthesize "fully pressed" events for the trigger
+		{
+			if ( !bIsTriggerFullyPressed &&	// Don't fire unless we are already pressed
+				Delta >= VREd::TriggerFullyPressedThreshold->GetFloat() )
+			{
+				bIsTriggerFullyPressed = true;
+
+				// Synthesize an input key for this full press
+				const EInputEvent InputEvent = IE_Pressed;
+				UViewportInteractor::HandleInputKey( ControllerHandSide == EControllerHand::Left ? MotionController_Left_FullyPressedTriggerAxis : MotionController_Right_FullyPressedTriggerAxis, InputEvent );
+			}
+			else if ( bIsTriggerFullyPressed && Delta < VREd::TriggerFullyPressedReleaseThreshold->GetFloat() )
+			{
+				bIsTriggerFullyPressed = false;
+
+				// Synthesize an input key for this full press
+				const EInputEvent InputEvent = IE_Released;
+				UViewportInteractor::HandleInputKey( ControllerHandSide == EControllerHand::Left ? MotionController_Left_FullyPressedTriggerAxis : MotionController_Right_FullyPressedTriggerAxis, InputEvent );
+			}
+		}
+	}
+
+	if ( !bOutWasHandled )
+	{
+		if ( Action.ActionType == VRActionTypes::TrackpadPositionX )
+		{
+			LastTrackpadPosition.X = bIsTrackpadPositionValid[ 0 ] ? TrackpadPosition.X : Delta;
+			LastTrackpadPositionUpdateTime = FTimespan::FromSeconds( FPlatformTime::Seconds() );
+			TrackpadPosition.X = Delta;
+			bIsTrackpadPositionValid[ 0 ] = true;
+		}
+
+		if ( Action.ActionType == VRActionTypes::TrackpadPositionY )
+		{
+			float DeltaAxis = Delta;
+			if ( VREd::InvertTrackpadVertical->GetInt() != 0 )
+			{
+				DeltaAxis = -DeltaAxis;	// Y axis is inverted from HMD
+			}
+
+			LastTrackpadPosition.Y = bIsTrackpadPositionValid[ 1 ] ? TrackpadPosition.Y : DeltaAxis;
+			LastTrackpadPositionUpdateTime = FTimespan::FromSeconds( FPlatformTime::Seconds() );
+			TrackpadPosition.Y = DeltaAxis;
+			bIsTrackpadPositionValid[ 1 ] = true;
+		}
+	}
+
+	Super::HandleInputAxis( Action, Key, Delta, DeltaTime, bOutWasHandled );
+}
+
+void UVREditorMotionControllerInteractor::PollInput()
+{
+	bHaveMotionController = false;
+	InteractorData.LastTransform = InteractorData.Transform;
+	InteractorData.LastRoomSpaceTransform = InteractorData.RoomSpaceTransform;
+
+	// Generic motion controllers
+	TArray<IMotionController*> MotionControllers = IModularFeatures::Get().GetModularFeatureImplementations<IMotionController>( IMotionController::GetModularFeatureName() );
+	for ( auto MotionController : MotionControllers )	// @todo viewportinteraction: Needs support for multiple pairs of motion controllers
+	{
+		if ( MotionController != nullptr && !bHaveMotionController )
+		{
+			FVector Location = FVector::ZeroVector;
+			FRotator Rotation = FRotator::ZeroRotator;
+
+			if ( MotionController->GetControllerOrientationAndPosition( WorldInteraction->GetMotionControllerID(), ControllerHandSide, /* Out */ Rotation, /* Out */ Location ) )
+			{
+				bHaveMotionController = true;
+				InteractorData.RoomSpaceTransform = FTransform( Rotation.Quaternion(), Location, FVector( 1.0f ) );
+				InteractorData.Transform = InteractorData.RoomSpaceTransform * WorldInteraction->GetRoomTransform();
+			}
+		}
+	}
+}
+
+bool UVREditorMotionControllerInteractor::IsTriggerAtLeastLightlyPressed() const
+{
+	return bIsTriggerAtLeastLightlyPressed;
+}
+
+double UVREditorMotionControllerInteractor::GetRealTimeTriggerWasLightlyPressed() const
+{
+	return RealTimeTriggerWasLightlyPressed;
+}
+
+void UVREditorMotionControllerInteractor::SetTriggerAtLeastLightlyPressed( const bool bInTriggerAtLeastLightlyPressed )
+{
+	bIsTriggerAtLeastLightlyPressed = bInTriggerAtLeastLightlyPressed;
+}
+
+void UVREditorMotionControllerInteractor::SetTriggerBeenReleasedSinceLastPress( const bool bInTriggerBeenReleasedSinceLastPress )
+{
+	bHasTriggerBeenReleasedSinceLastPress = bInTriggerBeenReleasedSinceLastPress;
+}
+
+void UVREditorMotionControllerInteractor::PlayHapticEffect( const float Strength )
+{
+	IInputInterface* InputInterface = FSlateApplication::Get().GetInputInterface();
+	if ( InputInterface )
+	{
+		const double CurrentTime = FPlatformTime::Seconds();
+
+		// If we're dealing with an Oculus Rift, we have to setup haptic feedback directly.  Otherwise we can use our
+		// generic force feedback system
+		if ( GetHMDDeviceType() == EHMDDeviceType::DT_OculusRift )
+		{
+			// Haptics are a little strong on Oculus Touch, so we scale them down a bit
+			const float HapticScaleForRift = 0.8f;
+
+			FHapticFeedbackValues HapticFeedbackValues;
+			HapticFeedbackValues.Amplitude = Strength * HapticScaleForRift;
+			HapticFeedbackValues.Frequency = 0.5f;
+			InputInterface->SetHapticFeedbackValues( WorldInteraction->GetMotionControllerID(), ( int32 )ControllerHandSide, HapticFeedbackValues );
+
+			LastHapticTime = CurrentTime;
+		}
+		else
+		{
+			//@todo viewportinteration
+			FForceFeedbackValues ForceFeedbackValues;
+			ForceFeedbackValues.LeftLarge = ControllerHandSide == EControllerHand::Left ? Strength : 0;
+			ForceFeedbackValues.RightLarge = ControllerHandSide == EControllerHand::Right ? Strength : 0;
+
+			// @todo vreditor: If an Xbox controller is plugged in, this causes both the motion controllers and the Xbox controller to vibrate!
+			InputInterface->SetForceFeedbackChannelValues( WorldInteraction->GetMotionControllerID(), ForceFeedbackValues );
+
+			if ( ForceFeedbackValues.LeftLarge > KINDA_SMALL_NUMBER )
+			{
+				LastHapticTime = CurrentTime;
+			}
+
+			if ( ForceFeedbackValues.RightLarge > KINDA_SMALL_NUMBER )
+			{
+				LastHapticTime = CurrentTime;
+			}
+		}
+	}
+
+	// @todo vreditor: We'll stop haptics right away because if the frame hitches, the controller will be left vibrating
+	StopOldHapticEffects();
+}
+
+void UVREditorMotionControllerInteractor::StopOldHapticEffects()
+{
+	const float MinHapticTime = VREd::MinHapticTimeForRift->GetFloat();
+
+	IInputInterface* InputInterface = FSlateApplication::Get().GetInputInterface();
+	if ( InputInterface )
+	{
+		// If we're dealing with an Oculus Rift, we have to setup haptic feedback directly.  Otherwise we can use our
+		// generic force feedback system
+		if ( GetHMDDeviceType() == EHMDDeviceType::DT_OculusRift )
+		{
+			bool bWaitingForMoreHaptics = false;
+
+			if ( LastHapticTime != 0.0f )
+			{
+				do
+				{
+					bWaitingForMoreHaptics = false;
+
+					const double CurrentTime = FPlatformTime::Seconds();
+
+					// Left hand
+					if ( CurrentTime - LastHapticTime >= MinHapticTime )
+					{
+						FHapticFeedbackValues HapticFeedbackValues;
+						HapticFeedbackValues.Amplitude = 0.0f;
+						HapticFeedbackValues.Frequency = 0.0f;
+						InputInterface->SetHapticFeedbackValues( WorldInteraction->GetMotionControllerID(), ( int32 ) ControllerHandSide, HapticFeedbackValues );
+
+						LastHapticTime = 0.0;
+					}
+					else if ( LastHapticTime != 0.0 )
+					{
+						bWaitingForMoreHaptics = true;
+					}
+
+					if ( bWaitingForMoreHaptics && VREd::SleepForRiftHaptics->GetInt() != 0 )
+					{
+						FPlatformProcess::Sleep( 0 );
+					}
+				}
+				// @todo vreditor urgent: This is needed so that haptics don't play too long.  Our Oculus code doesn't currently support 
+				// multi-threading, so we need to delay the main thread to make sure we stop it before it rumbles for more than an instant!
+				while ( bWaitingForMoreHaptics && VREd::SleepForRiftHaptics->GetInt() != 0 );
+			}
+		}
+		else
+		{
+			// @todo vreditor: Do we need to cancel haptics for non-Rift devices?  Doesn't seem like it
+		}
+	}
+}
+
+bool UVREditorMotionControllerInteractor::GetTransformAndForwardVector( FTransform& OutHandTransform, FVector& OutForwardVector )
+{
+	if ( bHaveMotionController )
+	{
+		OutHandTransform = InteractorData.Transform;
+
+		const float LaserPointerRotationOffset = GetHMDDeviceType() == EHMDDeviceType::DT_OculusRift ? VREd::OculusLaserPointerRotationOffset->GetFloat() : VREd::ViveLaserPointerRotationOffset->GetFloat();
+		OutForwardVector = OutHandTransform.GetRotation().RotateVector( FRotator( LaserPointerRotationOffset, 0.0f, 0.0f ).RotateVector( FVector( 1.0f, 0.0f, 0.0f ) ) );
+
+		return true;
+	}
+
+	return false;
+}
+
+float UVREditorMotionControllerInteractor::GetTrackpadSlideDelta()
+{
+	const bool bIsAbsolute = ( GetVRMode().GetHMDDeviceType() == EHMDDeviceType::DT_SteamVR );
+	float SlideDelta = 0.0f;
+	if ( bIsTouchingTrackpad || !bIsAbsolute )
+	{
+		if ( bIsAbsolute )
+		{
+			SlideDelta = ( ( TrackpadPosition.Y - LastTrackpadPosition.Y ) * VREd::TrackpadAbsoluteDragSpeed->GetFloat() );
+		}
+		else
+		{
+			SlideDelta = ( TrackpadPosition.Y * VREd::TrackpadRelativeDragSpeed->GetFloat() );
+		}
+	}
+
+	return SlideDelta;
+}
+
+EControllerHand UVREditorMotionControllerInteractor::GetControllerSide() const
+{
+	return ControllerHandSide;
+}
+
+UMotionControllerComponent* UVREditorMotionControllerInteractor::GetMotionControllerComponent() const
+{
+	return MotionControllerComponent;
+}
+
+bool UVREditorMotionControllerInteractor::IsTouchingTrackpad() const
+{
+	return bIsTouchingTrackpad;
+}
+
+FVector2D UVREditorMotionControllerInteractor::GetTrackpadPosition() const
+{
+	return TrackpadPosition;
+}
+
+FVector2D UVREditorMotionControllerInteractor::GetLastTrackpadPosition() const
+{
+	return LastTrackpadPosition;
+}
+
+bool UVREditorMotionControllerInteractor::IsTrackpadPositionValid( const int32 AxisIndex ) const
+{
+	return bIsTrackpadPositionValid[ AxisIndex ];
+}
+
+FTimespan& UVREditorMotionControllerInteractor::GetLastTrackpadPositionUpdateTime()
+{
+	return LastTrackpadPositionUpdateTime;
+}
+
+
+/** Changes the color of the buttons on the handmesh */
+void UVREditorMotionControllerInteractor::ApplyButtonPressColors( const FViewportActionKeyInput& Action )
+{
+	const float PressStrength = 10.0f;
+	const FName ActionType = Action.ActionType;
+	const EInputEvent Event = Action.Event;
+
+	// Trigger
+	if ( ActionType == ViewportWorldActionTypes::SelectAndMove ||
+		ActionType == ViewportWorldActionTypes::SelectAndMove_LightlyPressed )
+	{
+		static FName StaticTriggerParameter( "B1" );
+		SetMotionControllerButtonPressedVisuals( Event, StaticTriggerParameter, PressStrength );
+	}
+
+	// Shoulder button
+	if ( ActionType == ViewportWorldActionTypes::WorldMovement )
+	{
+		static FName StaticShoulderParameter( "B2" );
+		SetMotionControllerButtonPressedVisuals( Event, StaticShoulderParameter, PressStrength );
+	}
+
+	// Trackpad
+	if ( ActionType == VRActionTypes::ConfirmRadialSelection )
+	{
+		static FName StaticTrackpadParameter( "B3" );
+		SetMotionControllerButtonPressedVisuals( Event, StaticTrackpadParameter, PressStrength );
+	}
+
+	// Modifier
+	if ( ActionType == VRActionTypes::Modifier )
+	{
+		static FName StaticModifierParameter( "B4" );
+		SetMotionControllerButtonPressedVisuals( Event, StaticModifierParameter, PressStrength );
+	}
+}
+
+void UVREditorMotionControllerInteractor::SetMotionControllerButtonPressedVisuals( const EInputEvent Event, const FName& ParameterName, const float PressStrength )
+{
+	if ( Event == IE_Pressed )
+	{
+		HandMeshMID->SetScalarParameterValue( ParameterName, PressStrength );
+	}
+	else if ( Event == IE_Released )
+	{
+		HandMeshMID->SetScalarParameterValue( ParameterName, 0.0f );
+	}
+}
+
+void UVREditorMotionControllerInteractor::ShowHelpForHand( const bool bShowIt )
+{
+	if ( bShowIt != bWantHelpLabels )
+	{
+		bWantHelpLabels = bShowIt;
+
+		const FTimespan CurrentTime = FTimespan::FromSeconds( FApp::GetCurrentTime() );
+		const FTimespan TimeSinceStartedFadingOut = CurrentTime - HelpLabelShowOrHideStartTime;
+		const FTimespan HelpLabelFadeDuration = FTimespan::FromSeconds( VREd::HelpLabelFadeDuration->GetFloat() );
+
+		// If we were already fading, account for that here
+		if ( TimeSinceStartedFadingOut < HelpLabelFadeDuration )
+		{
+			// We were already fading, so we'll reverse the time value so it feels continuous
+			HelpLabelShowOrHideStartTime = CurrentTime - ( HelpLabelFadeDuration - TimeSinceStartedFadingOut );
+		}
+		else
+		{
+			HelpLabelShowOrHideStartTime = FTimespan::FromSeconds( FApp::GetCurrentTime() );
+		}
+
+		if ( bShowIt && HelpLabels.Num() == 0 )
+		{
+			for ( const auto& KeyToAction : KeyToActionMap )
+			{
+				const FKey Key = KeyToAction.Key;
+				const FViewportActionKeyInput& Action = KeyToAction.Value;
+
+				UStaticMeshSocket* Socket = FindMeshSocketForKey( HandMeshComponent->StaticMesh, Key );
+				if ( Socket != nullptr )
+				{
+					FText LabelText;
+					FString ComponentName;
+
+					if ( Action.ActionType == VRActionTypes::Modifier )
+					{
+						LabelText = LOCTEXT( "ModifierHelp", "Modifier" );
+						ComponentName = TEXT( "ModifierHelp" );
+					}
+					else if ( Action.ActionType == ViewportWorldActionTypes::WorldMovement )
+					{
+						LabelText = LOCTEXT( "WorldMovementHelp", "Move World" );
+						ComponentName = TEXT( "WorldMovementHelp" );
+					}
+					else if ( Action.ActionType == ViewportWorldActionTypes::SelectAndMove )
+					{
+						LabelText = LOCTEXT( "SelectAndMoveHelp", "Select & Move" );
+						ComponentName = TEXT( "SelectAndMoveHelp" );
+					}
+					else if ( Action.ActionType == ViewportWorldActionTypes::SelectAndMove_LightlyPressed )
+					{
+						LabelText = LOCTEXT( "SelectAndMove_LightlyPressedHelp", "Select & Move" );
+						ComponentName = TEXT( "SelectAndMove_LightlyPressedHelp" );
+					}
+					else if ( Action.ActionType == VRActionTypes::Touch )
+					{
+						LabelText = LOCTEXT( "TouchHelp", "Slide" );
+						ComponentName = TEXT( "TouchHelp" );
+					}
+					else if ( Action.ActionType == ViewportWorldActionTypes::Undo )
+					{
+						LabelText = LOCTEXT( "UndoHelp", "Undo" );
+						ComponentName = TEXT( "UndoHelp" );
+					}
+					else if ( Action.ActionType == ViewportWorldActionTypes::Redo )
+					{
+						LabelText = LOCTEXT( "RedoHelp", "Redo" );
+						ComponentName = TEXT( "RedoHelp" );
+					}
+					else if ( Action.ActionType == ViewportWorldActionTypes::Delete )
+					{
+						LabelText = LOCTEXT( "DeleteHelp", "Delete" );
+						ComponentName = TEXT( "DeleteHelp" );
+					}
+					else if ( Action.ActionType == VRActionTypes::ConfirmRadialSelection )
+					{
+						LabelText = LOCTEXT( "ConfirmRadialSelectionHelp", "Radial Menu" );
+						ComponentName = TEXT( "ConfirmRadialSelectionHelp" );
+					}
+
+					const bool bWithSceneComponent = false;	// Nope, we'll spawn our own inside AFloatingText
+					check( VRMode );
+					AFloatingText* FloatingText = GetVRMode().SpawnTransientSceneActor<AFloatingText>( ComponentName, bWithSceneComponent );
+					FloatingText->SetText( LabelText );
+
+					HelpLabels.Add( Key, FloatingText );
+				}
+			}
+		}
+	}
+}
+
+
+void UVREditorMotionControllerInteractor::UpdateHelpLabels()
+{
+	const FTimespan HelpLabelFadeDuration = FTimespan::FromSeconds( VREd::HelpLabelFadeDuration->GetFloat() );
+
+	const FTransform HeadTransform = GetVRMode().GetHeadTransform();
+
+	// Only show help labels if the hand is pretty close to the face
+	const float DistanceToHead = ( GetTransform().GetLocation() - HeadTransform.GetLocation() ).Size();
+	const float MinDistanceToHeadForHelp = VREd::HelpLabelFadeDistance->GetFloat() * GetVRMode().GetWorldScaleFactor();	// (in cm)
+	bool bShowHelp = DistanceToHead <= MinDistanceToHeadForHelp;
+
+	// Don't show help if a UI is summoned on that hand
+	if ( HasUIInFront() || HasUIOnForearm() || GetVRMode().GetUISystem().IsShowingRadialMenu( this ) )
+	{
+		bShowHelp = false;
+	}
+
+	ShowHelpForHand( bShowHelp );
+
+	// Have the labels finished fading out?  If so, we'll kill their actors!
+	const FTimespan CurrentTime = FTimespan::FromSeconds( FApp::GetCurrentTime() );
+	const FTimespan TimeSinceStartedFadingOut = CurrentTime - HelpLabelShowOrHideStartTime;
+	if ( !bWantHelpLabels && ( TimeSinceStartedFadingOut > HelpLabelFadeDuration ) )
+	{
+		// Get rid of help text
+		for ( auto& KeyAndValue : HelpLabels )
+		{
+			AFloatingText* FloatingText = KeyAndValue.Value;
+			GetVRMode().DestroyTransientActor( FloatingText );
+		}
+		HelpLabels.Reset();
+	}
+	else
+	{
+		// Update fading state
+		float FadeAlpha = FMath::Clamp( ( float ) TimeSinceStartedFadingOut.GetTotalSeconds() / ( float ) HelpLabelFadeDuration.GetTotalSeconds(), 0.0f, 1.0f );
+		if ( !bWantHelpLabels )
+		{
+			FadeAlpha = 1.0f - FadeAlpha;
+		}
+
+		// Exponential falloff, so the fade is really obvious (gamma/HDR)
+		FadeAlpha = FMath::Pow( FadeAlpha, 3.0f );
+
+		for ( auto& KeyAndValue : HelpLabels )
+		{
+			const FKey Key = KeyAndValue.Key;
+			AFloatingText* FloatingText = KeyAndValue.Value;
+
+			UStaticMeshSocket* Socket = FindMeshSocketForKey( HandMeshComponent->StaticMesh, Key );
+			check( Socket != nullptr );
+			FTransform SocketRelativeTransform( Socket->RelativeRotation, Socket->RelativeLocation, Socket->RelativeScale );
+
+			// Oculus has asymmetrical controllers, so we the sock transform horizontally
+			if ( ControllerHandSide == EControllerHand::Right &&
+				GetVRMode().GetHMDDeviceType() == EHMDDeviceType::DT_OculusRift )
+			{
+				const FVector Scale3D = SocketRelativeTransform.GetLocation();
+				SocketRelativeTransform.SetLocation( FVector( Scale3D.X, -Scale3D.Y, Scale3D.Z ) );
+			}
+
+			// Make sure the labels stay the same size even when the world is scaled
+			FTransform HandTransformWithWorldToMetersScaling = GetTransform();
+			HandTransformWithWorldToMetersScaling.SetScale3D( HandTransformWithWorldToMetersScaling.GetScale3D() * FVector( GetVRMode().GetWorldScaleFactor() ) );
+
+			// Position right on top of the controller itself
+			FTransform FloatingTextTransform = SocketRelativeTransform * HandTransformWithWorldToMetersScaling;
+			FloatingText->SetActorTransform( FloatingTextTransform );
+
+			// Orientate it toward the viewer
+			FloatingText->Update( HeadTransform.GetLocation() );
+
+			// Update fade state
+			FloatingText->SetOpacity( FadeAlpha );
+		}
+	}
+}
+
+
+UStaticMeshSocket* UVREditorMotionControllerInteractor::FindMeshSocketForKey( UStaticMesh* StaticMesh, const FKey Key )
+{
+	// @todo vreditor: Hard coded mapping of socket names (e.g. "Shoulder") to expected names of sockets in the static mesh
+	FName SocketName = NAME_None;
+	if ( Key == EKeys::MotionController_Left_Shoulder || Key == EKeys::MotionController_Right_Shoulder )
+	{
+		static FName ShoulderSocketName( "Shoulder" );
+		SocketName = ShoulderSocketName;
+	}
+	else if ( Key == EKeys::MotionController_Left_Trigger || Key == EKeys::MotionController_Right_Trigger ||
+		Key == VREditorKeyNames::MotionController_Left_FullyPressedTriggerAxis || Key == VREditorKeyNames::MotionController_Right_FullyPressedTriggerAxis ||
+		Key == VREditorKeyNames::MotionController_Left_LightlyPressedTriggerAxis || Key == VREditorKeyNames::MotionController_Right_LightlyPressedTriggerAxis )
+	{
+		static FName TriggerSocketName( "Trigger" );
+		SocketName = TriggerSocketName;
+	}
+	else if ( Key == EKeys::MotionController_Left_Grip1 || Key == EKeys::MotionController_Right_Grip1 )
+	{
+		static FName GripSocketName( "Grip" );
+		SocketName = GripSocketName;
+	}
+	else if ( Key == EKeys::MotionController_Left_Thumbstick || Key == EKeys::MotionController_Right_Thumbstick )
+	{
+		static FName ThumbstickSocketName( "Thumbstick" );
+		SocketName = ThumbstickSocketName;
+	}
+	else if ( Key == SteamVRControllerKeyNames::Touch0 || Key == SteamVRControllerKeyNames::Touch1 )
+	{
+		static FName TouchSocketName( "Touch" );
+		SocketName = TouchSocketName;
+	}
+	else if ( Key == EKeys::MotionController_Left_Thumbstick_Down || Key == EKeys::MotionController_Right_Thumbstick_Down )
+	{
+		static FName DownSocketName( "Down" );
+		SocketName = DownSocketName;
+	}
+	else if ( Key == EKeys::MotionController_Left_Thumbstick_Up || Key == EKeys::MotionController_Right_Thumbstick_Up )
+	{
+		static FName UpSocketName( "Up" );
+		SocketName = UpSocketName;
+	}
+	else if ( Key == EKeys::MotionController_Left_Thumbstick_Left || Key == EKeys::MotionController_Right_Thumbstick_Left )
+	{
+		static FName LeftSocketName( "Left" );
+		SocketName = LeftSocketName;
+	}
+	else if ( Key == EKeys::MotionController_Left_Thumbstick_Right || Key == EKeys::MotionController_Right_Thumbstick_Right )
+	{
+		static FName RightSocketName( "Right" );
+		SocketName = RightSocketName;
+	}
+	else if ( Key == EKeys::MotionController_Left_FaceButton1 || Key == EKeys::MotionController_Right_FaceButton1 )
+	{
+		static FName FaceButton1SocketName( "FaceButton1" );
+		SocketName = FaceButton1SocketName;
+	}
+	else if ( Key == EKeys::MotionController_Left_FaceButton2 || Key == EKeys::MotionController_Right_FaceButton2 )
+	{
+		static FName FaceButton2SocketName( "FaceButton2" );
+		SocketName = FaceButton2SocketName;
+	}
+	else if ( Key == EKeys::MotionController_Left_FaceButton3 || Key == EKeys::MotionController_Right_FaceButton3 )
+	{
+		static FName FaceButton3SocketName( "FaceButton3" );
+		SocketName = FaceButton3SocketName;
+	}
+	else if ( Key == EKeys::MotionController_Left_FaceButton4 || Key == EKeys::MotionController_Right_FaceButton4 )
+	{
+		static FName FaceButton4SocketName( "FaceButton4" );
+		SocketName = FaceButton4SocketName;
+	}
+	else
+	{
+		// Not a key that we care about
+	}
+
+	if ( SocketName != NAME_None )
+	{
+		UStaticMeshSocket* Socket = StaticMesh->FindSocket( SocketName );
+		if ( Socket != nullptr )
+		{
+			return Socket;
+		}
+	}
+
+	return nullptr;
+};
+
+void UVREditorMotionControllerInteractor::SetLaserVisuals( const FLinearColor& NewColor, const float CrawlFade, const float CrawlSpeed )
+{
+	static FName StaticLaserColorParameterName( "LaserColor" );
+	LaserPointerMID->SetVectorParameterValue( StaticLaserColorParameterName, NewColor );
+	TranslucentLaserPointerMID->SetVectorParameterValue( StaticLaserColorParameterName, NewColor );
+
+	static FName StaticCrawlParameterName( "Crawl" );
+	LaserPointerMID->SetScalarParameterValue( StaticCrawlParameterName, CrawlFade );
+	TranslucentLaserPointerMID->SetScalarParameterValue( StaticCrawlParameterName, CrawlFade );
+
+	static FName StaticCrawlSpeedParameterName( "CrawlSpeed" );
+	LaserPointerMID->SetScalarParameterValue( StaticCrawlSpeedParameterName, CrawlSpeed );
+	TranslucentLaserPointerMID->SetScalarParameterValue( StaticCrawlSpeedParameterName, CrawlSpeed );
+
+	static FName StaticHandTrimColorParameter( "TrimGlowColor" );
+	HandMeshMID->SetVectorParameterValue( StaticHandTrimColorParameter, NewColor );
+
+	HoverPointLightComponent->SetLightColor( NewColor );
+}
+
+void UVREditorMotionControllerInteractor::UpdateRadialMenuInput( const float DeltaTime )
+{
+	UVREditorUISystem& UISystem = GetVRMode().GetUISystem();
+
+	//Update the radial menu
+	EViewportInteractionDraggingMode DraggingMode = GetDraggingMode();
+	if ( !HasUIInFront() &&
+		( bIsTrackpadPositionValid[ 0 ] && bIsTrackpadPositionValid[ 1 ] ) &&
+		DraggingMode != EViewportInteractionDraggingMode::ActorsWithGizmo &&
+		DraggingMode != EViewportInteractionDraggingMode::ActorsFreely &&
+		DraggingMode != EViewportInteractionDraggingMode::ActorsAtLaserImpact &&
+		DraggingMode != EViewportInteractionDraggingMode::AssistingDrag &&
+		!UISystem.IsInteractorDraggingDockUI( this ) &&
+		!UISystem.IsShowingRadialMenu( Cast<UVREditorInteractor>( OtherInteractor ) ) )
+	{
+		const EHMDDeviceType::Type HMDDevice = GetHMDDeviceType();
+
+		// Spawn the radial menu if we are using the touchpad for steamvr or the analog stick for the oculus
+		if ( ( HMDDevice == EHMDDeviceType::DT_SteamVR && bIsTouchingTrackpad ) ||
+			( HMDDevice == EHMDDeviceType::DT_OculusRift && TrackpadPosition.Size() > VREd::MinJoystickOffsetBeforeRadialMenu->GetFloat() ) )
+		{
+			UISystem.TryToSpawnRadialMenu( this );
+		}
+		else
+		{
+			// Hide it if we are not using the touchpad or analog stick
+			UISystem.HideRadialMenu( this );
+		}
+
+		// Update the radial menu if we are already showing the radial menu
+		if ( UISystem.IsShowingRadialMenu( this ) )
+		{
+			UISystem.UpdateRadialMenu( this );
+		}
+	}
+}
+
+#undef LOCTEXT_NAMESPACE
\ No newline at end of file
diff --git a/Engine/Source/Editor/VREditor/VREditorMotionControllerInteractor.h b/Engine/Source/Editor/VREditor/VREditorMotionControllerInteractor.h
new file mode 100644
--- /dev/null
+++ b/Engine/Source/Editor/VREditor/VREditorMotionControllerInteractor.h
@@ -0,0 +1,215 @@
+// Copyright 1998-2015 Epic Games, Inc. All Rights Reserved.
+
+#pragma once
+
+#include "VREditorInteractor.h"
+#include "HeadMountedDisplayTypes.h"	// For EHMDDeviceType::Type
+#include "VREditorMotionControllerInteractor.generated.h"
+
+/**
+ * Represents the interactor in the world
+ */
+UCLASS()
+class UVREditorMotionControllerInteractor : public UVREditorInteractor
+{
+	GENERATED_UCLASS_BODY()
+
+public:
+
+	virtual ~UVREditorMotionControllerInteractor();
+	
+	// UVREditorInteractor overrides
+	virtual void Init( class FVREditorMode* InVRMode ) override;
+	/** Gets the trackpad slide delta */
+	virtual float GetSlideDelta() override;
+
+	/** Sets up all components */
+	void SetupComponent( AActor* OwningActor );
+
+	/** Sets the EControllerHand for this motioncontroller */
+	void SetControllerHandSide( const EControllerHand InControllerHandSide );
+
+	// IViewportInteractorInterface overrides
+	virtual void Shutdown() override;
+	virtual void Tick( const float DeltaTime ) override;
+	virtual void CalculateDragRay() override;
+
+	/** @return Returns the type of HMD we're dealing with */
+	EHMDDeviceType::Type GetHMDDeviceType() const;
+
+	// UViewportInteractor
+	virtual void HandleInputKey( FViewportActionKeyInput& Action, const FKey Key, const EInputEvent Event, bool& bOutWasHandled ) override;
+	virtual bool GetTransformAndForwardVector( FTransform& OutHandTransform, FVector& OutForwardVector ) override;
+
+	/** Gets the Y delta of the trackpad */
+	float GetTrackpadSlideDelta();
+
+	/** Starts haptic feedback for physical motion controller */
+	virtual void PlayHapticEffect( const float Strength ) override;
+	
+	/** Stops playing any haptic effects that have been going for a while.  Called every frame. */
+	void StopOldHapticEffects();
+	
+	/** Get the side of the controller */
+	EControllerHand GetControllerSide() const;
+
+	/** Get the motioncontroller component of this interactor */
+	class UMotionControllerComponent* GetMotionControllerComponent() const;
+
+	/** Check if the touchpad is currently touched */
+	bool IsTouchingTrackpad() const;
+
+	/** Get the current position of the trackpad or analog stick */
+	FVector2D GetTrackpadPosition() const;
+
+	/** Get the last position of the trackpad or analog stick */
+	FVector2D GetLastTrackpadPosition() const;
+
+	/** If the trackpad values are valid */
+	bool IsTrackpadPositionValid( const int32 AxisIndex ) const;
+
+	/** Get when the last time the trackpad position was updated */
+	FTimespan& GetLastTrackpadPositionUpdateTime();
+
+protected:
+
+	// ViewportInteractor
+	virtual void HandleInputAxis( FViewportActionKeyInput& Action, const FKey Key, const float Delta, const float DeltaTime, bool& bOutWasHandled ) override;
+
+	/** Polls input for the motion controllers transforms */
+	virtual void PollInput() override;
+
+	/** Check if the trigger is lightly pressed */
+	bool IsTriggerAtLeastLightlyPressed() const;
+
+	/** Get the real time the trigger was lightly pressed */
+	double GetRealTimeTriggerWasLightlyPressed() const;
+
+	/** Set that the trigger is at least lighty pressed */
+	void SetTriggerAtLeastLightlyPressed( const bool bInTriggerAtLeastLightlyPressed );
+
+	/** Set that the trigger has been released since last press */
+	void SetTriggerBeenReleasedSinceLastPress( const bool bInTriggerBeenReleasedSinceLastPress );
+
+private:
+
+	/** Changes the color of the buttons on the handmesh */
+	void ApplyButtonPressColors( const FViewportActionKeyInput& Action );
+
+	/** Set the visuals for a button on the motion controller */
+	void SetMotionControllerButtonPressedVisuals( const EInputEvent Event, const FName& ParameterName, const float PressStrength );
+
+	/** Pops up some help text labels for the controller in the specified hand, or hides it, if requested */
+	void ShowHelpForHand( const bool bShowIt );
+
+	/** Called every frame to update the position of any floating help labels */
+	void UpdateHelpLabels();
+
+	/** Given a mesh and a key name, tries to find a socket on the mesh that matches a supported key */
+	UStaticMeshSocket* FindMeshSocketForKey( UStaticMesh* StaticMesh, const FKey Key );
+
+	/** Sets the visuals of the LaserPointer */
+	void SetLaserVisuals( const FLinearColor& NewColor, const float CrawlFade, const float CrawlSpeed );
+
+	/** Updates the radial menu */
+	void UpdateRadialMenuInput( const float DeltaTime );
+
+protected:
+	
+	/** Motion controller component which handles late-frame transform updates of all parented sub-components */
+	UPROPERTY()
+	class UMotionControllerComponent* MotionControllerComponent;
+
+	//
+	// Graphics
+	//
+
+	/** Mesh for this hand */
+	UPROPERTY()
+	class UStaticMeshComponent* HandMeshComponent;
+
+	/** Mesh for this hand's laser pointer */
+	UPROPERTY()
+	class UStaticMeshComponent* LaserPointerMeshComponent;
+
+	/** MID for laser pointer material (opaque parts) */
+	UPROPERTY()
+	class UMaterialInstanceDynamic* LaserPointerMID;
+
+	/** MID for laser pointer material (translucent parts) */
+	UPROPERTY()
+	class UMaterialInstanceDynamic* TranslucentLaserPointerMID;
+
+	/** Hover impact indicator mesh */
+	UPROPERTY()
+	class UStaticMeshComponent* HoverMeshComponent;
+
+	/** Hover point light */
+	UPROPERTY()
+	class UPointLightComponent* HoverPointLightComponent;
+
+	/** MID for hand mesh */
+	UPROPERTY()
+	class UMaterialInstanceDynamic* HandMeshMID;
+
+	/** Right or left hand */
+	EControllerHand ControllerHandSide;
+
+	/** True if this hand has a motion controller (or both!) */
+	bool bHaveMotionController;
+
+	// Special key action names for motion controllers
+	static const FName TrackpadPositionX;
+	static const FName TrackpadPositionY;
+	static const FName TriggerAxis;
+	static const FName MotionController_Left_FullyPressedTriggerAxis;
+	static const FName MotionController_Right_FullyPressedTriggerAxis;
+	static const FName MotionController_Left_LightlyPressedTriggerAxis;
+	static const FName MotionController_Right_LightlyPressedTriggerAxis;
+
+private:
+
+	//
+	// Trigger axis state
+	//
+
+	/** True if trigger is fully pressed right now (or within some small threshold) */
+	bool bIsTriggerFullyPressed;
+
+	/** True if the trigger is currently pulled far enough that we consider it in a "half pressed" state */
+	bool bIsTriggerAtLeastLightlyPressed;
+
+	/** Real time that the trigger became lightly pressed.  If the trigger remains lightly pressed for a reasonably 
+	long duration before being pressed fully, we'll continue to treat it as a light press in some cases */
+	double RealTimeTriggerWasLightlyPressed;
+
+	/** True if trigger has been fully released since the last press */
+	bool bHasTriggerBeenReleasedSinceLastPress;
+
+	//
+	// Haptic feedback
+	//
+
+	/** The last real time we played a haptic effect.  This is used to turn off haptic effects shortly after they're triggered. */
+	double LastHapticTime;
+
+	//
+	// Trackpad support
+	//
+
+	/** True if the trackpad is actively being touched */
+	bool bIsTouchingTrackpad;
+
+	/** Position of the touched trackpad */
+	FVector2D TrackpadPosition;
+
+	/** Last position of the touched trackpad */
+	FVector2D LastTrackpadPosition;
+
+	/** True if we have a valid trackpad position (for each axis) */
+	bool bIsTrackpadPositionValid[ 2 ];
+
+	/** Real time that the last trackpad position was last updated.  Used to filter out stale previous data. */
+	FTimespan LastTrackpadPositionUpdateTime;
+
+};
\ No newline at end of file
diff --git a/Engine/Source/Editor/VREditor/VREditorRadialMenu.cpp b/Engine/Source/Editor/VREditor/VREditorRadialMenu.cpp
--- a/Engine/Source/Editor/VREditor/VREditorRadialMenu.cpp
+++ b/Engine/Source/Editor/VREditor/VREditorRadialMenu.cpp
@@ -9,6 +9,7 @@
 #include "VREditorWorldInteraction.h"
 #include "VREditorTransformGizmo.h"
 #include "VREditorInteractor.h"
+#include "VREditorMotionControllerInteractor.h"
 
 #include "Editor/LevelEditor/Public/LevelEditor.h"
 
@@ -30,53 +31,57 @@ PreviousItem( nullptr )
 
 void UVREditorRadialMenu::Update( const UVREditorInteractor* Interactor )
 {
-	TrackpadPosition = Interactor->GetTrackpadPosition();
-
-	if ( !GetOwner()->bHidden )
+	const UVREditorMotionControllerInteractor* MotionControllerInteractor = Cast< UVREditorMotionControllerInteractor >( Interactor );
+	if( MotionControllerInteractor )
 	{
-		// Check if position is from the center
-		if ( !IsInMenuRadius() )
-		{
-			const float AnglePerItem = 360 / MenuItems.Num();
-			float Angle = FRotator::NormalizeAxis( FMath::RadiansToDegrees( FMath::Atan2( TrackpadPosition.X, TrackpadPosition.Y ) ) );
-			TrackpadAngle = Angle;
+		TrackpadPosition = MotionControllerInteractor->GetTrackpadPosition();
 
-			Angle += AnglePerItem / 2.0f;
-			if (Angle < 0)
+		if ( !GetOwner()->bHidden )
+		{
+			// Check if position is from the center
+			if ( !IsInMenuRadius() )
 			{
-				Angle = Angle + 360;
-			}
+				const float AnglePerItem = 360 / MenuItems.Num();
+				float Angle = FRotator::NormalizeAxis( FMath::RadiansToDegrees( FMath::Atan2( TrackpadPosition.X, TrackpadPosition.Y ) ) );
+				TrackpadAngle = Angle;
 
-			const int32 Index = (Angle / AnglePerItem);
-			UVREditorRadialMenuItem* NewMenuItem = MenuItems[Index];
-			if (NewMenuItem)
-			{
-				CurrentItem = NewMenuItem;
-			}
+				Angle += AnglePerItem / 2.0f;
+				if ( Angle < 0 )
+				{
+					Angle = Angle + 360;
+				}
 
-			// Update the visuals
-			if (CurrentItem != PreviousItem)
-			{
-				if (CurrentItem)
+				const int32 Index = ( Angle / AnglePerItem );
+				UVREditorRadialMenuItem* NewMenuItem = MenuItems[ Index ];
+				if ( NewMenuItem )
 				{
-					CurrentItem->OnEnterHover();
+					CurrentItem = NewMenuItem;
 				}
 
-				if (PreviousItem)
+				// Update the visuals
+				if ( CurrentItem != PreviousItem )
 				{
-					PreviousItem->OnLeaveHover();
+					if ( CurrentItem )
+					{
+						CurrentItem->OnEnterHover();
+					}
+
+					if ( PreviousItem )
+					{
+						PreviousItem->OnLeaveHover();
+					}
 				}
-			}
 
-			PreviousItem = CurrentItem;
-		}
-		else
-		{
-			if (CurrentItem)
+				PreviousItem = CurrentItem;
+			}
+			else
 			{
-				CurrentItem->OnLeaveHover();
-				CurrentItem = nullptr;
-				PreviousItem = nullptr;
+				if ( CurrentItem )
+				{
+					CurrentItem->OnLeaveHover();
+					CurrentItem = nullptr;
+					PreviousItem = nullptr;
+				}
 			}
 		}
 	}
diff --git a/Engine/Source/Editor/VREditor/VREditorUISystem.cpp b/Engine/Source/Editor/VREditor/VREditorUISystem.cpp
--- a/Engine/Source/Editor/VREditor/VREditorUISystem.cpp
+++ b/Engine/Source/Editor/VREditor/VREditorUISystem.cpp
@@ -12,6 +12,8 @@
 #include "VREditorWorldInteraction.h"
 #include "ViewportInteractor.h"
 #include "VREditorInteractor.h"
+#include "VREditorMotionControllerInteractor.h"
+
 // UI
 #include "WidgetComponent.h"
 #include "VREditorWidgetComponent.h"
@@ -170,7 +172,7 @@ void UVREditorUISystem::Shutdown()
 void UVREditorUISystem::OnVRAction( FEditorViewportClient& ViewportClient, UViewportInteractor* Interactor,
 	const FViewportActionKeyInput& Action, bool& bOutIsInputCaptured, bool& bWasHandled )
 {
-	UVREditorInteractor* VREditorInteractor = Cast<UVREditorInteractor>( Interactor );
+	UVREditorMotionControllerInteractor* VREditorInteractor = Cast<UVREditorMotionControllerInteractor>( Interactor );
 	if ( VREditorInteractor )
 	{
 		if ( !bWasHandled )
@@ -326,7 +328,7 @@ void UVREditorUISystem::OnVRAction( FEditorViewportClient& ViewportClient, UView
 
 void UVREditorUISystem::OnVRHoverUpdate( FEditorViewportClient& ViewportClient, UViewportInteractor* Interactor, FVector& HoverImpactPoint, bool& bWasHandled )
 {
-	UVREditorInteractor* VREditorInteractor = Cast<UVREditorInteractor>( Interactor );
+	UVREditorMotionControllerInteractor* VREditorInteractor = Cast<UVREditorMotionControllerInteractor>( Interactor );
 	if( VREditorInteractor != nullptr )
 	{
 		if( !bWasHandled && Interactor->GetDraggingMode() != EViewportInteractionDraggingMode::Interactable )
@@ -557,13 +559,13 @@ void UVREditorUISystem::Tick( FEditorViewportClient* ViewportClient, const float
 
 		if( HandInteractorThatNeedsQuickMenu != nullptr && !QuickMenuUI->IsUIVisible() )
 		{
-			UVREditorInteractor* VREditorHandInteractorThatNeedsQuickMenu = Cast<UVREditorInteractor>( HandInteractorThatNeedsQuickMenu );
-			if ( VREditorHandInteractorThatNeedsQuickMenu )
+			UVREditorMotionControllerInteractor* MotionControllerHandInteractorThatNeedsQuickMenu = Cast<UVREditorMotionControllerInteractor>( HandInteractorThatNeedsQuickMenu );
+			if ( MotionControllerHandInteractorThatNeedsQuickMenu )
 			{
-				const AVREditorFloatingUI::EDockedTo DockTo = ( VREditorHandInteractorThatNeedsQuickMenu->GetControllerSide() == EControllerHand::Left ) ? AVREditorFloatingUI::EDockedTo::LeftArm : AVREditorFloatingUI::EDockedTo::RightArm;
+				const AVREditorFloatingUI::EDockedTo DockTo = ( MotionControllerHandInteractorThatNeedsQuickMenu->GetControllerSide() == EControllerHand::Left ) ? AVREditorFloatingUI::EDockedTo::LeftArm : AVREditorFloatingUI::EDockedTo::RightArm;
 				QuickMenuUI->SetDockedTo( DockTo );
 				QuickMenuUI->ShowUI( true );
-				VREditorHandInteractorThatNeedsQuickMenu->SetHasUIOnForearm( true );
+				MotionControllerHandInteractorThatNeedsQuickMenu->SetHasUIOnForearm( true );
 			}
 		}
 	}
@@ -947,8 +949,12 @@ void UVREditorUISystem::ShowEditorUIPanel( AVREditorFloatingUI* Panel, UVREditor
 				}
 			}
 			
-			const AVREditorFloatingUI::EDockedTo NewDockedTo = Interactor->GetControllerSide() == EControllerHand::Left ? AVREditorFloatingUI::EDockedTo::LeftHand : AVREditorFloatingUI::EDockedTo::RightHand;
-			Panel->SetDockedTo( NewDockedTo );
+			UVREditorMotionControllerInteractor* MotionControllerInteractor = Cast<UVREditorMotionControllerInteractor>( Interactor );
+			if ( MotionControllerInteractor )
+			{
+				const AVREditorFloatingUI::EDockedTo NewDockedTo = MotionControllerInteractor->GetControllerSide() == EControllerHand::Left ? AVREditorFloatingUI::EDockedTo::LeftHand : AVREditorFloatingUI::EDockedTo::RightHand;
+				Panel->SetDockedTo( NewDockedTo );
+			}
 			
 			if ( bShouldShow )
 			{
@@ -1005,37 +1011,41 @@ void UVREditorUISystem::UpdateRadialMenu( UVREditorInteractor* Interactor )
 
 void UVREditorUISystem::TryToSpawnRadialMenu( UVREditorInteractor* Interactor )
 {
-	UVREditorInteractor* DockedToInteractor = nullptr;
-	if( !QuickRadialMenu->bHidden )
+	UVREditorMotionControllerInteractor* MotionControllerInteractor = Cast<UVREditorMotionControllerInteractor>( Interactor );
+	if( MotionControllerInteractor )
 	{
-		const EControllerHand ControllerHand = QuickRadialMenu->GetDockedTo() == AVREditorFloatingUI::EDockedTo::LeftHand ? EControllerHand::Left : EControllerHand::Right;
-		DockedToInteractor = GetOwner().GetHandInteractor( ControllerHand );
-	}
+		UVREditorInteractor* DockedToInteractor = nullptr;
+		if ( !QuickRadialMenu->bHidden )
+		{
+			const EControllerHand ControllerHand = QuickRadialMenu->GetDockedTo() == AVREditorFloatingUI::EDockedTo::LeftHand ? EControllerHand::Left : EControllerHand::Right;
+			DockedToInteractor = GetOwner().GetHandInteractor( ControllerHand );
+		}
 
-	EViewportInteractionDraggingMode DraggingMode = Interactor->GetDraggingMode();
-
-	bool bNeedsSpawn =
-		( QuickRadialMenu->bHidden || DockedToInteractor != Interactor ) &&
-		DraggingMode != EViewportInteractionDraggingMode::ActorsAtLaserImpact &&	// Don't show radial menu if the hand is busy dragging something around
-		DraggingMode != EViewportInteractionDraggingMode::ActorsFreely &&
-		DraggingMode != EViewportInteractionDraggingMode::ActorsWithGizmo &&
-		DraggingMode != EViewportInteractionDraggingMode::AssistingDrag &&
-		( InteractorDraggingUI == nullptr || InteractorDraggingUI != Interactor ) &&
-		!Interactor->IsHoveringOverUI();	// Don't show radial menu when aiming at a UI  (too much clutter)
-
-	UVREditorRadialMenu* RadialMenu = QuickRadialMenu->GetUserWidget<UVREditorRadialMenu>();
-	// We need to update the trackpad position in the radialmenu before checking if it can be used
-	RadialMenu->Update( Interactor );
-	if( RadialMenu && RadialMenu->IsInMenuRadius() )
-	{
-		bNeedsSpawn = false;
-	}
+		EViewportInteractionDraggingMode DraggingMode = Interactor->GetDraggingMode();
+
+		bool bNeedsSpawn =
+			( QuickRadialMenu->bHidden || DockedToInteractor != Interactor ) &&
+			DraggingMode != EViewportInteractionDraggingMode::ActorsAtLaserImpact &&	// Don't show radial menu if the hand is busy dragging something around
+			DraggingMode != EViewportInteractionDraggingMode::ActorsFreely &&
+			DraggingMode != EViewportInteractionDraggingMode::ActorsWithGizmo &&
+			DraggingMode != EViewportInteractionDraggingMode::AssistingDrag &&
+			( InteractorDraggingUI == nullptr || InteractorDraggingUI != Interactor ) &&
+			!Interactor->IsHoveringOverUI();	// Don't show radial menu when aiming at a UI  (too much clutter)
+
+		UVREditorRadialMenu* RadialMenu = QuickRadialMenu->GetUserWidget<UVREditorRadialMenu>();
+		// We need to update the trackpad position in the radialmenu before checking if it can be used
+		RadialMenu->Update( Interactor );
+		if ( RadialMenu && RadialMenu->IsInMenuRadius() )
+		{
+			bNeedsSpawn = false;
+		}
 
-	if( bNeedsSpawn )
-	{
-		const AVREditorFloatingUI::EDockedTo DockedTo = Interactor->GetControllerSide() == EControllerHand::Left ? AVREditorFloatingUI::EDockedTo::LeftHand : AVREditorFloatingUI::EDockedTo::RightHand;
-		QuickRadialMenu->SetDockedTo( DockedTo );
-		QuickRadialMenu->ShowUI( true );
+		if ( bNeedsSpawn )
+		{
+			const AVREditorFloatingUI::EDockedTo DockedTo = MotionControllerInteractor->GetControllerSide() == EControllerHand::Left ? AVREditorFloatingUI::EDockedTo::LeftHand : AVREditorFloatingUI::EDockedTo::RightHand;
+			QuickRadialMenu->SetDockedTo( DockedTo );
+			QuickRadialMenu->ShowUI( true );
+		}
 	}
 }
 
@@ -1075,7 +1085,7 @@ FTransform UVREditorUISystem::MakeDockableUITransform( AVREditorDockableWindow*
 
 	if( DraggingUI != nullptr &&  DraggingUI->GetDockedTo() == AVREditorFloatingUI::EDockedTo::Dragging && InteractorDraggingUI && Interactor == InteractorDraggingUI )
 	{
-		UVREditorInteractor* OtherInteractor = Cast<UVREditorInteractor>( InteractorDraggingUI->GetOtherInteractor() );
+		UVREditorMotionControllerInteractor* OtherInteractor = Cast<UVREditorMotionControllerInteractor>( InteractorDraggingUI->GetOtherInteractor() );
 		if ( OtherInteractor )
 		{
 			const float WorldScaleFactor = GetOwner().GetWorldScaleFactor();
@@ -1160,11 +1170,11 @@ void UVREditorUISystem::StopDraggingDockUI( UVREditorInteractor* VREditorInterac
 		UViewportInteractor* OtherInteractor = VREditorInteractor->GetOtherInteractor();
 		if ( OtherInteractor )
 		{
-			UVREditorInteractor* OtherVREditorInteractor = Cast<UVREditorInteractor>( OtherInteractor );
-			if ( OtherVREditorInteractor )
+			UVREditorMotionControllerInteractor* MotionControllerOtherInteractor = Cast<UVREditorMotionControllerInteractor>( OtherInteractor );
+			if ( MotionControllerOtherInteractor )
 			{
 				bool bOnHand = false;
-				const float Distance = FVector::Dist( LastDraggingHoverLocation, OtherVREditorInteractor->GetTransform().GetLocation() ) / GetOwner().GetWorldScaleFactor();
+				const float Distance = FVector::Dist( LastDraggingHoverLocation, MotionControllerOtherInteractor->GetTransform().GetLocation() ) / GetOwner().GetWorldScaleFactor();
 				if ( Distance > VREd::MinDockDragDistance->GetFloat() )
 				{
 					DraggingUI->SetDockedTo( AVREditorFloatingUI::EDockedTo::Room );
@@ -1173,12 +1183,12 @@ void UVREditorUISystem::StopDraggingDockUI( UVREditorInteractor* VREditorInterac
 				{
 					bOnHand = true;
 					const FVector EditorUIRelativeOffset( DraggingUI->GetSize().Y * 0.5f + VREd::EditorUIOffsetFromHand->GetFloat(), 0.0f, 0.0f );
-					const FTransform MovingUIGoalTransform = DraggingUI->MakeUITransformLockedToHand( OtherVREditorInteractor, false, EditorUIRelativeOffset, FRotator( 90.0f, 180.0f, 0.0f ) );
-					const AVREditorFloatingUI::EDockedTo NewDockedTo = OtherVREditorInteractor->GetControllerSide() == EControllerHand::Left ? AVREditorFloatingUI::EDockedTo::LeftHand : AVREditorFloatingUI::EDockedTo::RightHand;
+					const FTransform MovingUIGoalTransform = DraggingUI->MakeUITransformLockedToHand( MotionControllerOtherInteractor, false, EditorUIRelativeOffset, FRotator( 90.0f, 180.0f, 0.0f ) );
+					const AVREditorFloatingUI::EDockedTo NewDockedTo = MotionControllerOtherInteractor->GetControllerSide() == EControllerHand::Left ? AVREditorFloatingUI::EDockedTo::LeftHand : AVREditorFloatingUI::EDockedTo::RightHand;
 					DraggingUI->MoveTo( MovingUIGoalTransform, VREd::DockSnapTime->GetFloat(), NewDockedTo );
 				}
 
-				ShowEditorUIPanel( DraggingUI, OtherVREditorInteractor, true, bOnHand );
+				ShowEditorUIPanel( DraggingUI, MotionControllerOtherInteractor, true, bOnHand );
 
 				DraggingUI = nullptr;
 				InteractorDraggingUI->SetDraggingMode( EViewportInteractionDraggingMode::Nothing );
diff --git a/Engine/Source/Editor/ViewportInteraction/Public/ViewportInteractorMotionController.h b/Engine/Source/Editor/ViewportInteraction/Public/ViewportInteractorMotionController.h
deleted file mode 100644
diff --git a/Engine/Source/Editor/ViewportInteraction/ViewportInteractorMotionController.cpp b/Engine/Source/Editor/ViewportInteraction/ViewportInteractorMotionController.cpp
deleted file mode 100644
