Auto-entry into VR editing mode when a Vive or Oculus Rift is worn. Editor setting under Experimental to toggle the behavior, new enum EHMDWornState which defaults to Unknown for headsets that don't have sensors to detect user proximity/interaction.

diff --git a/Engine/Plugins/Runtime/OculusRift/Source/OculusRift/Private/OculusRiftHMD.cpp b/Engine/Plugins/Runtime/OculusRift/Source/OculusRift/Private/OculusRiftHMD.cpp
--- a/Engine/Plugins/Runtime/OculusRift/Source/OculusRift/Private/OculusRiftHMD.cpp
+++ b/Engine/Plugins/Runtime/OculusRift/Source/OculusRift/Private/OculusRiftHMD.cpp
@@ -607,6 +607,29 @@ bool FOculusRiftHMD::IsHMDConnected()
 	return Settings->Flags.bHMDEnabled && ovr_Detect(0).IsOculusHMDConnected;
 }
 
+EHMDWornState::Type FOculusRiftHMD::GetHMDWornState()
+{
+	FOvrSessionShared::AutoSession OvrSession(Session);
+	// If there is no VR Session, initialize the Oculus HMD
+	if (!OvrSession)
+	{
+		InitDevice();
+	}
+	if (OvrSession && !pCustomPresent->NeedsToKillHmd())
+	{
+		ovrSessionStatus sessionStatus;
+		ovr_GetSessionStatus(OvrSession, &sessionStatus);
+		
+		// Check if the HMD is worn 
+		if (sessionStatus.HmdMounted )
+		{
+			return EHMDWornState::Worn; // already created and worn
+
+		}
+	}
+	return EHMDWornState::NotWorn;
+}
+
 FGameFrame* FOculusRiftHMD::GetFrame()
 {
 	return static_cast<FGameFrame*>(GetCurrentFrame());
diff --git a/Engine/Plugins/Runtime/OculusRift/Source/OculusRift/Private/OculusRiftHMD.h b/Engine/Plugins/Runtime/OculusRift/Source/OculusRift/Private/OculusRiftHMD.h
--- a/Engine/Plugins/Runtime/OculusRift/Source/OculusRift/Private/OculusRiftHMD.h
+++ b/Engine/Plugins/Runtime/OculusRift/Source/OculusRift/Private/OculusRiftHMD.h
@@ -423,6 +423,7 @@ public:
 	virtual bool OnStartGameFrame( FWorldContext& WorldContext ) override;
 
 	virtual bool IsHMDConnected() override;
+	virtual EHMDWornState::Type GetHMDWornState() override;
 	virtual EHMDDeviceType::Type GetHMDDeviceType() const override;
 	virtual bool GetHMDMonitorInfo(MonitorInfo&) override;
 
diff --git a/Engine/Plugins/Runtime/Steam/SteamVR/Source/SteamVR/Private/SteamVRHMD.cpp b/Engine/Plugins/Runtime/Steam/SteamVR/Source/SteamVR/Private/SteamVRHMD.cpp
--- a/Engine/Plugins/Runtime/Steam/SteamVR/Source/SteamVR/Private/SteamVRHMD.cpp
+++ b/Engine/Plugins/Runtime/Steam/SteamVR/Source/SteamVR/Private/SteamVRHMD.cpp
@@ -212,6 +212,7 @@ public:
 		return false;
 	}
 
+
 private:
 	vr::IVRSystem* VRSystem;
     
@@ -252,6 +253,12 @@ bool FSteamVRHMD::IsHMDEnabled() const
 	return bHmdEnabled;
 }
 
+EHMDWornState::Type FSteamVRHMD::GetHMDWornState()
+{
+	//HmdWornState is set in OnStartGameFrame's event loop
+	return HmdWornState;
+}
+
 void FSteamVRHMD::EnableHMD(bool enable)
 {
 	bHmdEnabled = enable;
@@ -796,6 +803,20 @@ bool FSteamVRHMD::OnStartGameFrame(FWorldContext& WorldContext)
 		case vr::VREvent_InputFocusReleased:
 			FCoreDelegates::ApplicationHasEnteredForegroundDelegate.Broadcast();
 			break;
+		case vr::VREvent_TrackedDeviceUserInteractionStarted:
+			// if the event was sent by the HMD
+			if(VREvent.trackedDeviceIndex == vr::k_unTrackedDeviceIndex_Hmd)
+			{
+				HmdWornState = EHMDWornState::Worn;
+			}
+			break;
+		case vr::VREvent_TrackedDeviceUserInteractionEnded:
+			// if the event was sent by the HMD
+			if (VREvent.trackedDeviceIndex == vr::k_unTrackedDeviceIndex_Hmd)
+			{
+				HmdWornState = EHMDWornState::NotWorn;
+			}
+			break;
 		}
 	}
 
@@ -1144,6 +1165,7 @@ bool FSteamVRHMD::NeedReAllocateViewportRenderTarget(const FViewport& Viewport)
 FSteamVRHMD::FSteamVRHMD(ISteamVRPlugin* SteamVRPlugin) :
 	VRSystem(nullptr),
 	bHmdEnabled(true),
+	HmdWornState(EHMDWornState::Unknown),
 	bStereoEnabled(false),
 	bHmdPosTracking(true),
 	bHaveVisionTracking(false),
diff --git a/Engine/Plugins/Runtime/Steam/SteamVR/Source/SteamVR/Private/SteamVRHMD.h b/Engine/Plugins/Runtime/Steam/SteamVR/Source/SteamVR/Private/SteamVRHMD.h
--- a/Engine/Plugins/Runtime/Steam/SteamVR/Source/SteamVR/Private/SteamVRHMD.h
+++ b/Engine/Plugins/Runtime/Steam/SteamVR/Source/SteamVR/Private/SteamVRHMD.h
@@ -42,6 +42,7 @@ public:
 
 	virtual bool IsHMDConnected() override { return true; }
 	virtual bool IsHMDEnabled() const override;
+	virtual EHMDWornState::Type GetHMDWornState() override;
 	virtual void EnableHMD(bool allow = true) override;
 	virtual EHMDDeviceType::Type GetHMDDeviceType() const override;
 	virtual bool GetHMDMonitorInfo(MonitorInfo&) override;
@@ -264,6 +265,7 @@ private:
 	void SetupOcclusionMeshes();
 
 	bool bHmdEnabled;
+	EHMDWornState::Type HmdWornState;
 	bool bStereoEnabled;
 	bool bHmdPosTracking;
 	mutable bool bHaveVisionTracking;
diff --git a/Engine/Source/Editor/UnrealEd/Classes/Settings/EditorExperimentalSettings.h b/Engine/Source/Editor/UnrealEd/Classes/Settings/EditorExperimentalSettings.h
--- a/Engine/Source/Editor/UnrealEd/Classes/Settings/EditorExperimentalSettings.h
+++ b/Engine/Source/Editor/UnrealEd/Classes/Settings/EditorExperimentalSettings.h
@@ -154,6 +154,10 @@ public:
 	UPROPERTY(EditAnywhere, config, Category=VR, meta = (DisplayName="Enable VR Editing"))
 	bool bEnableVREditing;
 
+	/**If true, wearing a Vive or Oculus Rift headset will automatically enter VR Editing mode if Enable VR Editing is true. */
+	UPROPERTY(EditAnywhere, config, Category = VR, meta = (DisplayName = "Enable VR Mode Auto-Entry"))
+	bool bEnableAutoVREditMode;
+
 	/** Allows editing of potentially unsafe properties during PIE. Advanced use only - use with caution. */
 	UPROPERTY(EditAnywhere, config, Category = Tools, meta = (DisplayName = "Allow editing of potentially unsafe properties."))
 	bool bAllowPotentiallyUnsafePropertyEditing;
diff --git a/Engine/Source/Editor/UnrealEd/Private/Settings/SettingsClasses.cpp b/Engine/Source/Editor/UnrealEd/Private/Settings/SettingsClasses.cpp
--- a/Engine/Source/Editor/UnrealEd/Private/Settings/SettingsClasses.cpp
+++ b/Engine/Source/Editor/UnrealEd/Private/Settings/SettingsClasses.cpp
@@ -76,6 +76,7 @@ UEditorExperimentalSettings::UEditorExperimentalSettings( const FObjectInitializ
 	, BlueprintProfilerInclNodeThreshold(0.25f)
 	, BlueprintProfilerMaxNodeThreshold(0.5f)
 	, bUseOpenCLForConvexHullDecomp(false)
+	, bEnableAutoVREditMode(true)
 	, bAllowPotentiallyUnsafePropertyEditing(false)
 {
 }
diff --git a/Engine/Source/Editor/VREditor/VREditorModule.cpp b/Engine/Source/Editor/VREditor/VREditorModule.cpp
--- a/Engine/Source/Editor/VREditor/VREditorModule.cpp
+++ b/Engine/Source/Editor/VREditor/VREditorModule.cpp
@@ -6,6 +6,7 @@
 #include "VREditorMode.h"
 #include "LevelEditor.h"
 #include "HeadMountedDisplay.h"
+#include "Kismet/HeadMountedDisplayFunctionLibrary.h"	// For EHMDWornState::Type
 
 class FVREditorModule : public IVREditorModule, public FTickableEditorObject
 {
@@ -14,7 +15,9 @@ public:
 		bEnableVRRequest( false ),
 		bPlayStartedFromVREditor( false ),
 		LastWorldToMeters( 100 ),
-		SavedWorldToMeters( 0 )
+		SavedWorldToMeters( 0 ),
+		HMDWornState( EHMDWornState::Unknown ),
+		TimeSinceHMDChecked( 0.0f )
 	{
 	}
 
@@ -67,6 +70,14 @@ private:
 
 	/** Saved last world to meters scale from last VR Editor session so we can restore it when entering the VR Editor when PIE starter from VR Editor */
 	float SavedWorldToMeters;
+
+	/** True when we detect that the user is wearing the HMD */
+	EHMDWornState::Type HMDWornState;
+
+	/** Timer for checking if the user is wearing the HMD */
+	float TimeSinceHMDChecked;
+
+
 };
 
 namespace VREd
@@ -174,12 +185,33 @@ void FVREditorModule::ToggleForceVRMode()
 
 void FVREditorModule::Tick( float DeltaTime )
 {
+	//@todo vreditor: Make the timer a configurable variable and/or change the polling system to an event-based one.
+	TimeSinceHMDChecked += DeltaTime;
+	// You can only auto-enter VR if the setting is enabled. Other criteria are that the VR Editor is enabled in experimental settings, that you are not in PIE, and that the editor is foreground.
+	bool bCanAutoEnterVR = (GetDefault<UEditorExperimentalSettings>()->bEnableAutoVREditMode) && (GetDefault<UEditorExperimentalSettings>()->bEnableVREditing) && !(GEditor->PlayWorld) && FPlatformProcess::IsThisApplicationForeground();
+
+	// Only check whether you are wearing the HMD every second, if you are allowed to auto-enter VR, and if your HMD state has changed since the last check. 
+	if((TimeSinceHMDChecked >= 1.0f) && bCanAutoEnterVR && (HMDWornState != GEngine->HMDDevice->GetHMDWornState()))
+	{
+		TimeSinceHMDChecked = 0.0f;
+		HMDWornState = GEngine->HMDDevice->GetHMDWornState();
+		if( HMDWornState == EHMDWornState::Worn )
+		{
+			EnableVREditor(true, false);
+		}
+		else if (HMDWornState == EHMDWornState::NotWorn)
+		{
+			EnableVREditor(false , false);
+		}
+	}
+
 	if( IsVREditorEnabled() )
 	{
-		FVREditorMode* VREditorMode = static_cast<FVREditorMode*>( GLevelEditorModeTools().FindMode( FVREditorMode::GetVREditorModeID() ) );
-		check( VREditorMode != nullptr );
+		FVREditorMode* VREditorMode = static_cast<FVREditorMode*>(GLevelEditorModeTools().FindMode(FVREditorMode::GetVREditorModeID()));
+		check(VREditorMode != nullptr);
 
-		if( VREditorMode->WantsToExitMode() )
+
+		if (VREditorMode->WantsToExitMode())
 		{
 			//Store the current WorldToMeters before exiting the mode
 			LastWorldToMeters = GWorld->GetWorldSettings()->WorldToMeters;
@@ -198,6 +230,7 @@ void FVREditorModule::Tick( float DeltaTime )
 			}
 		}
 	}
+
 	// Only check for input if we started this play session from the VR Editor
 	else if( bPlayStartedFromVREditor && GEditor->PlayWorld && !bEnableVRRequest )
 	{
@@ -226,6 +259,7 @@ void FVREditorModule::Tick( float DeltaTime )
 			EnableVREditor( true, false );
 			bEnableVRRequest = false;
 		}
+		
 	}
 }
 
diff --git a/Engine/Source/Runtime/Engine/Classes/Kismet/HeadMountedDisplayFunctionLibrary.h b/Engine/Source/Runtime/Engine/Classes/Kismet/HeadMountedDisplayFunctionLibrary.h
--- a/Engine/Source/Runtime/Engine/Classes/Kismet/HeadMountedDisplayFunctionLibrary.h
+++ b/Engine/Source/Runtime/Engine/Classes/Kismet/HeadMountedDisplayFunctionLibrary.h
@@ -27,6 +27,20 @@ namespace EHMDTrackingOrigin
 	};
 }
 
+/**
+* Stores if the user is wearing the HMD or not. For HMDs without a sensor to detect the user wearing it, the state defaults to Unknown.
+*/
+UENUM()
+namespace EHMDWornState
+{
+	enum Type
+	{
+		Unknown UMETA(DisplayName = "Unknown"),
+		Worn UMETA(DisplayName = "Worn"),
+		NotWorn UMETA(DisplayName = "Not Worn"),
+	};
+}
+
 UCLASS()
 class ENGINE_API UHeadMountedDisplayFunctionLibrary : public UBlueprintFunctionLibrary
 {
diff --git a/Engine/Source/Runtime/HeadMountedDisplay/Public/IHeadMountedDisplay.h b/Engine/Source/Runtime/HeadMountedDisplay/Public/IHeadMountedDisplay.h
--- a/Engine/Source/Runtime/HeadMountedDisplay/Public/IHeadMountedDisplay.h
+++ b/Engine/Source/Runtime/HeadMountedDisplay/Public/IHeadMountedDisplay.h
@@ -27,6 +27,11 @@ public:
 	 */
 	virtual bool IsHMDEnabled() const = 0;
 
+	/**
+	* Returns EHMDWornState::Worn if we detect that the user is wearing the HMD, EHMDWornState::NotWorn if we detect the user is not wearing the HMD, and EHMDWornState::Unknown if we cannot detect the state.
+	*/
+	virtual EHMDWornState::Type GetHMDWornState() { return EHMDWornState::Unknown; };
+
 	/**
 	 * Enables or disables switching to stereo.
 	 */
