Removed leftover debug code #rb mike.fricker

diff --git a/Engine/Source/Editor/VREditor/VREditorMotionControllerInteractor.cpp b/Engine/Source/Editor/VREditor/VREditorMotionControllerInteractor.cpp
--- a/Engine/Source/Editor/VREditor/VREditorMotionControllerInteractor.cpp
+++ b/Engine/Source/Editor/VREditor/VREditorMotionControllerInteractor.cpp
@@ -390,8 +390,6 @@ void UVREditorMotionControllerInteractor::Tick( const float DeltaTime )
 				HoverMeshComponent->SetVisibility( true );
 				HoverMeshComponent->SetWorldLocation( GetHoverLocation() );
 
-				GLog->Logf( TEXT( " IsHoveringOverUI: %d " ), IsHoveringOverUI() );
-
 				// Show the light too, unless it's on top of UI.  It looks too distracting on top of UI.
 				HoverPointLightComponent->SetVisibility( !IsHoveringOverUI() );
 
