Fixed laser not re-appearing after closing a docked UI panel (UE-33299)

diff --git a/Engine/Source/Editor/VREditor/VREditorUISystem.cpp b/Engine/Source/Editor/VREditor/VREditorUISystem.cpp
--- a/Engine/Source/Editor/VREditor/VREditorUISystem.cpp
+++ b/Engine/Source/Editor/VREditor/VREditorUISystem.cpp
@@ -970,7 +970,8 @@ void UVREditorUISystem::ShowEditorUIPanel( AVREditorFloatingUI* Panel, UVREditor
 
 		if ( Panel->GetDockedTo() == AVREditorFloatingUI::EDockedTo::LeftHand || Panel->GetDockedTo() == AVREditorFloatingUI::EDockedTo::RightHand )
 		{
-			Interactor->SetHasUIInFront( bShouldShow );
+			UVREditorInteractor* PanelDockedToVREditorInteractor = GetOwner().GetHandInteractor( Panel->GetDockedTo() == AVREditorFloatingUI::EDockedTo::LeftHand ? EControllerHand::Left : EControllerHand::Right );
+			PanelDockedToVREditorInteractor->SetHasUIInFront( bShouldShow );
 		}
 
 		if( bPlaySound )
