VREditor: Fixed foliage select/lasso selection highlight not shown/updated

diff --git a/Engine/Source/Editor/VREditor/VREditorMode.cpp b/Engine/Source/Editor/VREditor/VREditorMode.cpp
--- a/Engine/Source/Editor/VREditor/VREditorMode.cpp
+++ b/Engine/Source/Editor/VREditor/VREditorMode.cpp
@@ -290,6 +290,10 @@ void FVREditorMode::Enter()
 
 			// Make the new viewport the active level editing viewport right away
 			GCurrentLevelEditingViewportClient = &VRViewportClient;
+
+			// Enable selection outline right away
+			VREditorViewportClient.EngineShowFlags.SetSelection(true);
+			VREditorViewportClient.EngineShowFlags.SetSelectionOutline(true);
 		}
 
 		VREditorLevelViewport->EnableStereoRendering( bActuallyUsingVR );
