VR Editor: Avoid loading VR Editor classes in the cooker (UE-33318)

diff --git a/Engine/Source/Editor/VREditor/VREditorUISystem.cpp b/Engine/Source/Editor/VREditor/VREditorUISystem.cpp
--- a/Engine/Source/Editor/VREditor/VREditorUISystem.cpp
+++ b/Engine/Source/Editor/VREditor/VREditorUISystem.cpp
@@ -112,12 +112,19 @@ UVREditorUISystem::UVREditorUISystem( const FObjectInitializer& Initializer ) :
 	DefaultWindowTransforms.Add( VREd::DefaultWorldOutlinerTransform );
 	DefaultWindowTransforms.Add( VREd::DefaultActorDetailsTransform );
 	DefaultWindowTransforms.Add( VREd::DefaultModesTransform );
+}
+
+void UVREditorUISystem::Init()
+{
+	// Register to find out about VR events
+	GetOwner().GetWorldInteraction().OnViewportInteractionInputAction().AddUObject( this, &UVREditorUISystem::OnVRAction );
+	GetOwner().GetWorldInteraction().OnViewportInteractionHoverUpdate().AddUObject( this, &UVREditorUISystem::OnVRHoverUpdate );
 
 	QuickMenuWidgetClass = LoadClass<UVREditorQuickMenu>( nullptr, TEXT( "/Engine/VREditor/UI/VRQuickMenu.VRQuickMenu_C" ) );
 	check( QuickMenuWidgetClass != nullptr );
 
-	QuickRadialWidgetClass = LoadClass<UVREditorRadialMenu>(nullptr, TEXT("/Engine/VREditor/UI/VRRadialQuickMenu.VRRadialQuickMenu_C"));
-	check(QuickRadialWidgetClass != nullptr);
+	QuickRadialWidgetClass = LoadClass<UVREditorRadialMenu>( nullptr, TEXT( "/Engine/VREditor/UI/VRRadialQuickMenu.VRRadialQuickMenu_C" ) );
+	check( QuickRadialWidgetClass != nullptr );
 
 	TutorialWidgetClass = LoadClass<UVREditorBaseUserWidget>( nullptr, TEXT( "/Engine/VREditor/Tutorial/UI_VR_Tutorial_00.UI_VR_Tutorial_00_C" ) );
 	check( TutorialWidgetClass != nullptr );
@@ -134,13 +141,7 @@ UVREditorUISystem::UVREditorUISystem( const FObjectInitializer& Initializer ) :
 
 	ShowUISound = LoadObject<USoundCue>( nullptr, TEXT( "/Engine/VREditor/Sounds/VR_open_Cue" ) );
 	check( ShowUISound != nullptr );
-}
 
-void UVREditorUISystem::Init()
-{
-	// Register to find out about VR events
-	GetOwner().GetWorldInteraction().OnViewportInteractionInputAction().AddUObject( this, &UVREditorUISystem::OnVRAction );
-	GetOwner().GetWorldInteraction().OnViewportInteractionHoverUpdate().AddUObject( this, &UVREditorUISystem::OnVRHoverUpdate );
 	
 	// Create all of our UI panels
 	CreateUIs();
