[INTEGRATE] Fixed Linux compiler warning (from CL 3053686 in //UE4/Main)

diff --git a/Engine/Source/Editor/VREditor/VREditorMotionControllerInteractor.cpp b/Engine/Source/Editor/VREditor/VREditorMotionControllerInteractor.cpp
--- a/Engine/Source/Editor/VREditor/VREditorMotionControllerInteractor.cpp
+++ b/Engine/Source/Editor/VREditor/VREditorMotionControllerInteractor.cpp
@@ -422,10 +422,10 @@ void UVREditorMotionControllerInteractor::Tick( const float DeltaTime )
 
 		const EViewportInteractionDraggingMode DraggingMode = GetDraggingMode();
 
-		//@todo VREditor
-		const bool bIsDraggingWorldWithTwoHands = false;
-		( DraggingMode == EViewportInteractionDraggingMode::AssistingDrag && GetOtherInteractor()->GetDraggingMode() == EViewportInteractionDraggingMode::World ) ||
-			( DraggingMode == EViewportInteractionDraggingMode::World && GetOtherInteractor()->GetDraggingMode() == EViewportInteractionDraggingMode::AssistingDrag );
+		const bool bIsDraggingWorldWithTwoHands = 
+			GetOtherInteractor() != nullptr &&
+			( ( DraggingMode == EViewportInteractionDraggingMode::AssistingDrag && GetOtherInteractor()->GetDraggingMode() == EViewportInteractionDraggingMode::World ) ||
+			  ( DraggingMode == EViewportInteractionDraggingMode::World && GetOtherInteractor()->GetDraggingMode() == EViewportInteractionDraggingMode::AssistingDrag ) );
 
 		if ( bIsDraggingWorldWithTwoHands )
 		{
