VR Editor :
- Fixed bug that Laser extends beyond UI when hovered over Selection Bar or Close Button #jira UE-33552
- Fixed crash when Closing Editor (Alt F4 while in VR mode) #jira UE-32509
- Fixed crash when enabling VR Editor in editor preferences without a HMD connected
- Fixed bug if you "slowly press" over UI selection bars or close buttons, nothing happens #jira UE-33553
- Avatar code refactor to its own actor class #jira UETOOL-812

diff --git a/Engine/Source/Editor/VREditor/VREditorAvatarActor.cpp b/Engine/Source/Editor/VREditor/VREditorAvatarActor.cpp
new file mode 100644
--- /dev/null
+++ b/Engine/Source/Editor/VREditor/VREditorAvatarActor.cpp
@@ -0,0 +1,439 @@
+// Copyright 1998-2015 Epic Games, Inc. All Rights Reserved.
+
+#include "VREditorModule.h"
+#include "VREditorAvatarActor.h"
+#include "VREditorMode.h"
+#include "VREditorInteractor.h"
+#include "ViewportWorldInteraction.h"
+#include "IHeadMountedDisplay.h"
+
+namespace VREd
+{
+	static FAutoConsoleVariable GridMovementTolerance( TEXT( "VREd.GridMovementTolerance" ), 0.1f, TEXT( "Tolerance for movement when the grid must disappear" ) );
+	static FAutoConsoleVariable GridScaleMultiplier( TEXT( "VREd.GridScaleMultiplier" ), 35.0f, TEXT( "Scale of the grid" ) );
+	static FAutoConsoleVariable GridFadeMultiplier( TEXT( "VREd.GridFadeMultiplier" ), 3.0f, TEXT( "Grid fade in/out speed, in 'fades per second'" ) );
+	static FAutoConsoleVariable GridFadeStartVelocity( TEXT( "VREd.GridFadeStartVelocity" ), 10.f, TEXT( "Grid fade duration" ) );
+	static FAutoConsoleVariable GridMaxOpacity( TEXT( "VREd.GridMaxFade" ), 0.8f, TEXT( "Grid maximum opacity" ) );
+	static FAutoConsoleVariable GridHeightOffset( TEXT( "VREd.GridHeightOffset" ), 0.0f, TEXT( "Height offset for the world movement grid.  Useful when tracking space is not properly calibrated" ) );
+
+	static FAutoConsoleVariable WorldMovementFogOpacity( TEXT( "VREd.WorldMovementFogOpacity" ), 0.8f, TEXT( "How opaque the fog should be at the 'end distance' (0.0 - 1.0)" ) );
+	static FAutoConsoleVariable WorldMovementFogStartDistance( TEXT( "VREd.WorldMovementFogStartDistance" ), 300.0f, TEXT( "How far away fog will start rendering while in world movement mode" ) );
+	static FAutoConsoleVariable WorldMovementFogEndDistance( TEXT( "VREd.WorldMovementFogEndDistance" ), 800.0f, TEXT( "How far away fog will finish rendering while in world movement mode" ) );
+	static FAutoConsoleVariable WorldMovementFogSkyboxDistance( TEXT( "VREd.WorldMovementFogSkyboxDistance" ), 20000.0f, TEXT( "Anything further than this distance will be completed fogged and not visible" ) );
+
+	static FAutoConsoleVariable ScaleProgressBarLength( TEXT( "VREd.ScaleProgressBarLength" ), 50.0f, TEXT( "Length of the progressbar that appears when scaling" ) );
+	static FAutoConsoleVariable ScaleProgressBarRadius( TEXT( "VREd.ScaleProgressBarRadius" ), 1.0f, TEXT( "Radius of the progressbar that appears when scaling" ) );
+
+	static FAutoConsoleVariable ShowMovementGrid( TEXT( "VREd.ShowMovementGrid" ), 1, TEXT( "Showing the ground movement grid" ) );
+	static FAutoConsoleVariable ShowWorldMovementPostProcess( TEXT( "VREd.ShowWorldMovementPostProcess" ), 1, TEXT( "Showing the movement post processing" ) );
+}
+
+AVREditorAvatarActor::AVREditorAvatarActor( const FObjectInitializer& ObjectInitializer ) :
+	Super( ObjectInitializer ),
+	WorldMovementGridOpacity( 0.0f ),
+	bIsDrawingWorldMovementPostProcess( false ),
+	VRMode( nullptr )
+{
+	// Give us a head mesh
+	{
+		HeadMeshComponent = CreateDefaultSubobject<UStaticMeshComponent>( TEXT( "HeadMeshComponent" ) );
+		AddOwnedComponent( HeadMeshComponent );
+		HeadMeshComponent->SetupAttachment( RootComponent );
+
+		// @todo vreditor: This needs to adapt based on the device you're using
+		UStaticMesh* HeadMesh = LoadObject<UStaticMesh>( nullptr, TEXT( "/Engine/VREditor/Devices/Generic/GenericHMD" ) );
+		check( HeadMesh != nullptr );
+
+		HeadMeshComponent->SetStaticMesh( HeadMesh );
+		HeadMeshComponent->SetMobility( EComponentMobility::Movable );
+		HeadMeshComponent->SetCollisionEnabled( ECollisionEnabled::NoCollision );
+	}
+
+	// World movement grid mesh
+	{
+		WorldMovementGridMeshComponent = CreateDefaultSubobject<UStaticMeshComponent>( TEXT( "WorldMovementGridMeshComponent" ) );
+		AddOwnedComponent( WorldMovementGridMeshComponent );
+		WorldMovementGridMeshComponent->SetupAttachment( RootComponent );
+	
+		UStaticMesh* GridMesh = LoadObject<UStaticMesh>( nullptr, TEXT( "/Engine/VREditor/WorldMovementGrid/PlaneMesh" ) );
+		check( GridMesh != nullptr );
+		WorldMovementGridMeshComponent->SetStaticMesh( GridMesh );
+		WorldMovementGridMeshComponent->SetMobility( EComponentMobility::Movable );
+		WorldMovementGridMeshComponent->SetCollisionEnabled( ECollisionEnabled::NoCollision );
+
+		UMaterialInterface* GridMaterial = LoadObject<UMaterialInterface>( nullptr, TEXT( "/Engine/VREditor/WorldMovementGrid/GridMaterial" ) );
+		check( GridMaterial != nullptr );
+
+		WorldMovementGridMID = UMaterialInstanceDynamic::Create( GridMaterial, GetTransientPackage( ) );
+		check( WorldMovementGridMID != nullptr );
+		WorldMovementGridMeshComponent->SetMaterial( 0, WorldMovementGridMID );
+
+		// The grid starts off hidden
+		WorldMovementGridMeshComponent->SetVisibility( false );
+	}
+
+	{
+		{
+			UMaterialInterface* UserScaleIndicatorMaterial = LoadObject<UMaterialInterface>( nullptr, TEXT( "/Engine/VREditor/LaserPointer/LaserPointerMaterialInst" ) );
+			check( UserScaleIndicatorMaterial != nullptr );
+
+			UMaterialInstance* TranslucentUserScaleIndicatorMaterial = LoadObject<UMaterialInstance>( nullptr, TEXT( "/Engine/VREditor/LaserPointer/TranslucentLaserPointerMaterialInst" ) );
+			check( TranslucentUserScaleIndicatorMaterial != nullptr );
+
+			FixedUserScaleMID = UMaterialInstanceDynamic::Create( UserScaleIndicatorMaterial, GetTransientPackage() );
+			check( FixedUserScaleMID != nullptr );
+
+			TranslucentFixedUserScaleMID = UMaterialInstanceDynamic::Create( TranslucentUserScaleIndicatorMaterial, GetTransientPackage() );
+			check( TranslucentFixedUserScaleMID != nullptr );
+			
+			CurrentUserScaleMID = UMaterialInstanceDynamic::Create( UserScaleIndicatorMaterial, GetTransientPackage() );
+			check( CurrentUserScaleMID != nullptr );
+
+			TranslucentCurrentUserScaleMID = UMaterialInstanceDynamic::Create( TranslucentUserScaleIndicatorMaterial, GetTransientPackage() );
+			check( TranslucentCurrentUserScaleMID != nullptr );
+
+			UStaticMesh* ScaleLineMesh = LoadObject<UStaticMesh>( nullptr, TEXT( "/Engine/VREditor/LaserPointer/LaserPointerMesh" ) );
+			check( ScaleLineMesh != nullptr );
+
+			// Creating the background bar progress of the scale 
+			{
+				ScaleProgressMeshComponent = CreateDefaultSubobject<UStaticMeshComponent>( TEXT( "ScaleProgressMeshComponent" ) );
+				this->AddOwnedComponent( ScaleProgressMeshComponent );
+				ScaleProgressMeshComponent->SetupAttachment( RootComponent );
+
+				ScaleProgressMeshComponent->SetStaticMesh( ScaleLineMesh );
+				ScaleProgressMeshComponent->SetMobility( EComponentMobility::Movable );
+				ScaleProgressMeshComponent->SetCollisionEnabled( ECollisionEnabled::NoCollision );
+				ScaleProgressMeshComponent->SetMaterial( 0, FixedUserScaleMID );
+				ScaleProgressMeshComponent->SetMaterial( 1, TranslucentFixedUserScaleMID );
+
+				// The user scale indicator starts invisible
+				ScaleProgressMeshComponent->SetVisibility( false );
+			}
+
+			// Creating the current progress of the scale 
+			{
+				CurrentScaleProgressMeshComponent = CreateDefaultSubobject<UStaticMeshComponent>( TEXT( "CurrentScaleProgressMeshComponent" ) );
+				AddOwnedComponent( CurrentScaleProgressMeshComponent );
+				CurrentScaleProgressMeshComponent->SetupAttachment( RootComponent );
+
+				CurrentScaleProgressMeshComponent->SetStaticMesh( ScaleLineMesh );
+				CurrentScaleProgressMeshComponent->SetMobility( EComponentMobility::Movable );
+				CurrentScaleProgressMeshComponent->SetCollisionEnabled( ECollisionEnabled::NoCollision );
+				CurrentScaleProgressMeshComponent->SetMaterial( 0, CurrentUserScaleMID );
+				CurrentScaleProgressMeshComponent->SetMaterial( 1, TranslucentCurrentUserScaleMID );
+
+				// The user scale indicator starts invisible
+				CurrentScaleProgressMeshComponent->SetVisibility( false );
+			}
+		}
+
+		// Creating the text for scaling
+		{
+			UFont* TextFont = LoadObject<UFont>( nullptr, TEXT( "/Engine/VREditor/Fonts/VRText_RobotoLarge" ) );
+			check( TextFont != nullptr );
+
+			UMaterialInterface* UserScaleIndicatorMaterial = LoadObject<UMaterialInterface>( nullptr, TEXT( "/Engine/VREditor/Fonts/VRTextMaterial" ) );
+			check( UserScaleIndicatorMaterial != nullptr );
+
+			UserScaleIndicatorText = CreateDefaultSubobject<UTextRenderComponent>( TEXT( "UserScaleIndicatorText" ) );
+			AddOwnedComponent( UserScaleIndicatorText );
+			UserScaleIndicatorText->SetupAttachment( RootComponent );
+
+			UserScaleIndicatorText->SetMobility( EComponentMobility::Movable );
+			UserScaleIndicatorText->SetCollisionEnabled( ECollisionEnabled::NoCollision );
+			UserScaleIndicatorText->SetCollisionProfileName( UCollisionProfile::NoCollision_ProfileName );
+
+			UserScaleIndicatorText->bGenerateOverlapEvents = false;
+			UserScaleIndicatorText->SetCanEverAffectNavigation( false );
+			UserScaleIndicatorText->bCastDynamicShadow = false;
+			UserScaleIndicatorText->bCastStaticShadow = false;
+			UserScaleIndicatorText->bAffectDistanceFieldLighting = false;
+			UserScaleIndicatorText->bAffectDynamicIndirectLighting = false;
+
+			// Use a custom font.  The text will be visible up close.
+			UserScaleIndicatorText->SetFont( TextFont );
+			UserScaleIndicatorText->SetWorldSize( 8.0f );
+			UserScaleIndicatorText->SetTextMaterial( UserScaleIndicatorMaterial );
+
+			// Center the text horizontally
+			UserScaleIndicatorText->SetHorizontalAlignment( EHTA_Center );
+			UserScaleIndicatorText->SetVisibility( false );
+		}
+	}
+
+	// Load our post process material for the world movement grid
+	{
+		UMaterial* Material = LoadObject<UMaterial>( nullptr, TEXT( "/Engine/VREditor/WorldMovementGrid/GridPostProcess" ) );
+		check( Material != nullptr );
+		WorldMovementPostProcessMaterial = UMaterialInstanceDynamic::Create( Material, GetTransientPackage() );
+		check( WorldMovementPostProcessMaterial != nullptr );
+	}
+
+	// Post processing
+	{
+		PostProcessComponent = CreateDefaultSubobject<UPostProcessComponent>( TEXT( "PostProcessComponent" ) );
+		AddOwnedComponent( PostProcessComponent );
+		PostProcessComponent->SetupAttachment( this->GetRootComponent( ) );
+
+		// Unlimited size
+		PostProcessComponent->bEnabled = VREd::ShowWorldMovementPostProcess->GetInt( ) == 1 ? true : false;
+		PostProcessComponent->bUnbound = true;
+	}
+
+}
+
+AVREditorAvatarActor::~AVREditorAvatarActor()
+{
+	WorldMovementGridMID = nullptr;
+	WorldMovementPostProcessMaterial = nullptr;
+	HeadMeshComponent = nullptr;
+	WorldMovementGridMeshComponent = nullptr;
+	PostProcessComponent = nullptr;
+	ScaleProgressMeshComponent = nullptr;
+	CurrentScaleProgressMeshComponent = nullptr;
+	UserScaleIndicatorText = nullptr;
+}
+
+void AVREditorAvatarActor::Init( FVREditorMode* InVRMode  )
+{
+	VRMode = InVRMode;
+
+	// Set the default color for the progress bar
+	{
+		static FName StaticLaserColorName( "LaserColor" );
+		const FLinearColor FixedProgressbarColor = VRMode->GetColor( FVREditorMode::EColors::WorldDraggingColor_TwoHands );
+		FixedUserScaleMID->SetVectorParameterValue( StaticLaserColorName, FixedProgressbarColor );
+		TranslucentFixedUserScaleMID->SetVectorParameterValue( StaticLaserColorName, FixedProgressbarColor );
+
+		const FLinearColor CurrentProgressbarColor = VRMode->GetColor( FVREditorMode::EColors::GreenGizmoColor );
+		CurrentUserScaleMID->SetVectorParameterValue( StaticLaserColorName, CurrentProgressbarColor );
+		TranslucentCurrentUserScaleMID->SetVectorParameterValue( StaticLaserColorName, CurrentProgressbarColor );
+	}
+
+	UserScaleIndicatorText->SetTextRenderColor( VRMode->GetColor( FVREditorMode::EColors::WorldDraggingColor_TwoHands ).ToFColor( false ) );
+
+	// Tell the grid to stay relative to the rootcomponent
+	const FAttachmentTransformRules AttachmentTransformRules = FAttachmentTransformRules( EAttachmentRule::KeepRelative, true );
+	WorldMovementGridMeshComponent->AttachToComponent( RootComponent, AttachmentTransformRules );
+}
+
+void AVREditorAvatarActor::TickManually( const float DeltaTime )
+{
+	SetActorTransform( VRMode->GetRoomTransform() );
+	const float WorldScaleFactor = VRMode->GetWorldScaleFactor();
+	const FVector WorldScaleFactorVector = FVector( WorldScaleFactor );
+
+	// Our head will lock to the HMD position
+	{
+		FTransform RoomSpaceTransformWithWorldToMetersScaling = VRMode->GetRoomSpaceHeadTransform( );
+		RoomSpaceTransformWithWorldToMetersScaling.SetScale3D( RoomSpaceTransformWithWorldToMetersScaling.GetScale3D( ) * WorldScaleFactorVector );
+
+		// @todo vreditor urgent: Head disabled until we can fix late frame update issue
+		if ( false && VRMode->IsActuallyUsingVR() && GEngine->HMDDevice.IsValid( ) && GEngine->HMDDevice->IsHeadTrackingAllowed( ) )
+		{
+			HeadMeshComponent->SetVisibility( true );
+
+			// Apply late frame update to the head mesh
+			HeadMeshComponent->ResetRelativeTransform( );
+			const FTransform ParentToWorld = HeadMeshComponent->GetComponentToWorld( );
+			GEngine->HMDDevice->SetupLateUpdate( ParentToWorld, HeadMeshComponent );
+			HeadMeshComponent->SetRelativeTransform( RoomSpaceTransformWithWorldToMetersScaling );
+		}
+		else
+		{
+			HeadMeshComponent->SetVisibility( false );
+		}
+	}
+
+	// Scale the grid so that it stays the same size in the tracking space
+	WorldMovementGridMeshComponent->SetRelativeScale3D( WorldScaleFactorVector * VREd::GridScaleMultiplier->GetFloat() );
+	WorldMovementGridMeshComponent->SetRelativeLocation( WorldScaleFactorVector * FVector( 0.0f, 0.0f, VREd::GridHeightOffset->GetFloat() ) );
+
+	// Update the user scale indicator //@todo
+	{
+		UVREditorInteractor* LeftHandInteractor = VRMode->GetHandInteractor( EControllerHand::Left );
+		UVREditorInteractor* RightHandInteractor = VRMode->GetHandInteractor( EControllerHand::Right );
+
+		if ( LeftHandInteractor != nullptr && RightHandInteractor != nullptr &&
+			( LeftHandInteractor->GetDraggingMode( ) == EViewportInteractionDraggingMode::World && RightHandInteractor->GetDraggingMode( ) == EViewportInteractionDraggingMode::AssistingDrag ) ||
+			( LeftHandInteractor->GetDraggingMode( ) == EViewportInteractionDraggingMode::AssistingDrag && RightHandInteractor->GetDraggingMode( ) == EViewportInteractionDraggingMode::World ) )
+		{
+			// Setting all components to be visible
+			CurrentScaleProgressMeshComponent->SetVisibility( true );
+			ScaleProgressMeshComponent->SetVisibility( true );
+			UserScaleIndicatorText->SetVisibility( true );
+
+			// Here we calculate the distance, direction and center of the two hands
+			const FVector ScaleIndicatorStartPosition = LeftHandInteractor->GetTransform( ).GetLocation( );
+			const FVector ScaleIndicatorEndPosition = RightHandInteractor->GetTransform( ).GetLocation( );
+			const FVector ScaleIndicatorDirection = ( ScaleIndicatorEndPosition - ScaleIndicatorStartPosition ).GetSafeNormal( );
+
+			const float ScaleIndicatorLength = FMath::Max( 0.000001f, ( ScaleIndicatorEndPosition - ScaleIndicatorStartPosition ).Size( ) );
+			const float Radius = ( VREd::ScaleProgressBarRadius->GetFloat( ) * WorldScaleFactor );
+			const float ProgressbarLength = VREd::ScaleProgressBarLength->GetFloat( );
+			const float Scale = WorldScaleFactor;
+
+			// Add an offset to the center, we don't want the hands to clip through the hand meshes
+			FVector MiddleLocation = ScaleIndicatorEndPosition - ( ScaleIndicatorLength * 0.5f ) * ScaleIndicatorDirection;
+			MiddleLocation.Z += Scale * 5.f;
+
+			// Setting the transform for the scale progressbar
+			{
+				const FVector ProgressbarStart = MiddleLocation - ( ( ProgressbarLength * 0.5f ) * Scale ) * ScaleIndicatorDirection;
+
+				ScaleProgressMeshComponent->SetWorldTransform( FTransform( ScaleIndicatorDirection.ToOrientationRotator( ),
+					ProgressbarStart,
+					FVector( ProgressbarLength * Scale, Radius, Radius )
+					) );
+			}
+
+			// Setting the transform for the current scale progressbar from the center
+			{
+				const float CurrentProgressScale = ( Scale * Scale ) * ( ProgressbarLength / ( VRMode->GetWorldInteraction().GetMaxScale( ) / 100 ) );
+				const FVector CurrentProgressStart = MiddleLocation - ( CurrentProgressScale * 0.5f ) * ScaleIndicatorDirection;
+
+				CurrentScaleProgressMeshComponent->SetWorldTransform( FTransform( ScaleIndicatorDirection.ToOrientationRotator( ),
+					CurrentProgressStart,
+					FVector( CurrentProgressScale, Radius * 2, Radius * 2 )
+					) );
+			}
+
+			//Setting the transform for the scale text
+			{
+				MiddleLocation.Z += Scale * 5.0f;
+				UserScaleIndicatorText->SetWorldTransform( FTransform( ( VRMode->GetHeadTransform().GetLocation( ) - MiddleLocation ).ToOrientationRotator( ),
+					MiddleLocation,
+					VRMode->GetRoomSpaceHeadTransform().GetScale3D()  * Scale
+					) );
+
+				FNumberFormattingOptions NumberFormat;
+				NumberFormat.MinimumIntegralDigits = 1;
+				NumberFormat.MaximumIntegralDigits = 10000;
+				NumberFormat.MinimumFractionalDigits = 1;
+				NumberFormat.MaximumFractionalDigits = 1;
+				UserScaleIndicatorText->SetText( FText::AsNumber( Scale, &NumberFormat ) );
+			}
+		}
+		else
+		{
+			//Setting all components invisible
+			CurrentScaleProgressMeshComponent->SetVisibility( false );
+			ScaleProgressMeshComponent->SetVisibility( false );
+			UserScaleIndicatorText->SetVisibility( false );
+		}
+	}
+
+	// Updating the opacity and visibility of the grid according to the controllers //@todo
+	{
+		if ( VREd::ShowMovementGrid->GetInt( ) == 1 )
+		{
+			UVREditorInteractor* LeftHandInteractor = VRMode->GetHandInteractor( EControllerHand::Left );
+			UVREditorInteractor* RightHandInteractor = VRMode->GetHandInteractor( EControllerHand::Right );
+
+			if ( LeftHandInteractor != nullptr && RightHandInteractor )
+			{
+				// Show the grid full opacity when dragging or scaling
+				float GoalOpacity = 0.f;
+				if ( LeftHandInteractor->GetDraggingMode() == EViewportInteractionDraggingMode::World || RightHandInteractor->GetDraggingMode() == EViewportInteractionDraggingMode::World )
+				{
+					GoalOpacity = 1.0f;
+				}
+				else if ( ( LeftHandInteractor->GetLastDraggingMode() == EViewportInteractionDraggingMode::World && !LeftHandInteractor->GetDragTranslationVelocity().IsNearlyZero( VREd::GridMovementTolerance->GetFloat() ) ) )
+				{
+					GoalOpacity = FMath::Clamp( LeftHandInteractor->GetDragTranslationVelocity().Size() / VREd::GridFadeStartVelocity->GetFloat(), 0.0f, 1.0f );
+				}
+				else if ( ( RightHandInteractor->GetLastDraggingMode() == EViewportInteractionDraggingMode::World && !RightHandInteractor->GetDragTranslationVelocity().IsNearlyZero( VREd::GridMovementTolerance->GetFloat() ) ) )
+				{
+					GoalOpacity = FMath::Clamp( RightHandInteractor->GetDragTranslationVelocity().Size() / VREd::GridFadeStartVelocity->GetFloat(), 0.0f, 1.0f );
+				}
+
+				// Check the current opacity and add or subtract to reach the goal
+				float NewOpacity = WorldMovementGridOpacity;
+				if ( NewOpacity < GoalOpacity )
+				{
+					NewOpacity = FMath::Min( GoalOpacity, NewOpacity + DeltaTime * VREd::GridFadeMultiplier->GetFloat() );
+				}
+				else if ( NewOpacity > GoalOpacity )
+				{
+					NewOpacity = FMath::Max( GoalOpacity, NewOpacity - DeltaTime * VREd::GridFadeMultiplier->GetFloat() );
+				}
+
+				// Only update when the opacity is different
+				if ( !FMath::IsNearlyEqual( WorldMovementGridOpacity, GoalOpacity ) )
+				{
+					WorldMovementGridOpacity = NewOpacity;
+
+					// See if the opacity is almost zero and if the goal opacity is lower than the new opacity set it to zero if so. Otherwise it will flicker
+					if ( WorldMovementGridOpacity < SMALL_NUMBER )
+					{
+						WorldMovementGridOpacity = 0.f;
+						WorldMovementGridMeshComponent->SetVisibility( false );
+					}
+					else
+					{
+						WorldMovementGridMeshComponent->SetVisibility( true );
+					}
+
+					static const FName OpacityName( "OpacityParam" );
+					WorldMovementGridMID->SetScalarParameterValue( OpacityName, WorldMovementGridOpacity * VREd::GridMaxOpacity->GetFloat() );
+				}
+			}
+		}
+		else
+		{
+			WorldMovementGridMeshComponent->SetVisibility( false );
+		}
+	}
+
+	// Apply a post process effect while the user is moving the world around.  The effect "greys out" any pixels
+	// that are not nearby, and completely hides the skybox and other very distant objects.  This is to help
+	// prevent simulation sickness when moving/rotating/scaling the entire world around them.
+	{
+		PostProcessComponent->bEnabled = VREd::ShowWorldMovementPostProcess->GetInt( ) == 1 ? true : false;
+
+		// Make sure our world size is reflected in the post process material
+		static FName WorldScaleFactorParameterName( "WorldScaleFactor" );
+		WorldMovementPostProcessMaterial->SetScalarParameterValue( WorldScaleFactorParameterName, WorldScaleFactor );
+
+		static FName RoomOriginParameterName( "RoomOrigin" );
+		WorldMovementPostProcessMaterial->SetVectorParameterValue( RoomOriginParameterName, VRMode->GetRoomTransform( ).GetLocation( ) );
+
+		static FName FogColorParameterName( "FogColor" );
+		// WorldMovementPostProcessMaterial->SetVectorParameterValue( FogColorParameterName, FLinearColor::Black );
+
+		static FName StartDistanceParameterName( "StartDistance" );
+		WorldMovementPostProcessMaterial->SetScalarParameterValue( StartDistanceParameterName, VREd::WorldMovementFogStartDistance->GetFloat( ) );
+
+		static FName EndDistanceParameterName( "EndDistance" );
+		WorldMovementPostProcessMaterial->SetScalarParameterValue( EndDistanceParameterName, VREd::WorldMovementFogEndDistance->GetFloat( ) );
+
+		static FName FogOpacityParameterName( "FogOpacity" );
+		WorldMovementPostProcessMaterial->SetScalarParameterValue( FogOpacityParameterName, VREd::WorldMovementFogOpacity->GetFloat( ) );
+
+		static FName SkyboxDistanceParameterName( "SkyboxDistance" );
+		WorldMovementPostProcessMaterial->SetScalarParameterValue( SkyboxDistanceParameterName, VREd::WorldMovementFogSkyboxDistance->GetFloat( ) );
+
+		const bool bShouldDrawWorldMovementPostProcess = WorldMovementGridOpacity > KINDA_SMALL_NUMBER;
+		if ( bShouldDrawWorldMovementPostProcess != bIsDrawingWorldMovementPostProcess )
+		{
+			UPostProcessComponent* PostProcess = PostProcessComponent;
+			if ( bShouldDrawWorldMovementPostProcess )
+			{
+				PostProcess->AddOrUpdateBlendable( WorldMovementPostProcessMaterial );
+				bIsDrawingWorldMovementPostProcess = true;
+			}
+			else
+			{
+				bIsDrawingWorldMovementPostProcess = false;
+				PostProcess->Settings.RemoveBlendable( WorldMovementPostProcessMaterial );
+			}
+		}
+
+		if ( bIsDrawingWorldMovementPostProcess )
+		{
+			static FName OpacityParameterName( "Opacity" );
+			WorldMovementPostProcessMaterial->SetScalarParameterValue( OpacityParameterName, FMath::Clamp( WorldMovementGridOpacity, 0.0f, 1.0f ) );
+		}
+	}
+}
diff --git a/Engine/Source/Editor/VREditor/VREditorAvatarActor.h b/Engine/Source/Editor/VREditor/VREditorAvatarActor.h
new file mode 100644
--- /dev/null
+++ b/Engine/Source/Editor/VREditor/VREditorAvatarActor.h
@@ -0,0 +1,98 @@
+// Copyright 1998-2015 Epic Games, Inc. All Rights Reserved.
+
+#pragma once
+
+#include "GameFramework/Actor.h"
+#include "VREditorAvatarActor.generated.h"
+
+/**
+ * Avatar Actor
+ */
+UCLASS()
+class AVREditorAvatarActor : public AActor
+{
+	GENERATED_UCLASS_BODY()
+
+public:
+
+	~AVREditorAvatarActor();
+
+	/** Called by VREditorMode::Enter to initialize all post constructor components and to set the VRMode */
+	void Init( class FVREditorMode* InVRMode );
+
+	/** Called by VREditorMode to update us every frame */
+	void TickManually( const float DeltaTime );
+
+private:
+
+	/** Our avatar's head mesh */
+	UPROPERTY()
+	class UStaticMeshComponent* HeadMeshComponent;
+
+	/** The grid that appears while the user is dragging the world around */
+	UPROPERTY()
+	class UStaticMeshComponent* WorldMovementGridMeshComponent;
+
+	//
+	// World movement grid & FX
+	//
+
+	/** Grid mesh component dynamic material instance to set the opacity */
+	UPROPERTY()
+	class UMaterialInstanceDynamic* WorldMovementGridMID;
+
+	/** Opacity of the movement grid and post process */
+	UPROPERTY()
+	float WorldMovementGridOpacity;
+
+	/** True if we're currently drawing our world movement post process */
+	UPROPERTY()
+	bool bIsDrawingWorldMovementPostProcess;
+
+	/** Post process material for "greying out" the world while in world movement mode */
+	UPROPERTY()
+	class UMaterialInstanceDynamic* WorldMovementPostProcessMaterial;
+
+	//
+	// World scaling progress bar
+	//
+
+	/** Background progressbar scaling mesh */
+	UPROPERTY()
+	class UStaticMeshComponent* ScaleProgressMeshComponent;
+
+	/** Current scale progressbar mesh */
+	UPROPERTY()
+	class UStaticMeshComponent* CurrentScaleProgressMeshComponent;
+
+	/** Current scale text */
+	UPROPERTY()
+	class UTextRenderComponent* UserScaleIndicatorText;
+
+	/** Base dynamic material for the user scale fixed progressbar */
+	UPROPERTY()
+	UMaterialInstanceDynamic* FixedUserScaleMID;
+
+	/** Translucent dynamic material for the user scale fixed progressbar */
+	UPROPERTY()
+	UMaterialInstanceDynamic* TranslucentFixedUserScaleMID;
+	
+	/** Base dynamic material for the current user scale progressbar */
+	UPROPERTY()
+	UMaterialInstanceDynamic* CurrentUserScaleMID;
+
+	/** Translucent dynamic material for the current user scale progressbar */
+	UPROPERTY()
+	UMaterialInstanceDynamic* TranslucentCurrentUserScaleMID;
+
+	//
+	// Post process
+	//
+
+	/** Post process for drawing VR-specific post effects */
+	UPROPERTY()
+	class UPostProcessComponent* PostProcessComponent;
+
+	/** Owning object */
+	class FVREditorMode* VRMode;
+};
\ No newline at end of file
diff --git a/Engine/Source/Editor/VREditor/VREditorInteractor.cpp b/Engine/Source/Editor/VREditor/VREditorInteractor.cpp
--- a/Engine/Source/Editor/VREditor/VREditorInteractor.cpp
+++ b/Engine/Source/Editor/VREditor/VREditorInteractor.cpp
@@ -90,6 +90,7 @@ FHitResult UVREditorInteractor::GetHitResultFromLaserPointer( TArray<AActor*>* O
 	}
 
 	ObjectsInFrontOfGizmo->Add( AVREditorButton::StaticClass() );
+	ObjectsInFrontOfGizmo->Add( AVREditorFloatingUI::StaticClass() );
 	ObjectsInFrontOfGizmo->Add( AVREditorDockableWindow::StaticClass() );
 
 	return UViewportInteractor::GetHitResultFromLaserPointer( OptionalListOfIgnoredActors, bIgnoreGizmos, ObjectsInFrontOfGizmo, bEvenIfBlocked, LaserLengthOverride );
diff --git a/Engine/Source/Editor/VREditor/VREditorMode.cpp b/Engine/Source/Editor/VREditor/VREditorMode.cpp
--- a/Engine/Source/Editor/VREditor/VREditorMode.cpp
+++ b/Engine/Source/Editor/VREditor/VREditorMode.cpp
@@ -7,9 +7,9 @@
 #include "VREditorTransformGizmo.h"
 #include "VREditorFloatingText.h"
 #include "VREditorFloatingUI.h"
+#include "VREditorAvatarActor.h"
 #include "Teleporter/VREditorTeleporter.h"
 
-
 #include "CameraController.h"
 #include "DynamicMeshBuilder.h"
 #include "Features/IModularFeatures.h"
@@ -35,26 +35,7 @@
 namespace VREd
 {
 	static FAutoConsoleVariable UseMouseAsHandInForcedVRMode( TEXT( "VREd.UseMouseAsHandInForcedVRMode" ), 1, TEXT( "When in forced VR mode, enabling this setting uses the mouse cursor as a virtual hand instead of motion controllers" ) );
-
-	static FAutoConsoleVariable GridMovementTolerance( TEXT( "VREd.GridMovementTolerance" ), 0.1f, TEXT( "Tolerance for movement when the grid must disappear" ) );
-	static FAutoConsoleVariable GridScaleMultiplier( TEXT( "VREd.GridScaleMultiplier" ), 35.0f, TEXT( "Scale of the grid" ) );
-	static FAutoConsoleVariable GridFadeMultiplier( TEXT( "VREd.GridFadeMultiplier" ), 3.0f, TEXT( "Grid fade in/out speed, in 'fades per second'" ) );
-	static FAutoConsoleVariable GridFadeStartVelocity( TEXT( "VREd.GridFadeStartVelocity" ), 10.f, TEXT( "Grid fade duration" ) );
-	static FAutoConsoleVariable GridMaxOpacity( TEXT( "VREd.GridMaxFade" ), 0.8f, TEXT( "Grid maximum opacity" ) );
-	static FAutoConsoleVariable GridHeightOffset( TEXT( "VREd.GridHeightOffset" ), 0.0f, TEXT( "Height offset for the world movement grid.  Useful when tracking space is not properly calibrated" ) );
-
-	static FAutoConsoleVariable WorldMovementFogOpacity( TEXT( "VREd.WorldMovementFogOpacity" ), 0.8f, TEXT( "How opaque the fog should be at the 'end distance' (0.0 - 1.0)" ) );
-	static FAutoConsoleVariable WorldMovementFogStartDistance( TEXT( "VREd.WorldMovementFogStartDistance" ), 300.0f, TEXT( "How far away fog will start rendering while in world movement mode" ) );
-	static FAutoConsoleVariable WorldMovementFogEndDistance( TEXT( "VREd.WorldMovementFogEndDistance" ), 800.0f, TEXT( "How far away fog will finish rendering while in world movement mode" ) );
-	static FAutoConsoleVariable WorldMovementFogSkyboxDistance( TEXT( "VREd.WorldMovementFogSkyboxDistance" ), 20000.0f, TEXT( "Anything further than this distance will be completed fogged and not visible" ) );
-
-	static FAutoConsoleVariable ScaleProgressBarLength( TEXT( "VREd.ScaleProgressBarLength" ), 50.0f, TEXT( "Length of the progressbar that appears when scaling" ) );
-	static FAutoConsoleVariable ScaleProgressBarRadius( TEXT( "VREd.ScaleProgressBarRadius" ), 1.0f, TEXT( "Radius of the progressbar that appears when scaling" ) );
-
 	static FAutoConsoleVariable ForceOculusMirrorMode( TEXT( "VREd.ForceOculusMirrorMode" ), 3, TEXT( "Which Oculus display mirroring mode to use (see MirrorWindowModeType in OculusRiftHMD.h)" ) );
-
-	static FAutoConsoleVariable ShowMovementGrid( TEXT( "VREd.ShowMovementGrid" ), 1, TEXT( "Showing the ground movement grid" ) );
-	static FAutoConsoleVariable ShowWorldMovementPostProcess( TEXT( "VREd.ShowWorldMovementPostProcess" ), 1, TEXT( "Showing the movement post processing" ) );
 }
 
 FEditorModeID FVREditorMode::VREditorModeID( "VREditor" );
@@ -66,19 +47,9 @@ FVREditorMode::FVREditorMode()
 	: bWantsToExitMode( false ),
 	  bIsFullyInitialized( false ),
 	  AppTimeModeEntered( FTimespan::Zero() ),
-	  AvatarMeshActor( nullptr ),
-	  HeadMeshComponent( nullptr ),
+	  AvatarActor( nullptr ),
       FlashlightComponent( nullptr ),
 	  bIsFlashlightOn( false ),
-	  WorldMovementGridMeshComponent( nullptr ),
-	  WorldMovementGridMID( nullptr ),
-	  WorldMovementGridOpacity( 0.0f ),	// NOTE: Intentionally not quite zero so that we update the MIDs on the first frame
-	  bIsDrawingWorldMovementPostProcess( false ),
-	  WorldMovementPostProcessMaterial( nullptr ),
-	  ScaleProgressMeshComponent( nullptr ),
-	  CurrentScaleProgressMeshComponent( nullptr ),
-	  UserScaleIndicatorText( nullptr ),
-	  PostProcessComponent( nullptr ),
 	  MotionControllerID( 0 ),	// @todo vreditor minor: We only support a single controller, and we assume the first controller are the motion controls
 	  UISystem( nullptr ),
 	  TeleporterSystem( nullptr ),
@@ -114,9 +85,6 @@ void FVREditorMode::Enter()
 
 	AppTimeModeEntered = FTimespan::FromSeconds( FApp::GetCurrentTime() );
 
-	WorldMovementGridOpacity = 0.0f;
-	bIsDrawingWorldMovementPostProcess = false;
-
 	// Take note of VREditor activation
 	if( FEngineAnalytics::IsAvailable() )
 	{
@@ -233,7 +201,7 @@ void FVREditorMode::Enter()
 			}
 		}
 		this->VREditorLevelViewportWeakPtr = VREditorLevelViewport;
-
+		
 		{
 			FLevelEditorViewportClient& VRViewportClient = VREditorLevelViewport->GetLevelViewportClient();
 			FEditorViewportClient& VREditorViewportClient = VRViewportClient;
@@ -321,20 +289,12 @@ void FVREditorMode::Enter()
 		}
 	}
 
-	// Load our post process material for the world movement grid
-	{
-		UMaterial* Material = LoadObject<UMaterial>( nullptr, TEXT( "/Engine/VREditor/WorldMovementGrid/GridPostProcess" ) );
-		check( Material != nullptr );
-		WorldMovementPostProcessMaterial = UMaterialInstanceDynamic::Create( Material, GetTransientPackage() );
-		check( WorldMovementPostProcessMaterial != nullptr );
-	}
-
 	// Setup sub systems
 	{
 		// Setup world interaction
 		WorldInteraction = NewObject<UVREditorWorldInteraction>();
 		WorldInteraction->SetOwner( this );
-		WorldInteraction->Init( VREditorLevelViewportWeakPtr.Pin()->GetViewportClient().Get() );
+		WorldInteraction->Init( VREditorLevelViewportWeakPtr.Pin()->GetViewportClient() );
 
 		// Motion controllers
 		{
@@ -384,26 +344,12 @@ void FVREditorMode::Exit()
 
 	FSlateApplication::Get().SetInputPreProcessor(false);
 
-	if( WorldMovementPostProcessMaterial != nullptr )
+	//Destroy the avatar
 	{
-		WorldMovementPostProcessMaterial->MarkPendingKill();
-		WorldMovementPostProcessMaterial = nullptr;
-	}
+		DestroyTransientActor( AvatarActor );
+		AvatarActor = nullptr;
 
-	check( WorldInteraction->GetViewportClient() != nullptr );
-
-	{
-		DestroyTransientActor( AvatarMeshActor );
-		AvatarMeshActor = nullptr;
-		
-		HeadMeshComponent = nullptr;
 		FlashlightComponent = nullptr;
-		WorldMovementGridMeshComponent = nullptr;
-		WorldMovementGridMID = nullptr;
-		PostProcessComponent = nullptr;
-		ScaleProgressMeshComponent = nullptr;
-		CurrentScaleProgressMeshComponent = nullptr;
-		UserScaleIndicatorText = nullptr;
 	}
 
 	{
@@ -510,185 +456,21 @@ void FVREditorMode::StartExitingVRMode( const EVREditorExitType InExitType /*= E
 void FVREditorMode::SpawnAvatarMeshActor()
 {
 	// Setup our avatar
-	if( AvatarMeshActor == nullptr )
+	if( AvatarActor == nullptr )
 	{
 		{
 			const bool bWithSceneComponent = true;
-			AvatarMeshActor = SpawnTransientSceneActor( AActor::StaticClass(), TEXT( "AvatarMesh" ), bWithSceneComponent );
-		}
-
-		// Give us a head mesh
-		{
-			HeadMeshComponent = NewObject<UStaticMeshComponent>( AvatarMeshActor );
-			AvatarMeshActor->AddOwnedComponent( HeadMeshComponent );
-			HeadMeshComponent->SetupAttachment( AvatarMeshActor->GetRootComponent() );
-			HeadMeshComponent->RegisterComponent();
-
-			// @todo vreditor: This needs to adapt based on the device you're using
-			UStaticMesh* HeadMesh = LoadObject<UStaticMesh>( nullptr, TEXT( "/Engine/VREditor/Devices/Generic/GenericHMD" ) );
-			check( HeadMesh != nullptr );
-
-			HeadMeshComponent->SetStaticMesh( HeadMesh );
-			HeadMeshComponent->SetMobility( EComponentMobility::Movable );
-			HeadMeshComponent->SetCollisionEnabled( ECollisionEnabled::NoCollision );
+			AvatarActor = SpawnTransientSceneActor<AVREditorAvatarActor>( TEXT( "AvatarActor" ), bWithSceneComponent );
+			AvatarActor->Init( this );
 		}
 
 		//@todo VREditor: Hardcoded interactors
-		LeftHandInteractor->SetupComponent( AvatarMeshActor );
-		RightHandInteractor->SetupComponent( AvatarMeshActor );
-		
-		// World movement grid mesh
-		{
-			WorldMovementGridMeshComponent = NewObject<UStaticMeshComponent>( AvatarMeshActor );
-			AvatarMeshActor->AddOwnedComponent( WorldMovementGridMeshComponent );
-			WorldMovementGridMeshComponent->SetupAttachment( AvatarMeshActor->GetRootComponent() );
-			WorldMovementGridMeshComponent->RegisterComponent();
-
-			UStaticMesh* GridMesh = LoadObject<UStaticMesh>(nullptr, TEXT( "/Engine/VREditor/WorldMovementGrid/PlaneMesh" ) );
-			check( GridMesh != nullptr );
-			WorldMovementGridMeshComponent->SetStaticMesh( GridMesh );
-			WorldMovementGridMeshComponent->SetMobility( EComponentMobility::Movable );
-			WorldMovementGridMeshComponent->SetCollisionEnabled( ECollisionEnabled::NoCollision );
-
-			UMaterialInterface* GridMaterial = LoadObject<UMaterialInterface>( nullptr, TEXT( "/Engine/VREditor/WorldMovementGrid/GridMaterial" ) );
-			check( GridMaterial != nullptr );
-			
-			WorldMovementGridMID = UMaterialInstanceDynamic::Create( GridMaterial, GetTransientPackage() );
-			check( WorldMovementGridMID != nullptr );
-			WorldMovementGridMeshComponent->SetMaterial( 0, WorldMovementGridMID );
-
-			// The grid starts off hidden
-			WorldMovementGridMeshComponent->SetVisibility( false );
-		}
-
-		{
-			{
-				UMaterialInterface* UserScaleIndicatorMaterial = LoadObject<UMaterialInterface>( nullptr, TEXT( "/Engine/VREditor/LaserPointer/LaserPointerMaterialInst" ) );
-				check( UserScaleIndicatorMaterial != nullptr );
-
-				UMaterialInstance* TranslucentUserScaleIndicatorMaterial = LoadObject<UMaterialInstance>( nullptr, TEXT( "/Engine/VREditor/LaserPointer/TranslucentLaserPointerMaterialInst" ) );
-				check( TranslucentUserScaleIndicatorMaterial != nullptr );
-
-				UStaticMesh* ScaleLineMesh = LoadObject<UStaticMesh>( nullptr, TEXT( "/Engine/VREditor/LaserPointer/LaserPointerMesh" ) );
-				check( ScaleLineMesh != nullptr );
-
-				// Creating the background bar progress of the scale 
-				{
-					ScaleProgressMeshComponent = NewObject<UStaticMeshComponent>( AvatarMeshActor );
-					AvatarMeshActor->AddOwnedComponent( ScaleProgressMeshComponent );
-					ScaleProgressMeshComponent->SetupAttachment( AvatarMeshActor->GetRootComponent() );
-					ScaleProgressMeshComponent->RegisterComponent();
-				
-					ScaleProgressMeshComponent->SetStaticMesh( ScaleLineMesh );
-					ScaleProgressMeshComponent->SetMobility( EComponentMobility::Movable );
-					ScaleProgressMeshComponent->SetCollisionEnabled( ECollisionEnabled::NoCollision );
-
-					UMaterialInstanceDynamic* UserScaleMID = UMaterialInstanceDynamic::Create( UserScaleIndicatorMaterial, GetTransientPackage() );
-					check( UserScaleMID != nullptr );
-					ScaleProgressMeshComponent->SetMaterial( 0, UserScaleMID );
-
-					UMaterialInstanceDynamic* TranslucentUserScaleMID = UMaterialInstanceDynamic::Create( TranslucentUserScaleIndicatorMaterial, GetTransientPackage() );
-					check( TranslucentUserScaleMID != nullptr );
-					ScaleProgressMeshComponent->SetMaterial( 1, TranslucentUserScaleMID );
-
-					static FName StaticLaserColorName( "LaserColor" );
-					const FLinearColor Color = GetColor( EColors::WorldDraggingColor_TwoHands );
-					UserScaleMID->SetVectorParameterValue( StaticLaserColorName, Color );
-					TranslucentUserScaleMID->SetVectorParameterValue( StaticLaserColorName, Color );
-
-					// The user scale indicator starts invisible
-					ScaleProgressMeshComponent->SetVisibility( false );
-				}
-
-
-				// Creating the current progress of the scale 
-				{
-					CurrentScaleProgressMeshComponent = NewObject<UStaticMeshComponent>( AvatarMeshActor );
-					AvatarMeshActor->AddOwnedComponent( CurrentScaleProgressMeshComponent );
-					CurrentScaleProgressMeshComponent->SetupAttachment( AvatarMeshActor->GetRootComponent() );
-					CurrentScaleProgressMeshComponent->RegisterComponent();
-
-					CurrentScaleProgressMeshComponent->SetStaticMesh( ScaleLineMesh );
-					CurrentScaleProgressMeshComponent->SetMobility( EComponentMobility::Movable );
-					CurrentScaleProgressMeshComponent->SetCollisionEnabled( ECollisionEnabled::NoCollision );
-
-					UMaterialInstanceDynamic* UserScaleMID = UMaterialInstanceDynamic::Create( UserScaleIndicatorMaterial, GetTransientPackage() );
-					check( UserScaleMID != nullptr );
-					CurrentScaleProgressMeshComponent->SetMaterial( 0, UserScaleMID );
-
-					UMaterialInstanceDynamic* TranslucentUserScaleMID = UMaterialInstanceDynamic::Create( TranslucentUserScaleIndicatorMaterial, GetTransientPackage() );
-					check( TranslucentUserScaleMID != nullptr );
-					CurrentScaleProgressMeshComponent->SetMaterial( 1, TranslucentUserScaleMID );
-
-					static FName StaticLaserColorName( "LaserColor" );
-					const FLinearColor Color = GetColor( EColors::GreenGizmoColor );
-					UserScaleMID->SetVectorParameterValue( StaticLaserColorName, Color );
-					TranslucentUserScaleMID->SetVectorParameterValue( StaticLaserColorName, Color );
-
-					// The user scale indicator starts invisible
-					CurrentScaleProgressMeshComponent->SetVisibility( false );
-				}
-			}
-
-			// Creating the text for scaling
-			{
-				UFont* TextFont = LoadObject<UFont>( nullptr, TEXT( "/Engine/VREditor/Fonts/VRText_RobotoLarge" ) );
-				check( TextFont != nullptr );
-
-				UMaterialInterface* UserScaleIndicatorMaterial = LoadObject<UMaterialInterface>( nullptr, TEXT( "/Engine/VREditor/Fonts/VRTextMaterial" ) );
-				check( UserScaleIndicatorMaterial != nullptr );
-
-				UserScaleIndicatorText = NewObject<UTextRenderComponent>( AvatarMeshActor );
-				AvatarMeshActor->AddOwnedComponent( UserScaleIndicatorText );
-				UserScaleIndicatorText->SetupAttachment( AvatarMeshActor->GetRootComponent() );
-				UserScaleIndicatorText->RegisterComponent();
-
-				UserScaleIndicatorText->SetMobility( EComponentMobility::Movable );
-				UserScaleIndicatorText->SetCollisionEnabled( ECollisionEnabled::NoCollision );
-				UserScaleIndicatorText->SetCollisionProfileName( UCollisionProfile::NoCollision_ProfileName );
-
-				UserScaleIndicatorText->bGenerateOverlapEvents = false;
-				UserScaleIndicatorText->SetCanEverAffectNavigation( false );
-				UserScaleIndicatorText->bCastDynamicShadow = false;
-				UserScaleIndicatorText->bCastStaticShadow = false;
-				UserScaleIndicatorText->bAffectDistanceFieldLighting = false;
-				UserScaleIndicatorText->bAffectDynamicIndirectLighting = false;
-
-				// Use a custom font.  The text will be visible up close.
-				UserScaleIndicatorText->SetFont( TextFont );
-				UserScaleIndicatorText->SetWorldSize( 8.0f );
-				UserScaleIndicatorText->SetTextMaterial( UserScaleIndicatorMaterial );
-				UserScaleIndicatorText->SetTextRenderColor( GetColor( EColors::WorldDraggingColor_TwoHands ).ToFColor( false ) );
-
-				// Center the text horizontally
-				UserScaleIndicatorText->SetHorizontalAlignment( EHTA_Center );
-				UserScaleIndicatorText->SetVisibility( false );
-			}
-		}
-
-		// Post processing
-		{
-			PostProcessComponent = NewObject<UPostProcessComponent>( AvatarMeshActor );
-			AvatarMeshActor->AddOwnedComponent( PostProcessComponent );
-			PostProcessComponent->SetupAttachment( AvatarMeshActor->GetRootComponent() );
-			PostProcessComponent->RegisterComponent();
-
-			// Unlimited size
-			PostProcessComponent->bEnabled = VREd::ShowWorldMovementPostProcess->GetInt() == 1 ? true : false;
-			PostProcessComponent->bUnbound = true;
-		}
+		LeftHandInteractor->SetupComponent( AvatarActor );
+		RightHandInteractor->SetupComponent( AvatarActor );
 	}
 }
 
 
-UPostProcessComponent* FVREditorMode::GetPostProcessComponent()
-{
-	SpawnAvatarMeshActor();
-	check( PostProcessComponent != nullptr );
-	return PostProcessComponent;
-}
-
-
 void FVREditorMode::OnVREditorWindowClosed( const TSharedRef<SWindow>& ClosedWindow )
 {
 	StartExitingVRMode();
@@ -727,7 +509,7 @@ void FVREditorMode::Tick( FEditorViewportClient* ViewportClient, float DeltaTime
 		SetRoomTransform( ResultToWorld );
 	}
 
-	if ( AvatarMeshActor == nullptr )
+	if ( AvatarActor == nullptr )
 	{
 		SpawnAvatarMeshActor();
 	}
@@ -741,106 +523,8 @@ void FVREditorMode::Tick( FEditorViewportClient* ViewportClient, float DeltaTime
 	// Update avatar meshes
 	{
 		// Move our avatar mesh along with the room.  We need our hand components to remain the same coordinate space as the 
-		{
-			const FTransform RoomTransform = GetRoomTransform();
-			AvatarMeshActor->SetActorTransform( RoomTransform );
-		}
-
-		// Our head will lock to the HMD position
-		{
-			FTransform RoomSpaceTransformWithWorldToMetersScaling = GetRoomSpaceHeadTransform();
-			RoomSpaceTransformWithWorldToMetersScaling.SetScale3D( RoomSpaceTransformWithWorldToMetersScaling.GetScale3D() * FVector( GetWorldScaleFactor() ) );
-
-			// @todo vreditor urgent: Head disabled until we can fix late frame update issue
-			if( false && bActuallyUsingVR && GEngine->HMDDevice.IsValid() && GEngine->HMDDevice->IsHeadTrackingAllowed() )
-			{
-				HeadMeshComponent->SetVisibility( true );
-
-				// Apply late frame update to the head mesh
-				HeadMeshComponent->ResetRelativeTransform();
-				const FTransform ParentToWorld = HeadMeshComponent->GetComponentToWorld();
-				GEngine->HMDDevice->SetupLateUpdate( ParentToWorld, HeadMeshComponent );
-				HeadMeshComponent->SetRelativeTransform( RoomSpaceTransformWithWorldToMetersScaling );
-			}
-			else
-			{
-				HeadMeshComponent->SetVisibility( false );
-			}
-		}
-
-		// Scale the grid so that it stays the same size in the tracking space
-		WorldMovementGridMeshComponent->SetRelativeScale3D( FVector( GetWorldScaleFactor() ) * VREd::GridScaleMultiplier->GetFloat() );
-		WorldMovementGridMeshComponent->SetRelativeLocation( FVector( GetWorldScaleFactor() ) * FVector( 0.0f, 0.0f, VREd::GridHeightOffset->GetFloat() ) );
-
-		// Update the user scale indicator //@todo
-		{
-			if ( ( LeftHandInteractor->GetDraggingMode() == EViewportInteractionDraggingMode::World && RightHandInteractor->GetDraggingMode() == EViewportInteractionDraggingMode::AssistingDrag ) ||
-				 ( LeftHandInteractor->GetDraggingMode() == EViewportInteractionDraggingMode::AssistingDrag && RightHandInteractor->GetDraggingMode() == EViewportInteractionDraggingMode::World ) )
-			{
-				// Setting all components to be visible
-				CurrentScaleProgressMeshComponent->SetVisibility( true );
-				ScaleProgressMeshComponent->SetVisibility( true );
-				UserScaleIndicatorText->SetVisibility( true );
-
-				// Here we calculate the distance, direction and center of the two hands
-				const FVector ScaleIndicatorStartPosition = LeftHandInteractor->GetTransform().GetLocation();
-				const FVector ScaleIndicatorEndPosition = RightHandInteractor->GetTransform().GetLocation();
-				const FVector ScaleIndicatorDirection = (ScaleIndicatorEndPosition - ScaleIndicatorStartPosition).GetSafeNormal();
-
-				const float ScaleIndicatorLength = FMath::Max( 0.000001f, (ScaleIndicatorEndPosition - ScaleIndicatorStartPosition).Size() );
-				const float Radius = ( VREd::ScaleProgressBarRadius->GetFloat() * GetWorldScaleFactor() );
-				const float ProgressbarLength = VREd::ScaleProgressBarLength->GetFloat();
-				const float Scale = GetWorldScaleFactor();
-				 
-				// Add an offset to the center, we don't want the hands to clip through the hand meshes
-				FVector MiddleLocation = ScaleIndicatorEndPosition - ( ScaleIndicatorLength * 0.5f ) * ScaleIndicatorDirection;
-				MiddleLocation.Z += Scale * 5.f;
-				
-				// Setting the transform for the scale progressbar
-				{
-					const FVector ProgressbarStart = MiddleLocation - ((ProgressbarLength * 0.5f) * Scale) * ScaleIndicatorDirection;
-
-					ScaleProgressMeshComponent->SetWorldTransform( FTransform( ScaleIndicatorDirection.ToOrientationRotator(),
-						ProgressbarStart,
-						 FVector( ProgressbarLength * Scale, Radius, Radius )
-						) );
-				}
-
-				// Setting the transform for the current scale progressbar from the center
-				{
-					const float CurrentProgressScale = (Scale * Scale) * (ProgressbarLength / ( WorldInteraction->GetMaxScale() / 100));
-					const FVector CurrentProgressStart = MiddleLocation - (CurrentProgressScale * 0.5f) * ScaleIndicatorDirection;
-					
-					CurrentScaleProgressMeshComponent->SetWorldTransform( FTransform (  ScaleIndicatorDirection.ToOrientationRotator(),
-							CurrentProgressStart,
-							FVector( CurrentProgressScale, Radius * 2, Radius * 2 )
-						) );
-				}
-
-				//Setting the transform for the scale text
-				{
-					MiddleLocation.Z += Scale * 5.0f;
-					UserScaleIndicatorText->SetWorldTransform( FTransform( ( GetHeadTransform().GetLocation() - MiddleLocation ).ToOrientationRotator(),
-						MiddleLocation,
-						GetRoomSpaceHeadTransform().GetScale3D()  * Scale
-						) );
-					
-					FNumberFormattingOptions NumberFormat;
-					NumberFormat.MinimumIntegralDigits = 1;
-					NumberFormat.MaximumIntegralDigits = 10000;
-					NumberFormat.MinimumFractionalDigits = 1;
-					NumberFormat.MaximumFractionalDigits = 1;
-					UserScaleIndicatorText->SetText( FText::AsNumber( Scale, &NumberFormat ) );
-				}
-			}
-			else
-			{
-				//Setting all components invisible
-				CurrentScaleProgressMeshComponent->SetVisibility( false );
-				ScaleProgressMeshComponent->SetVisibility( false );
-				UserScaleIndicatorText->SetVisibility( false );
-			}
-		}
+		AvatarActor->SetActorTransform( GetRoomTransform() );
+		AvatarActor->TickManually( DeltaTime );
 	}
 
 	// Updating the scale and intensity of the flashlight according to the world scale
@@ -852,113 +536,6 @@ void FVREditorMode::Tick( FEditorViewportClient* ViewportClient, float DeltaTime
 		FlashlightComponent->SetLightFalloffExponent(UpdatedFalloffExponent);
 	}
 
-	// Updating the opacity and visibility of the grid according to the controllers //@todo
-	{
-		if( VREd::ShowMovementGrid->GetInt() == 1)
-		{
-			// Show the grid full opacity when dragging or scaling
-			float GoalOpacity = 0.f;
-			if ( LeftHandInteractor->GetDraggingMode() == EViewportInteractionDraggingMode::World || RightHandInteractor->GetDraggingMode() == EViewportInteractionDraggingMode::World )
-			{
-				GoalOpacity = 1.0f;
-			}
-			else if ( ( LeftHandInteractor->GetLastDraggingMode() == EViewportInteractionDraggingMode::World && !LeftHandInteractor->GetDragTranslationVelocity().IsNearlyZero( VREd::GridMovementTolerance->GetFloat() ) ) )
-			{
-				GoalOpacity = FMath::Clamp( LeftHandInteractor->GetDragTranslationVelocity().Size() / VREd::GridFadeStartVelocity->GetFloat(), 0.0f, 1.0f );
-			}
-			else if( ( RightHandInteractor->GetLastDraggingMode() == EViewportInteractionDraggingMode::World && !RightHandInteractor->GetDragTranslationVelocity().IsNearlyZero( VREd::GridMovementTolerance->GetFloat() ) ) )
-			{
-				GoalOpacity = FMath::Clamp( RightHandInteractor->GetDragTranslationVelocity().Size() / VREd::GridFadeStartVelocity->GetFloat(), 0.0f, 1.0f );
-			}
-
-			// Check the current opacity and add or subtract to reach the goal
-			float NewOpacity = WorldMovementGridOpacity;
-			if( NewOpacity < GoalOpacity )
-			{
-				NewOpacity = FMath::Min( GoalOpacity, NewOpacity + DeltaTime * VREd::GridFadeMultiplier->GetFloat() );
-			}
-			else if( NewOpacity > GoalOpacity )
-			{
-				NewOpacity = FMath::Max( GoalOpacity, NewOpacity - DeltaTime * VREd::GridFadeMultiplier->GetFloat() );
-			}
-
-			// Only update when the opacity is different
-			if( !FMath::IsNearlyEqual( WorldMovementGridOpacity, GoalOpacity ) )
-			{
-				WorldMovementGridOpacity = NewOpacity;
-
-				// See if the opacity is almost zero and if the goal opacity is lower than the new opacity set it to zero if so. Otherwise it will flicker
-				if( WorldMovementGridOpacity < SMALL_NUMBER )
-				{
-					WorldMovementGridOpacity = 0.f;
-					WorldMovementGridMeshComponent->SetVisibility( false );
-				}
-				else
-				{
-					WorldMovementGridMeshComponent->SetVisibility( true );
-				}
-
-				static const FName OpacityName( "OpacityParam" );
-				WorldMovementGridMID->SetScalarParameterValue( OpacityName, WorldMovementGridOpacity * VREd::GridMaxOpacity->GetFloat() );
-			}
-		}
-		else
-		{
-			WorldMovementGridMeshComponent->SetVisibility( false );
-		}
-	}
-
-	// Apply a post process effect while the user is moving the world around.  The effect "greys out" any pixels
-	// that are not nearby, and completely hides the skybox and other very distant objects.  This is to help
-	// prevent simulation sickness when moving/rotating/scaling the entire world around them.
-	{
-		PostProcessComponent->bEnabled = VREd::ShowWorldMovementPostProcess->GetInt() == 1 ? true : false;
-
-		// Make sure our world size is reflected in the post process material
-		static FName WorldScaleFactorParameterName( "WorldScaleFactor" );
-		WorldMovementPostProcessMaterial->SetScalarParameterValue( WorldScaleFactorParameterName, GetWorldScaleFactor() );
-
-		static FName RoomOriginParameterName( "RoomOrigin" );
-		WorldMovementPostProcessMaterial->SetVectorParameterValue( RoomOriginParameterName, GetRoomTransform().GetLocation() );
-
-		static FName FogColorParameterName( "FogColor" );
-		// WorldMovementPostProcessMaterial->SetVectorParameterValue( FogColorParameterName, FLinearColor::Black );
-
-		static FName StartDistanceParameterName( "StartDistance" );
-		WorldMovementPostProcessMaterial->SetScalarParameterValue( StartDistanceParameterName, VREd::WorldMovementFogStartDistance->GetFloat() );
-
-		static FName EndDistanceParameterName( "EndDistance" );
-		WorldMovementPostProcessMaterial->SetScalarParameterValue( EndDistanceParameterName, VREd::WorldMovementFogEndDistance->GetFloat() );
-
-		static FName FogOpacityParameterName( "FogOpacity" );
-		WorldMovementPostProcessMaterial->SetScalarParameterValue( FogOpacityParameterName, VREd::WorldMovementFogOpacity->GetFloat() );
-
-		static FName SkyboxDistanceParameterName( "SkyboxDistance" );
-		WorldMovementPostProcessMaterial->SetScalarParameterValue( SkyboxDistanceParameterName, VREd::WorldMovementFogSkyboxDistance->GetFloat() );
-
-		const bool bShouldDrawWorldMovementPostProcess = WorldMovementGridOpacity > KINDA_SMALL_NUMBER;
-		if( bShouldDrawWorldMovementPostProcess != bIsDrawingWorldMovementPostProcess )
-		{
-			UPostProcessComponent* PostProcess = GetPostProcessComponent();
-			if( bShouldDrawWorldMovementPostProcess )
-			{
-				PostProcess->AddOrUpdateBlendable( WorldMovementPostProcessMaterial );
-				bIsDrawingWorldMovementPostProcess = true;
-			}
-			else
-			{
-				bIsDrawingWorldMovementPostProcess = false;
-				PostProcess->Settings.RemoveBlendable( WorldMovementPostProcessMaterial );
-			}
-		}
-
-		if( bIsDrawingWorldMovementPostProcess )
-		{
-			static FName OpacityParameterName( "Opacity" );
-			WorldMovementPostProcessMaterial->SetScalarParameterValue( OpacityParameterName, FMath::Clamp( WorldMovementGridOpacity, 0.0f, 1.0f ) );
-		}
-	}
-
 	StopOldHapticEffects();
 
 	bFirstTick = false;
@@ -1007,16 +584,8 @@ bool FVREditorMode::IsCompatibleWith(FEditorModeID OtherModeID) const
 
 void FVREditorMode::AddReferencedObjects( FReferenceCollector& Collector )
 {
-	Collector.AddReferencedObject( AvatarMeshActor );
-	Collector.AddReferencedObject( HeadMeshComponent );
-	Collector.AddReferencedObject( FlashlightComponent );
-	Collector.AddReferencedObject( WorldMovementGridMeshComponent );
-	Collector.AddReferencedObject( PostProcessComponent );
-	Collector.AddReferencedObject( WorldMovementGridMID );
-	Collector.AddReferencedObject( WorldMovementPostProcessMaterial );
-	Collector.AddReferencedObject( ScaleProgressMeshComponent );
-	Collector.AddReferencedObject( CurrentScaleProgressMeshComponent );
-	Collector.AddReferencedObject( UserScaleIndicatorText );	
+	Collector.AddReferencedObject( AvatarActor );
+	Collector.AddReferencedObject( FlashlightComponent );	
 	Collector.AddReferencedObject( UISystem );
 	Collector.AddReferencedObject( WorldInteraction );
 	Collector.AddReferencedObject( TeleporterSystem );
@@ -1045,7 +614,7 @@ void FVREditorMode::Render( const FSceneView* SceneView, FViewport* Viewport, FP
 
 AActor* FVREditorMode::GetAvatarMeshActor()
 {
-	return AvatarMeshActor;
+	return AvatarActor;
 }
 
 UWorld* FVREditorMode::GetWorld() const
@@ -1123,21 +692,7 @@ AActor* FVREditorMode::SpawnTransientSceneActor( TSubclassOf<AActor> ActorClass,
 
 void FVREditorMode::DestroyTransientActor( AActor* Actor ) const
 {
-	if( Actor != nullptr )
-	{
-		const bool bWasWorldPackageDirty = GetWorld()->GetOutermost()->IsDirty();
-
-		const bool bNetForce = false;
-		const bool bShouldModifyLevel = false;	// Don't modify level for transient actor destruction
-		GetWorld()->DestroyActor( Actor, bNetForce, bShouldModifyLevel );
-		Actor = nullptr;
-
-		// Don't dirty the level file after destroying a transient actor
-		if( !bWasWorldPackageDirty )
-		{
-			GetWorld()->GetOutermost()->SetDirtyFlag( false );
-		}
-	}
+	WorldInteraction->DestroyTransientActor( Actor );
 }
 
 const SLevelViewport& FVREditorMode::GetLevelViewportPossessedForVR() const
@@ -1184,30 +739,10 @@ void FVREditorMode::CleanUpActorsBeforeMapChangeOrSimulate()
 {
 	// NOTE: This will be called even when this mode is not currently active!
 
-	DestroyTransientActor( AvatarMeshActor );
-	AvatarMeshActor = nullptr;
-	HeadMeshComponent = nullptr;
+	DestroyTransientActor( AvatarActor );
+	AvatarActor = nullptr;
 	FlashlightComponent = nullptr;
 
-	WorldMovementGridMeshComponent = nullptr;
-
-	if( WorldMovementGridMID != nullptr )
-	{
-		WorldMovementGridMID->MarkPendingKill();
-		WorldMovementGridMID = nullptr;
-	}
-
-	if( WorldMovementPostProcessMaterial != nullptr )
-	{
-		WorldMovementPostProcessMaterial->MarkPendingKill();
-		WorldMovementPostProcessMaterial = nullptr;
-	}
-
-	PostProcessComponent = nullptr;
-	ScaleProgressMeshComponent = nullptr;
-	CurrentScaleProgressMeshComponent = nullptr;
-	UserScaleIndicatorText = nullptr;
-
 	if( UISystem != nullptr )
 	{
 		UISystem->CleanUpActorsBeforeMapChangeOrSimulate();
@@ -1226,8 +761,8 @@ void FVREditorMode::ToggleFlashlight( UVREditorInteractor* Interactor )
 	{
 		if ( FlashlightComponent == nullptr )
 		{
-			FlashlightComponent = NewObject<USpotLightComponent>( AvatarMeshActor );
-			AvatarMeshActor->AddOwnedComponent( FlashlightComponent );
+			FlashlightComponent = NewObject<USpotLightComponent>( AvatarActor );
+			AvatarActor->AddOwnedComponent( FlashlightComponent );
 			FlashlightComponent->RegisterComponent();
 			FlashlightComponent->SetMobility( EComponentMobility::Movable );
 			FlashlightComponent->SetCastShadows( false );
diff --git a/Engine/Source/Editor/VREditor/VREditorMode.h b/Engine/Source/Editor/VREditor/VREditorMode.h
--- a/Engine/Source/Editor/VREditor/VREditorMode.h
+++ b/Engine/Source/Editor/VREditor/VREditorMode.h
@@ -154,9 +154,6 @@ public:
 	/** Sets up our avatar mesh, if not already spawned */
 	void SpawnAvatarMeshActor();
 
-	/** Gets our VR post process component, creating it if it doesn't exist yet */
-	UPostProcessComponent* GetPostProcessComponent();
-
 	/** Called internally when the user changes maps, enters/exits PIE or SIE, or switched between PIE/SIE */
 	void CleanUpActorsBeforeMapChangeOrSimulate();
 
@@ -259,10 +256,12 @@ protected:
 	//
 
 	/** Actor with components to represent the VR avatar in the world, including motion controller meshes */
-	class AActor* AvatarMeshActor;
+	class AVREditorAvatarActor* AvatarActor;
+
 
-	/** Our avatar's head mesh */
-	class UStaticMeshComponent* HeadMeshComponent;
+	//
+	// Flashlight
+	//
 
 	/** Spotlight for the flashlight */
 	class USpotLightComponent* FlashlightComponent;
@@ -270,47 +269,6 @@ protected:
 	/** If there is currently a flashlight in the scene */
 	bool bIsFlashlightOn;
 
-
-	//
-	// World movement grid & FX
-	//
-
-	/** The grid that appears while the user is dragging the world around */
-	class UStaticMeshComponent* WorldMovementGridMeshComponent;
-
-	/** Grid mesh component dynamic material instance to set the opacity */
-	class UMaterialInstanceDynamic* WorldMovementGridMID;
-
-	/** Opacity of the movement grid and post process */
-	float WorldMovementGridOpacity;
-
-	/** True if we're currently drawing our world movement post process */
-	bool bIsDrawingWorldMovementPostProcess;
-
-	/** Post process material for "greying out" the world while in world movement mode */
-	class UMaterialInstanceDynamic* WorldMovementPostProcessMaterial;
-
-	//
-	// World scaling progress bar
-	//
-
-	/** Background progressbar scaling mesh */
-	class UStaticMeshComponent* ScaleProgressMeshComponent;
-
-	/** Current scale progressbar mesh */
-	class UStaticMeshComponent* CurrentScaleProgressMeshComponent;
-
-	/** Current scale text */
-	class UTextRenderComponent* UserScaleIndicatorText;
-
-
-	//
-	// Post process
-	//
-
-	/** Post process for drawing VR-specific post effects */
-	class UPostProcessComponent* PostProcessComponent;
-
 	//
 	// Input
 	//
diff --git a/Engine/Source/Editor/VREditor/VREditorModule.cpp b/Engine/Source/Editor/VREditor/VREditorModule.cpp
--- a/Engine/Source/Editor/VREditor/VREditorModule.cpp
+++ b/Engine/Source/Editor/VREditor/VREditorModule.cpp
@@ -190,18 +190,21 @@ void FVREditorModule::Tick( float DeltaTime )
 	// You can only auto-enter VR if the setting is enabled. Other criteria are that the VR Editor is enabled in experimental settings, that you are not in PIE, and that the editor is foreground.
 	bool bCanAutoEnterVR = (GetDefault<UEditorExperimentalSettings>()->bEnableAutoVREditMode) && (GetDefault<UEditorExperimentalSettings>()->bEnableVREditing) && !(GEditor->PlayWorld) && FPlatformProcess::IsThisApplicationForeground();
 
-	// Only check whether you are wearing the HMD every second, if you are allowed to auto-enter VR, and if your HMD state has changed since the last check. 
-	if((TimeSinceHMDChecked >= 1.0f) && bCanAutoEnterVR && (HMDWornState != GEngine->HMDDevice->GetHMDWornState()))
+	if ( GEngine != nullptr && GEngine->HMDDevice.IsValid() )
 	{
-		TimeSinceHMDChecked = 0.0f;
-		HMDWornState = GEngine->HMDDevice->GetHMDWornState();
-		if( HMDWornState == EHMDWornState::Worn )
+		// Only check whether you are wearing the HMD every second, if you are allowed to auto-enter VR, and if your HMD state has changed since the last check. 
+		if ((TimeSinceHMDChecked >= 1.0f) && bCanAutoEnterVR && (HMDWornState != GEngine->HMDDevice->GetHMDWornState()))
 		{
-			EnableVREditor(true, false);
-		}
-		else if (HMDWornState == EHMDWornState::NotWorn)
-		{
-			EnableVREditor(false , false);
+			TimeSinceHMDChecked = 0.0f;
+			HMDWornState = GEngine->HMDDevice->GetHMDWornState();
+			if (HMDWornState == EHMDWornState::Worn)
+			{
+				EnableVREditor(true, false);
+			}
+			else if (HMDWornState == EHMDWornState::NotWorn)
+			{
+				EnableVREditor(false, false);
+			}
 		}
 	}
 
diff --git a/Engine/Source/Editor/VREditor/VREditorMotionControllerInteractor.cpp b/Engine/Source/Editor/VREditor/VREditorMotionControllerInteractor.cpp
--- a/Engine/Source/Editor/VREditor/VREditorMotionControllerInteractor.cpp
+++ b/Engine/Source/Editor/VREditor/VREditorMotionControllerInteractor.cpp
@@ -6,6 +6,7 @@
 #include "VREditorMode.h"
 #include "VREditorUISystem.h"
 #include "VREditorFloatingText.h"
+#include "VREditorDockableWindow.h"
 
 #include "IMotionController.h"
 #include "IHeadMountedDisplay.h"
@@ -349,9 +350,10 @@ void UVREditorMotionControllerInteractor::Tick( const float DeltaTime )
 			// point.  Otherwise, we always want the laser to protrude through hovered objects, so that you can
 			// interact with translucent gizmo handles that are occluded by geometry
 			FVector LaserPointerImpactPoint = LaserPointerEnd;
-			if ( GetInteractorData().bIsHovering && ( GetInteractorData().DraggingMode != EViewportInteractionDraggingMode::Nothing || IsHoveringOverUI() ) )
+			if ( IsHoveringOverGizmo() || IsHoveringOverUI() || 
+				 ( GetHoverComponent() != nullptr && GetHoverComponent()->GetOwner()->IsA<AVREditorDockableWindow>() ) )
 			{
-				LaserPointerImpactPoint = GetInteractorData().HoverLocation;
+				LaserPointerImpactPoint = GetHoverLocation();
 			}
 
 			// Apply rotation offset to the laser direction
diff --git a/Engine/Source/Editor/ViewportInteraction/Public/ViewportInteractor.h b/Engine/Source/Editor/ViewportInteraction/Public/ViewportInteractor.h
--- a/Engine/Source/Editor/ViewportInteraction/Public/ViewportInteractor.h
+++ b/Engine/Source/Editor/ViewportInteraction/Public/ViewportInteractor.h
@@ -128,6 +128,9 @@ public:
 	/** If the interactor laser is currently hovering */
 	bool IsHovering() const;
 
+	/** If the interactor laser is currently hovering over a gizmo handle */
+	bool IsHoveringOverGizmo();
+
 	/** Sets the current dragging mode for this interactor */
 	void SetDraggingMode( const EViewportInteractionDraggingMode NewDraggingMode );
 
diff --git a/Engine/Source/Editor/ViewportInteraction/Public/ViewportWorldInteraction.h b/Engine/Source/Editor/ViewportInteraction/Public/ViewportWorldInteraction.h
--- a/Engine/Source/Editor/ViewportInteraction/Public/ViewportWorldInteraction.h
+++ b/Engine/Source/Editor/ViewportInteraction/Public/ViewportWorldInteraction.h
@@ -69,8 +69,10 @@ class VIEWPORTINTERACTION_API UViewportWorldInteraction : public UObject, public
 	
 public:
 
+	~UViewportWorldInteraction();
+
 	// IViewportWorldInteraction overrides
-	virtual void Init( class FEditorViewportClient* InEditorViewportClient ) override;
+	virtual void Init( const TSharedPtr<class FEditorViewportClient>& InEditorViewportClient ) override;
 	virtual void Shutdown() override;
 	virtual void Tick( class FEditorViewportClient* ViewportClient, const float DeltaTime ) override;
 	virtual void AddInteractor( UViewportInteractor* Interactor ) override;
@@ -215,6 +217,9 @@ public:
 	/** Switch which transform gizmo coordinate space we're using. */
 	void CycleTransformGizmoCoordinateSpace();
 
+	/** Destroys a transient actor we created earlier and nulls out the pointer */
+	void DestroyTransientActor(AActor* Actor) const;
+
 protected:
 
 	/** Allow base class to refresh on selection changed */
@@ -312,9 +317,6 @@ private:
 	/** Spawns a transient actor that we can use in the editor world */
 	AActor* SpawnTransientSceneActor( TSubclassOf<AActor> ActorClass, const FString& ActorName, const bool bWithSceneComponent ) const;
 
-	/** Destroys a transient actor we created earlier and nulls out the pointer */
-	void DestroyTransientActor( AActor* Actor ) const;
-
 	/** Gets the snap grid mesh actor */
 	AActor* GetSnapGridActor()
 	{
@@ -387,7 +389,7 @@ private:
 	//
 
 	/** The viewport that the worldinteraction is used for */
-	class FEditorViewportClient* EditorViewportClient;
+	TSharedPtr<class FEditorViewportClient> EditorViewportClient;
 
 	/** The last frame that input was polled.  We keep track of this so that we can make sure to poll for latest input right before
 	its first needed each frame */
diff --git a/Engine/Source/Editor/ViewportInteraction/Public/ViewportWorldInteractionInterface.h b/Engine/Source/Editor/ViewportInteraction/Public/ViewportWorldInteractionInterface.h
--- a/Engine/Source/Editor/ViewportInteraction/Public/ViewportWorldInteractionInterface.h
+++ b/Engine/Source/Editor/ViewportInteraction/Public/ViewportWorldInteractionInterface.h
@@ -16,7 +16,7 @@ class VIEWPORTINTERACTION_API IViewportWorldInteractionInterface
 
 public:
 	/** Initializes the world interaction */
-	virtual void Init( class FEditorViewportClient* InEditorViewportClient ) = 0;
+	virtual void Init( const TSharedPtr<class FEditorViewportClient>& InEditorViewportClient ) = 0;
 
 	/** Shuts down the world interaction */
 	virtual void Shutdown() = 0;
diff --git a/Engine/Source/Editor/ViewportInteraction/ViewportInteractor.cpp b/Engine/Source/Editor/ViewportInteraction/ViewportInteractor.cpp
--- a/Engine/Source/Editor/ViewportInteraction/ViewportInteractor.cpp
+++ b/Engine/Source/Editor/ViewportInteraction/ViewportInteractor.cpp
@@ -127,10 +127,10 @@ bool UViewportInteractor::HandleInputKey( const FKey Key, const EInputEvent Even
 
 							if ( IViewportInteractableInterface* ActorInteractable = Cast<IViewportInteractableInterface>( Actor ) )
 							{
-								if ( Action->ActionType == ViewportWorldActionTypes::SelectAndMove )
+								if ( Action->ActionType == ViewportWorldActionTypes::SelectAndMove_LightlyPressed )
 								{
 									bHandled = true;
-
+									SetAllowTriggerFullPress( false );
 									bool bResultedInInteractableDrag = false;
 									ActorInteractable->OnPressed( this, HitResult, bResultedInInteractableDrag );
 
@@ -491,10 +491,6 @@ FHitResult UViewportInteractor::GetHitResultFromLaserPointer( TArray<AActor*>* O
 			const bool bTraceComplex = true;
 			FCollisionQueryParams TraceParams( NAME_None, bTraceComplex, nullptr );
 
-			const FCollisionResponseParams& ResponseParam = FCollisionResponseParams::DefaultResponseParam;
-
-			const ECollisionChannel CollisionChannel = bOnlyEditorGizmos ? COLLISION_GIZMO : ECC_Visibility;
-
 			if ( OptionalListOfIgnoredActors != nullptr )
 			{
 				TraceParams.AddIgnoredActors( *OptionalListOfIgnoredActors );
@@ -504,38 +500,48 @@ FHitResult UViewportInteractor::GetHitResultFromLaserPointer( TArray<AActor*>* O
 			FHitResult HitResult;
 			if ( bOnlyEditorGizmos )
 			{
+				const FCollisionResponseParams& ResponseParam = FCollisionResponseParams::DefaultResponseParam;
+				const ECollisionChannel CollisionChannel = bOnlyEditorGizmos ? COLLISION_GIZMO : ECC_Visibility;
+
 				bHit = WorldInteraction->GetViewportWorld()->LineTraceSingleByChannel( HitResult, LaserPointerStart, LaserPointerEnd, CollisionChannel, TraceParams, ResponseParam );
+				if ( bHit )
+				{
+					BestHitResult = HitResult;
+				}
 			}
 			else
 			{
 				FCollisionObjectQueryParams EverythingButGizmos( FCollisionObjectQueryParams::AllObjects );
 				EverythingButGizmos.RemoveObjectTypesToQuery( COLLISION_GIZMO );
 				bHit = WorldInteraction->GetViewportWorld()->LineTraceSingleByObjectType( HitResult, LaserPointerStart, LaserPointerEnd, EverythingButGizmos, TraceParams );
-			}
-
-			if ( bHit )
-			{
-				bool bHitResultIsPriorityType = false;
-				if ( !bOnlyEditorGizmos && ObjectsInFrontOfGizmo )
+				
+				if ( bHit )
 				{
-					for ( UClass* CurrentClass : *ObjectsInFrontOfGizmo )
+					bool bHitResultIsPriorityType = false;
+					if ( !bOnlyEditorGizmos && ObjectsInFrontOfGizmo )
 					{
-						bool bClassHasPriority = false;
-						bClassHasPriority =
-							( HitResult.GetComponent() != nullptr && HitResult.GetComponent()->IsA( CurrentClass ) ) ||
-							( HitResult.GetActor() != nullptr && HitResult.GetActor()->IsA( CurrentClass ) );
-
-						if ( bClassHasPriority )
+						for ( UClass* CurrentClass : *ObjectsInFrontOfGizmo )
 						{
-							bHitResultIsPriorityType = bClassHasPriority;
-							break;
+							bool bClassHasPriority = false;
+							bClassHasPriority =
+								( HitResult.GetComponent() != nullptr && HitResult.GetComponent()->IsA( CurrentClass ) ) ||
+								( HitResult.GetActor() != nullptr && HitResult.GetActor()->IsA( CurrentClass ) );
+
+							if ( bClassHasPriority )
+							{
+								bHitResultIsPriorityType = bClassHasPriority;
+								break;
+							}
 						}
 					}
-				}
 
-				if ( BestHitResult.GetActor() == nullptr || bHitResultIsPriorityType )
-				{
-					BestHitResult = HitResult;
+					const bool bHitResultIsGizmo = HitResult.GetActor() != nullptr && HitResult.GetActor() == WorldInteraction->GetTransformGizmoActor();
+					if ( BestHitResult.GetActor() == nullptr ||
+						 bHitResultIsPriorityType ||
+						 bHitResultIsGizmo )
+					{
+						BestHitResult = HitResult;
+					}
 				}
 			}
 		}
@@ -562,6 +568,11 @@ bool UViewportInteractor::IsHovering() const
 	return InteractorData.bIsHovering;
 }
 
+bool UViewportInteractor::IsHoveringOverGizmo()
+{
+	return InteractorData.HoveringOverTransformGizmoComponent.IsValid();
+}
+
 void UViewportInteractor::SetDraggingMode( const EViewportInteractionDraggingMode NewDraggingMode )
 {
 	InteractorData.DraggingMode = NewDraggingMode;
diff --git a/Engine/Source/Editor/ViewportInteraction/ViewportWorldInteraction.cpp b/Engine/Source/Editor/ViewportInteraction/ViewportWorldInteraction.cpp
--- a/Engine/Source/Editor/ViewportInteraction/ViewportWorldInteraction.cpp
+++ b/Engine/Source/Editor/ViewportInteraction/ViewportWorldInteraction.cpp
@@ -215,7 +215,6 @@ UViewportWorldInteraction::UViewportWorldInteraction( const FObjectInitializer&
 	LastDragGizmoStartTransform( FTransform::Identity ),
 	TrackingTransaction( FTrackingTransaction() ),
 	AppTimeEntered( FTimespan::Zero() ),
-	EditorViewportClient( nullptr ),
 	LastFrameNumberInputWasPolled( 0 ),
 	MotionControllerID( 0 ),	// @todo ViewportInteraction: We only support a single controller, and we assume the first controller are the motion controls
 	LastWorldToMetersScale( 100.0f ),
@@ -236,7 +235,12 @@ UViewportWorldInteraction::UViewportWorldInteraction( const FObjectInitializer&
 {
 }
 
-void UViewportWorldInteraction::Init( FEditorViewportClient* InEditorViewportClient )
+UViewportWorldInteraction::~UViewportWorldInteraction()
+{
+	Shutdown();
+}
+
+void UViewportWorldInteraction::Init( const TSharedPtr<FEditorViewportClient>& InEditorViewportClient )
 {
 	InputProcessor = MakeShareable( new FViewportInteractionInputProcessor( *this ) );
 	FSlateApplication::Get().SetInputPreProcessor( true, InputProcessor );
@@ -244,7 +248,6 @@ void UViewportWorldInteraction::Init( FEditorViewportClient* InEditorViewportCli
 	// Find out about selection changes
 	USelection::SelectionChangedEvent.AddUObject( this, &UViewportWorldInteraction::OnActorSelectionChanged ); //@todo viewportinteraction
 
-																											   // Setting up colors
 	Colors.SetNumZeroed( (int32)EColors::TotalCount );
 	{
 		Colors[ (int32)EColors::DefaultColor ] = FLinearColor( 0.7f, 0.7f, 0.7f, 1.0f );;
@@ -294,13 +297,16 @@ void UViewportWorldInteraction::Shutdown()
 	}
 	Interactors.Empty();
 
+	EditorViewportClient.Reset();
+	DraggedInteractable = nullptr;
+
 	USelection::SelectionChangedEvent.RemoveAll( this );
 }
 
 void UViewportWorldInteraction::Tick( FEditorViewportClient* ViewportClient, const float DeltaTime )
 {
 	// Only if this is our viewport. Remember, editor modes get a chance to tick and receive input for each active viewport.
-	if ( ViewportClient != EditorViewportClient )
+	if ( !EditorViewportClient.IsValid() || ViewportClient != EditorViewportClient.Get() )
 	{
 		return;
 	}
@@ -429,12 +435,16 @@ float UViewportWorldInteraction::GetWorldScaleFactor() const
 
 UWorld* UViewportWorldInteraction::GetViewportWorld() const
 {
-	return EditorViewportClient->GetWorld();
+	if ( EditorViewportClient.IsValid() )
+	{
+		return EditorViewportClient->GetWorld();
+	}
+	return nullptr;
 }
 
 FEditorViewportClient* UViewportWorldInteraction::GetViewportClient() const
 {
-	return EditorViewportClient;
+	return EditorViewportClient.Get();
 }
 
 void UViewportWorldInteraction::Undo()
@@ -506,23 +516,22 @@ void UViewportWorldInteraction::HoverTick( FEditorViewportClient* ViewportClient
 			InteractorData.HoveringOverTransformGizmoComponent = nullptr;
 			InteractorData.HoverLocation = FVector::ZeroVector;
 
-			FHitResult HitResult = Interactor->GetHitResultFromLaserPointer();
-			if ( HitResult.Actor.IsValid() )
+			if ( HitHoverResult.Actor.IsValid() )
 			{
-				USceneComponent* HoveredActorComponent = HitResult.GetComponent();
-				AActor* Actor = HitResult.Actor.Get();
+				USceneComponent* HoveredActorComponent = HitHoverResult.GetComponent();
+				AActor* Actor = HitHoverResult.Actor.Get();
 
 				if ( HoveredActorComponent && IsInteractableComponent( HoveredActorComponent ) )
 				{
 					HoveredObjects.Add( FViewportHoverTarget( Actor ) );
 
-					InteractorData.HoverLocation = HitResult.ImpactPoint;
+					InteractorData.HoverLocation = HitHoverResult.ImpactPoint;
 
 					if ( Actor == TransformGizmoActor )
 					{
 						NewHoveredActorComponent = HoveredActorComponent;
 						InteractorData.HoveringOverTransformGizmoComponent = HoveredActorComponent;
-						InteractorData.HoverHapticCheckLastHoveredGizmoComponent = HitResult.GetComponent();
+						InteractorData.HoverHapticCheckLastHoveredGizmoComponent = HitHoverResult.GetComponent();
 					}
 					else
 					{
@@ -534,7 +543,7 @@ void UViewportWorldInteraction::HoverTick( FEditorViewportClient* ViewportClient
 							// Check if the current hovered component of the interactor is different from the hitresult component
 							if ( NewHoveredActorComponent != InteractorData.HoveredActorComponent && !IsOtherInteractorHoveringOverComponent( Interactor, NewHoveredActorComponent ) )
 							{
-								Interactable->OnHoverEnter( Interactor, HitResult );
+								Interactable->OnHoverEnter( Interactor, HitHoverResult );
 							}
 
 							Interactable->OnHover( Interactor );
@@ -567,8 +576,6 @@ void UViewportWorldInteraction::HoverTick( FEditorViewportClient* ViewportClient
 			//Update the hovered actor component with the new component
 			InteractorData.HoveredActorComponent = NewHoveredActorComponent;
 		}
-
-		
 	}
 }
 
@@ -2387,7 +2394,7 @@ void UViewportWorldInteraction::PollInputIfNeeded()
 
 void UViewportWorldInteraction::OnActorSelectionChanged( UObject* ChangedObject )
 {
-	if ( EditorViewportClient != nullptr )
+	if ( EditorViewportClient.IsValid() )
 	{
 		const bool bNewObjectsSelected = true;
 		const bool bAllHandlesVisible = true;
@@ -2767,19 +2774,20 @@ AActor* UViewportWorldInteraction::SpawnTransientSceneActor( TSubclassOf<AActor>
 
 void UViewportWorldInteraction::DestroyTransientActor( AActor* Actor ) const
 {
-	if ( Actor != nullptr )
+	UWorld* World = GetViewportWorld();
+	if (Actor != nullptr && World != nullptr)
 	{
-		const bool bWasWorldPackageDirty = GetViewportWorld()->GetOutermost()->IsDirty();
+		const bool bWasWorldPackageDirty = World->GetOutermost()->IsDirty();
 
 		const bool bNetForce = false;
 		const bool bShouldModifyLevel = false;	// Don't modify level for transient actor destruction
-		GetViewportWorld()->DestroyActor( Actor, bNetForce, bShouldModifyLevel );
+		World->DestroyActor(Actor, bNetForce, bShouldModifyLevel);
 		Actor = nullptr;
 
 		// Don't dirty the level file after destroying a transient actor
-		if ( !bWasWorldPackageDirty )
+		if (!bWasWorldPackageDirty)
 		{
-			GetViewportWorld()->GetOutermost()->SetDirtyFlag( false );
+			World->GetOutermost()->SetDirtyFlag(false);
 		}
 	}
 }
