VREditor : - Fix crash when starting VREditor and then changing levels #jira UE-33766

diff --git a/Engine/Source/Editor/VREditor/VREditorMode.cpp b/Engine/Source/Editor/VREditor/VREditorMode.cpp
--- a/Engine/Source/Editor/VREditor/VREditorMode.cpp
+++ b/Engine/Source/Editor/VREditor/VREditorMode.cpp
@@ -737,19 +737,18 @@ void FVREditorMode::OnSwitchBetweenPIEAndSIE( const bool bIsSimulatingInEditor )
 
 void FVREditorMode::CleanUpActorsBeforeMapChangeOrSimulate()
 {
-	// NOTE: This will be called even when this mode is not currently active!
-
-	DestroyTransientActor( AvatarActor );
-	AvatarActor = nullptr;
-	FlashlightComponent = nullptr;
-
-	if( UISystem != nullptr )
+	if ( WorldInteraction != nullptr )
 	{
-		UISystem->CleanUpActorsBeforeMapChangeOrSimulate();
-	}
+		// NOTE: This will be called even when this mode is not currently active!
+		DestroyTransientActor( AvatarActor );
+		AvatarActor = nullptr;
+		FlashlightComponent = nullptr;
+
+		if ( UISystem != nullptr )
+		{
+			UISystem->CleanUpActorsBeforeMapChangeOrSimulate();
+		}
 
-	if( WorldInteraction != nullptr )
-	{
 		WorldInteraction->Shutdown();
 	}
 }
