[INTEGRATE] VREditor: Don't teleport viewport camera when entering 'forced VR mode' (from CL 3071686)

diff --git a/Engine/Source/Editor/VREditor/VREditorMode.cpp b/Engine/Source/Editor/VREditor/VREditorMode.cpp
--- a/Engine/Source/Editor/VREditor/VREditorMode.cpp
+++ b/Engine/Source/Editor/VREditor/VREditorMode.cpp
@@ -381,11 +381,15 @@ void FVREditorMode::Exit()
 				VRViewportClient.bAlwaysShowModeWidgetAfterSelectionChanges = SavedEditorState.bAlwaysShowModeWidgetAfterSelectionChanges;
 				VRViewportClient.EngineShowFlags = SavedEditorState.ShowFlags;
 				VRViewportClient.SetGameView( SavedEditorState.bGameView );
-				VRViewportClient.SetViewLocation( GetHeadTransform().GetLocation() );
+
+				if( bActuallyUsingVR )
+				{
+					VRViewportClient.SetViewLocation( GetHeadTransform().GetLocation() );
 				
-				FRotator HeadRotationNoRoll = GetHeadTransform().GetRotation().Rotator();
-				HeadRotationNoRoll.Roll = 0.0f;
-				VRViewportClient.SetViewRotation(HeadRotationNoRoll); // Use SavedEditorState.ViewRotation to go back to start rot
+					FRotator HeadRotationNoRoll = GetHeadTransform().GetRotation().Rotator();
+					HeadRotationNoRoll.Roll = 0.0f;
+					VRViewportClient.SetViewRotation(HeadRotationNoRoll); // Use SavedEditorState.ViewRotation to go back to start rot
+				}
 				
 				VRViewportClient.SetRealtime( SavedEditorState.bRealTime );
 
@@ -489,7 +493,7 @@ void FVREditorMode::Tick( FEditorViewportClient* ViewportClient, float DeltaTime
 	}
 
 	//Setting the initial position and rotation based on the editor viewport when going into VR mode
-	if ( bFirstTick )
+	if ( bFirstTick && bActuallyUsingVR )
 	{
 		const FTransform RoomToWorld = GetRoomTransform();
 		const FTransform WorldToRoom = RoomToWorld.Inverse();
