VR Editor: - VREditor extensibility make sure all worlds have world interaction #jira UETOOL-892
- FVREditorMode made into a UObject
- Added EditorWorldManager that holds all worlds with a ViewportWorldInteraction and VREditorMode
- ViewportWorldInteraction is ticked from UEditorEngine instead of using FTickableEditorObject
- Fixed transform gizmo being visible on start at the viewport location
- Fixed Crash when double clicking a level in Content Browser to open it #jira UE-32541
- Fixed AvatarActor log warning when going into VREditor mode #jira UE-34434

diff --git a/Engine/Source/Editor/FoliageEdit/Private/FoliageEdMode.cpp b/Engine/Source/Editor/FoliageEdit/Private/FoliageEdMode.cpp
--- a/Engine/Source/Editor/FoliageEdit/Private/FoliageEdMode.cpp
+++ b/Engine/Source/Editor/FoliageEdit/Private/FoliageEdMode.cpp
@@ -38,8 +38,10 @@
 
 // VR Editor
 #include "IVREditorModule.h"
+#include "VREditorMode.h"
 #include "ViewportWorldInteraction.h"
 #include "VREditorInteractor.h"
+#include "EditorWorldManager.h"
 
 
 #define LOCTEXT_NAMESPACE "FoliageEdMode"
@@ -345,18 +347,15 @@ void FEdModeFoliage::Enter()
 	// Update UI
 	NotifyNewCurrentLevel();
 
-	// Register to find out when foliage mode is activated
-	GetModeManager()->OnEditorModeChanged().AddRaw(this, &FEdModeFoliage::OnEditorModeChanged);
-
 	// Register for VR input events
-	IVREditorMode* VREditorMode = static_cast<IVREditorMode*>(GetModeManager()->GetActiveMode(IVREditorModule::Get().GetVREditorModeID()));
-	if (VREditorMode != nullptr)
+	UViewportWorldInteraction* ViewportWorldInteraction = GEditor->GetEditorWorldManager()->GetEditorWorldWrapper(GetWorld())->GetViewportWorldInteraction();
+	if (ViewportWorldInteraction != nullptr)
 	{
-		VREditorMode->GetWorldInteraction().OnViewportInteractionInputAction().RemoveAll(this);
-		VREditorMode->GetWorldInteraction().OnViewportInteractionInputAction().AddRaw(this, &FEdModeFoliage::OnVRAction);
+		ViewportWorldInteraction->OnViewportInteractionInputAction().RemoveAll(this);
+		ViewportWorldInteraction->OnViewportInteractionInputAction().AddRaw(this, &FEdModeFoliage::OnVRAction);
 
-		VREditorMode->GetWorldInteraction().OnViewportInteractionHoverUpdate().RemoveAll(this);
-		VREditorMode->GetWorldInteraction().OnViewportInteractionHoverUpdate().AddRaw(this, &FEdModeFoliage::OnVRHoverUpdate);
+		ViewportWorldInteraction->OnViewportInteractionHoverUpdate().RemoveAll(this);
+		ViewportWorldInteraction->OnViewportInteractionHoverUpdate().AddRaw(this, &FEdModeFoliage::OnVRHoverUpdate);
 	}
 }
 
@@ -366,18 +365,15 @@ void FEdModeFoliage::Exit()
 	// Unregister VR mode from event handlers
 	if (IVREditorModule::IsAvailable())
 	{
-		IVREditorMode* VREditorMode = static_cast<IVREditorMode*>(GetModeManager()->GetActiveMode(IVREditorModule::Get().GetVREditorModeID()));
-		if (VREditorMode != nullptr)
+		UViewportWorldInteraction* WorldInteraction = GEditor->GetEditorWorldManager()->GetEditorWorldWrapper(GetWorld())->GetViewportWorldInteraction();
+		if (WorldInteraction != nullptr)
 		{
-			VREditorMode->GetWorldInteraction().OnViewportInteractionInputAction().RemoveAll(this);
-			VREditorMode->GetWorldInteraction().OnViewportInteractionHoverUpdate().RemoveAll(this);
+			WorldInteraction->OnViewportInteractionInputAction().RemoveAll(this);
+			WorldInteraction->OnViewportInteractionHoverUpdate().RemoveAll(this);
 			FoliageInteractor = nullptr;
 		}
 	}
 
-	// Unregister foliage mode from event handlers
-	GetModeManager()->OnEditorModeChanged().RemoveAll(this);
-
 	FToolkitManager::Get().CloseToolkit(Toolkit.ToSharedRef());
 	Toolkit.Reset();
 
@@ -434,32 +430,12 @@ void FEdModeFoliage::Exit()
 	FEdMode::Exit();
 }
 
-void FEdModeFoliage::OnEditorModeChanged(FEdMode* EditorMode, bool bEntered)
-{
-	// Detect if we're switching to VR Editor mode
-	// VR Editor mode may have gone away, so re-register for events in case it's a different object than the one we were originally bound to
-	check(EditorMode != nullptr);
-	FoliageInteractor = nullptr;
-	if (bEntered && EditorMode->GetID() == IVREditorModule::Get().GetVREditorModeID())
-	{
-		IVREditorMode* VREditorMode = static_cast<IVREditorMode*>(EditorMode);
-		if (VREditorMode != nullptr)
-		{
-			VREditorMode->GetWorldInteraction().OnViewportInteractionInputAction().RemoveAll(this);
-			VREditorMode->GetWorldInteraction().OnViewportInteractionInputAction().AddRaw(this, &FEdModeFoliage::OnVRAction);
-
-			VREditorMode->GetWorldInteraction().OnViewportInteractionHoverUpdate().RemoveAll(this);
-			VREditorMode->GetWorldInteraction().OnViewportInteractionHoverUpdate().AddRaw(this, &FEdModeFoliage::OnVRHoverUpdate);
-		}
-	}
-}
-
 void FEdModeFoliage::OnVRHoverUpdate(FEditorViewportClient& ViewportClient, UViewportInteractor* Interactor, FVector& HoverImpactPoint, bool& bWasHandled)
 {
 	// Check if VR Editor is active
 	if (IVREditorModule::IsAvailable())
 	{
-		IVREditorMode* VREditorMode = static_cast<IVREditorMode*>(GetModeManager()->GetActiveMode(IVREditorModule::Get().GetVREditorModeID()));
+		UVREditorMode* VREditorMode = GEditor->GetEditorWorldManager()->GetEditorWorldWrapper(ViewportClient.GetWorld())->GetVREditorMode();
 		if (VREditorMode != nullptr && VREditorMode->IsFullyInitialized())
 		{
 			// Grab interactor
@@ -499,8 +475,7 @@ void FEdModeFoliage::OnVRHoverUpdate(FEditorViewportClient& ViewportClient, UVie
 
 void FEdModeFoliage::OnVRAction(class FEditorViewportClient& ViewportClient, UViewportInteractor* Interactor, const FViewportActionKeyInput& Action, bool& bOutIsInputCaptured, bool& bWasHandled)
 {
-	IVREditorMode* VREditorMode = static_cast<IVREditorMode*>(GetModeManager()->GetActiveMode(IVREditorModule::Get().GetVREditorModeID()));
-
+	UVREditorMode* VREditorMode = GEditor->GetEditorWorldManager()->GetEditorWorldWrapper(ViewportClient.GetWorld())->GetVREditorMode();
 	if (VREditorMode != nullptr && Interactor != nullptr)
 	{
 		const UVREditorInteractor* VRInteractor =  Cast<UVREditorInteractor>(Interactor);
@@ -843,8 +818,8 @@ void FEdModeFoliage::FoliageBrushTrace(FEditorViewportClient* ViewportClient, co
 bool FEdModeFoliage::MouseMove(FEditorViewportClient* ViewportClient, FViewport* Viewport, int32 MouseX, int32 MouseY)
 {
 	// Use mouse capture if there's no other interactor currently tracing brush
-	IVREditorMode* VREditorMode = static_cast<IVREditorMode*>(GetModeManager()->GetActiveMode(IVREditorModule::Get().GetVREditorModeID()));
-	if (VREditorMode == nullptr)
+	UVREditorMode* VREditorMode = GEditor->GetEditorWorldManager()->GetEditorWorldWrapper(ViewportClient->GetWorld())->GetVREditorMode();
+	if (!VREditorMode->IsActive())
 	{
 		// Compute a world space ray from the screen space mouse coordinates
 		FSceneViewFamilyContext ViewFamily(FSceneViewFamily::ConstructionValues(
@@ -881,8 +856,8 @@ bool FEdModeFoliage::MouseMove(FEditorViewportClient* ViewportClient, FViewport*
 bool FEdModeFoliage::CapturedMouseMove(FEditorViewportClient* ViewportClient, FViewport* Viewport, int32 MouseX, int32 MouseY)
 {
 	// Use mouse capture if there's no other interactor currently tracing brush
-	IVREditorMode* VREditorMode = static_cast<IVREditorMode*>(GetModeManager()->GetActiveMode(IVREditorModule::Get().GetVREditorModeID()));
-	if (VREditorMode == nullptr)
+	UVREditorMode* VREditorMode = GEditor->GetEditorWorldManager()->GetEditorWorldWrapper(ViewportClient->GetWorld())->GetVREditorMode();
+	if (!VREditorMode->IsActive())
 	{
 		 //Compute a world space ray from the screen space mouse coordinates
 		FSceneViewFamilyContext ViewFamily(FSceneViewFamily::ConstructionValues(
@@ -2831,10 +2806,9 @@ bool FEdModeFoliage::CanPaint(const UFoliageType* FoliageType, const ULevel* InL
 
 bool FEdModeFoliage::IsModifierButtonPressed(const FEditorViewportClient* ViewportClient) const
 {
-	IVREditorMode* VREditorMode = static_cast<IVREditorMode*>(GetModeManager()->GetActiveMode(IVREditorModule::Get().GetVREditorModeID()));
 	const UVREditorInteractor* VRInteractor = Cast<UVREditorInteractor>(FoliageInteractor);
 	bool bIsModifierPressed = false;
-	if (VREditorMode != nullptr && VRInteractor != nullptr)
+	if (VRInteractor != nullptr)
 	{
 		bIsModifierPressed = VRInteractor->IsModifierPressed();
 	}
diff --git a/Engine/Source/Editor/FoliageEdit/Private/FoliageEdMode.h b/Engine/Source/Editor/FoliageEdit/Private/FoliageEdMode.h
--- a/Engine/Source/Editor/FoliageEdit/Private/FoliageEdMode.h
+++ b/Engine/Source/Editor/FoliageEdit/Private/FoliageEdMode.h
@@ -411,9 +411,6 @@ public:
 	/** Add desired instances. Uses foliage settings to determine location/scale/rotation and whether instances should be ignored */
 	static void AddInstances(UWorld* InWorld, const TArray<FDesiredFoliageInstance>& DesiredInstances, const FFoliagePaintingGeometryFilter& OverrideGeometryFilter);
 
-	/** Called when an editor mode is entered or exited */
-	void OnEditorModeChanged(FEdMode* EditorMode, bool bEntered);
-
 	/** Called when the user presses a button on their motion controller device */
 	void OnVRAction(class FEditorViewportClient& ViewportClient, class UViewportInteractor* Interactor, const struct FViewportActionKeyInput& Action, bool& bOutIsInputCaptured, bool& bWasHandled);
 
diff --git a/Engine/Source/Editor/LandscapeEditor/Private/LandscapeEdMode.cpp b/Engine/Source/Editor/LandscapeEditor/Private/LandscapeEdMode.cpp
--- a/Engine/Source/Editor/LandscapeEditor/Private/LandscapeEdMode.cpp
+++ b/Engine/Source/Editor/LandscapeEditor/Private/LandscapeEdMode.cpp
@@ -29,6 +29,8 @@
 #include "IVREditorModule.h"
 #include "ViewportWorldInteraction.h"
 #include "VREditorInteractor.h"
+#include "VREditorMode.h"
+#include "EditorWorldManager.h"
 
 // Classes
 #include "Landscape.h"
@@ -468,18 +470,15 @@ void FEdModeLandscape::Enter()
 		GizmoBrush->EnterBrush();
 	}
 
-	// Register to find out about editor modes being activated
-	GetModeManager()->OnEditorModeChanged().AddRaw(this, &FEdModeLandscape::OnEditorModeChanged);
-
 	// Register to find out about VR input events
-	IVREditorMode* VREditorMode = static_cast<IVREditorMode*>(GetModeManager()->GetActiveMode(IVREditorModule::Get().GetVREditorModeID()));
-	if (VREditorMode != nullptr)
+	UViewportWorldInteraction* ViewportWorldInteraction = GEditor->GetEditorWorldManager()->GetEditorWorldWrapper(GetWorld())->GetViewportWorldInteraction();
+	if (ViewportWorldInteraction != nullptr)
 	{
-		VREditorMode->GetWorldInteraction().OnViewportInteractionInputAction().RemoveAll(this);
-		VREditorMode->GetWorldInteraction().OnViewportInteractionInputAction().AddRaw(this, &FEdModeLandscape::OnVRAction);
+		ViewportWorldInteraction->OnViewportInteractionInputAction().RemoveAll(this);
+		ViewportWorldInteraction->OnViewportInteractionInputAction().AddRaw(this, &FEdModeLandscape::OnVRAction);
 
-		VREditorMode->GetWorldInteraction().OnViewportInteractionHoverUpdate().RemoveAll(this);
-		VREditorMode->GetWorldInteraction().OnViewportInteractionHoverUpdate().AddRaw(this, &FEdModeLandscape::OnVRHoverUpdate);
+		ViewportWorldInteraction->OnViewportInteractionHoverUpdate().RemoveAll(this);
+		ViewportWorldInteraction->OnViewportInteractionHoverUpdate().AddRaw(this, &FEdModeLandscape::OnVRHoverUpdate);
 	}
 }
 
@@ -490,18 +489,14 @@ void FEdModeLandscape::Exit()
 	// Unregister VR mode from event handlers
 	if (IVREditorModule::IsAvailable())
 	{
-		IVREditorMode* VREditorMode = static_cast<IVREditorMode*>(GetModeManager()->GetActiveMode(IVREditorModule::Get().GetVREditorModeID()));
-		if (VREditorMode != nullptr)
+		UViewportWorldInteraction* ViewportWorldInteraction = GEditor->GetEditorWorldManager()->GetEditorWorldWrapper(GetWorld())->GetViewportWorldInteraction();
+		if (ViewportWorldInteraction != nullptr)
 		{
-			VREditorMode->GetWorldInteraction().OnViewportInteractionInputAction().RemoveAll(this);
-			VREditorMode->GetWorldInteraction().OnViewportInteractionHoverUpdate().RemoveAll(this);
+			ViewportWorldInteraction->OnViewportInteractionInputAction().RemoveAll(this);
+			ViewportWorldInteraction->OnViewportInteractionHoverUpdate().RemoveAll(this);
 		}
 	}
 	
-	// Unregister landscape mode from event handlers
-	GetModeManager()->OnEditorModeChanged().RemoveAll(this);
-
-
 	FEditorSupportDelegates::WorldChange.Remove(OnWorldChangeDelegateHandle);
 	UMaterial::OnMaterialCompilationFinished().Remove(OnMaterialCompilationFinishedDelegateHandle);
 
@@ -573,29 +568,10 @@ void FEdModeLandscape::Exit()
 }
 
 
-void FEdModeLandscape::OnEditorModeChanged( FEdMode* EditorMode, bool bEntered )
-{
-	// VR Editor mode may have gone away, so re-register for events in case it's a different object than the one we were originally bound to
-	check(EditorMode != nullptr);
-	if (bEntered && EditorMode->GetID() == IVREditorModule::Get().GetVREditorModeID())
-	{
-		IVREditorMode* VREditorMode = static_cast<IVREditorMode*>(EditorMode);
-		if (VREditorMode != nullptr)
-		{
-			VREditorMode->GetWorldInteraction().OnViewportInteractionInputAction().RemoveAll(this);
-			VREditorMode->GetWorldInteraction().OnViewportInteractionInputAction().AddRaw(this, &FEdModeLandscape::OnVRAction);
-
-			VREditorMode->GetWorldInteraction().OnViewportInteractionHoverUpdate().RemoveAll(this);
-			VREditorMode->GetWorldInteraction().OnViewportInteractionHoverUpdate().AddRaw(this, &FEdModeLandscape::OnVRHoverUpdate);
-		}
-	}
-}
-
 void FEdModeLandscape::OnVRHoverUpdate(FEditorViewportClient& ViewportClient, UViewportInteractor* Interactor, FVector& HoverImpactPoint, bool& bWasHandled)
 {
-	IVREditorMode* VREditorMode = static_cast<IVREditorMode*>(GetModeManager()->GetActiveMode(IVREditorModule::Get().GetVREditorModeID()));
-
-	if (VREditorMode != nullptr && Interactor != nullptr)
+	UVREditorMode* VREditorMode = GEditor->GetEditorWorldManager()->GetEditorWorldWrapper(ViewportClient.GetWorld())->GetVREditorMode();
+	if (VREditorMode != nullptr && VREditorMode->IsActive() && Interactor != nullptr)
 	{
 		const UVREditorInteractor* VRInteractor = Cast<UVREditorInteractor>(Interactor);
 		const bool bIsHoveringOverUIVR = VRInteractor->IsHoveringOverUI();
@@ -626,12 +602,11 @@ void FEdModeLandscape::OnVRHoverUpdate(FEditorViewportClient& ViewportClient, UV
 
 void FEdModeLandscape::OnVRAction(FEditorViewportClient& ViewportClient, UViewportInteractor* Interactor, const struct FViewportActionKeyInput& Action, bool& bOutIsInputCaptured, bool& bWasHandled)
 {
-	IVREditorMode* VREditorMode = static_cast<IVREditorMode*>(GetModeManager()->GetActiveMode(IVREditorModule::Get().GetVREditorModeID()));
-	
+	UVREditorMode* VREditorMode = GEditor->GetEditorWorldManager()->GetEditorWorldWrapper(ViewportClient.GetWorld())->GetVREditorMode();
 	// Never show the traditional Unreal transform widget.  It doesn't work in VR because we don't have hit proxies.
 	ViewportClient.EngineShowFlags.SetModeWidgets(false);
 
-	if (VREditorMode != nullptr && Interactor != nullptr)
+	if (VREditorMode != nullptr && VREditorMode->IsActive() && Interactor != nullptr)
 	{
 		if (Action.ActionType == ViewportWorldActionTypes::SelectAndMove_LightlyPressed || Action.ActionType == ViewportWorldActionTypes::SelectAndMove)
 		{
@@ -2517,8 +2492,8 @@ void FEdModeLandscape::ForceRealTimeViewports(const bool bEnable, const bool bSt
 					Viewport.SetRealtime(bEnable, bStoreCurrentState);
 
 					// @todo vreditor: Force game view to true in VREditor since we can't use hitproxies and debug objects yet
-					IVREditorMode* VREditorMode = static_cast<IVREditorMode*>(GetModeManager()->GetActiveMode(IVREditorModule::Get().GetVREditorModeID()));
-					if (VREditorMode != nullptr)
+					UVREditorMode* VREditorMode = GEditor->GetEditorWorldManager()->GetEditorWorldWrapper(Viewport.GetWorld())->GetVREditorMode();
+					if (VREditorMode != nullptr && !VREditorMode->IsActive())
 					{
 						Viewport.SetGameView(true);
 					} 
diff --git a/Engine/Source/Editor/LandscapeEditor/Private/LandscapeEdMode.h b/Engine/Source/Editor/LandscapeEditor/Private/LandscapeEdMode.h
--- a/Engine/Source/Editor/LandscapeEditor/Private/LandscapeEdMode.h
+++ b/Engine/Source/Editor/LandscapeEditor/Private/LandscapeEdMode.h
@@ -426,9 +426,6 @@ public:
 	DECLARE_EVENT(FEdModeLandscape, FTargetsListUpdated);
 	static FTargetsListUpdated TargetsListUpdated;
 
-	/** Called when an editor mode is entered or exited */
-	void OnEditorModeChanged(FEdMode* EditorMode, bool bEntered);
-
 	/** Called when the user presses a button on their motion controller device */
 	void OnVRAction(FEditorViewportClient& ViewportClient, UViewportInteractor* Interactor, const FViewportActionKeyInput& Action, bool& bOutIsInputCaptured, bool& bWasHandled);
 
diff --git a/Engine/Source/Editor/MeshPaint/Private/MeshPaintEdMode.cpp b/Engine/Source/Editor/MeshPaint/Private/MeshPaintEdMode.cpp
--- a/Engine/Source/Editor/MeshPaint/Private/MeshPaintEdMode.cpp
+++ b/Engine/Source/Editor/MeshPaint/Private/MeshPaintEdMode.cpp
@@ -32,10 +32,12 @@
 #include "MeshPaintAdapterFactory.h"
 #include "Components/SplineMeshComponent.h"
 
+#include "VREditorMode.h"
 #include "ViewportWorldInteraction.h"
 #include "ViewportInteractableInterface.h"
 #include "VREditorInteractor.h"
 #include "VIBaseTransformGizmo.h"
+#include "EditorWorldManager.h"
 
 #define LOCTEXT_NAMESPACE "MeshPaint_Mode"
 
@@ -201,18 +203,19 @@ void FEdModeMeshPaint::Enter()
 		ApplyOrRemoveForceBestLOD(/*bApply=*/ true);
 	}
 
-	// Register to find out about editor modes being activated
-	GetModeManager()->OnEditorModeChanged().AddRaw( this, &FEdModeMeshPaint::OnEditorModeChanged );
-
-	// Register to find out about VR input events
-	IVREditorMode* VREditorMode = static_cast<IVREditorMode*>( GetModeManager()->GetActiveMode( IVREditorModule::Get().GetVREditorModeID() ) );
-	if( VREditorMode != nullptr )
+	TSharedPtr<FEditorWorldWrapper> EditorWorld = GEditor->GetEditorWorldManager()->GetEditorWorldWrapper( GetWorld() );
+	if(EditorWorld.IsValid())
 	{
-		VREditorMode->GetWorldInteraction().OnViewportInteractionInputAction().RemoveAll( this );
-		VREditorMode->GetWorldInteraction().OnViewportInteractionInputAction().AddRaw( this, &FEdModeMeshPaint::OnVRAction );
+		// Register to find out about VR input events
+		UViewportWorldInteraction* WorldInteraction = EditorWorld->GetViewportWorldInteraction();
+		if(WorldInteraction != nullptr)
+		{
+			WorldInteraction->OnViewportInteractionInputAction().RemoveAll(this);
+			WorldInteraction->OnViewportInteractionInputAction().AddRaw(this, &FEdModeMeshPaint::OnVRAction);
 
-		// Hide the VR transform gizmo while we're in mesh paint mode.  It sort of gets in the way of painting.
-		VREditorMode->GetWorldInteraction().SetTransformGizmoVisible( false );
+			// Hide the VR transform gizmo while we're in mesh paint mode.  It sort of gets in the way of painting.
+			WorldInteraction->SetTransformGizmoVisible(false);
+		}
 	}
 }
 
@@ -221,19 +224,17 @@ void FEdModeMeshPaint::Exit()
 {
 	if( IVREditorModule::IsAvailable() )
 	{
-		IVREditorMode* VREditorMode = static_cast<IVREditorMode*>( GetModeManager()->GetActiveMode( IVREditorModule::Get().GetVREditorModeID() ) );
-		if( VREditorMode != nullptr )
+		UViewportWorldInteraction* ViewpportWorldInteraction = GEditor->GetEditorWorldManager()->GetEditorWorldWrapper(GetWorld())->GetViewportWorldInteraction();
+		if(ViewpportWorldInteraction != nullptr )
 		{
 			// Restore the transform gizmo visibility
-			VREditorMode->GetWorldInteraction().SetTransformGizmoVisible( true );
+			ViewpportWorldInteraction->SetTransformGizmoVisible( true );
 
 			// Unregister from event handlers
-			VREditorMode->GetWorldInteraction().OnViewportInteractionInputAction().RemoveAll( this );
+			ViewpportWorldInteraction->OnViewportInteractionInputAction().RemoveAll( this );
 		}
 	}
 
-	GetModeManager()->OnEditorModeChanged().RemoveAll( this );
-
 	//If we're painting vertex colors then propagate the painting done on LOD0 to all lower LODs. 
 	//Then stop forcing the LOD level of the mesh to LOD0.
 	if( FMeshPaintSettings::Get().ResourceType == EMeshPaintResource::VertexColors )
@@ -303,24 +304,6 @@ void FEdModeMeshPaint::Exit()
 }
 
 
-void FEdModeMeshPaint::OnEditorModeChanged( FEdMode* EditorMode, bool bEntered )
-{
-	// VR Editor mode may have gone away, so re-register for events in case it's a different object than the one we were originally bound to
-	check( EditorMode != nullptr );
-	if( bEntered && EditorMode->GetID() == IVREditorModule::Get().GetVREditorModeID() )
-	{
-		IVREditorMode* VREditorMode = static_cast<IVREditorMode*>( EditorMode );
-		if ( VREditorMode != nullptr )
-		{
-			VREditorMode->GetWorldInteraction().OnViewportInteractionInputAction().RemoveAll( this );
-			VREditorMode->GetWorldInteraction().OnViewportInteractionInputAction().AddRaw( this, &FEdModeMeshPaint::OnVRAction );
-		}
-
-		PaintingWithInteractorInVR = nullptr;
-	}
-}
-
-
 /** FEdMode: Called when the mouse is moved over the viewport */
 bool FEdModeMeshPaint::MouseMove( FEditorViewportClient* ViewportClient, FViewport* Viewport, int32 x, int32 y )
 {
@@ -2653,8 +2636,7 @@ void FEdModeMeshPaint::Render( const FSceneView* View, FViewport* Viewport, FPri
 		static TArray<FPaintRay> PaintRays;
 		PaintRays.Reset();
 		
-
-		IVREditorMode* VREditorMode = static_cast<IVREditorMode*>(GetModeManager()->GetActiveMode(IVREditorModule::Get().GetVREditorModeID()));
+		UVREditorMode* VREditorMode = GEditor->GetEditorWorldManager()->GetEditorWorldWrapper( Viewport->GetClient()->GetWorld() )->GetVREditorMode();
 
 		// Check to see if VREditorMode is active. If so, we're painting with interactor
 		bool bIsInVRMode = false;
@@ -3566,7 +3548,7 @@ EMeshPaintAction::Type FEdModeMeshPaint::GetPaintAction(FViewport* InViewport)
 		// When painting using VR, allow the modifier button to activate Erase mode
 		if( PaintingWithInteractorInVR )
 		{
-			IVREditorMode* VREditorMode = static_cast<IVREditorMode*>( GetModeManager()->GetActiveMode( IVREditorModule::Get().GetVREditorModeID() ) );
+			UVREditorMode* VREditorMode = GEditor->GetEditorWorldManager()->GetEditorWorldWrapper( GetWorld() )->GetVREditorMode();
 			UVREditorInteractor* VRInteractor = Cast<UVREditorInteractor>( PaintingWithInteractorInVR );
 			if( VREditorMode != nullptr && VRInteractor != nullptr )
 			{
@@ -5125,10 +5107,9 @@ TWeakObjectPtr<AActor> FEdModeMeshPaint::GetEditingActor() const
 void FEdModeMeshPaint::OnVRAction( class FEditorViewportClient& ViewportClient, UViewportInteractor* Interactor, 
 	const FViewportActionKeyInput& Action, bool& bOutIsInputCaptured, bool& bWasHandled )
 {
-	IVREditorMode* VREditorMode = static_cast<IVREditorMode*>( GetModeManager()->GetActiveMode( IVREditorModule::Get().GetVREditorModeID() ) );
+	UVREditorMode* VREditorMode = GEditor->GetEditorWorldManager()->GetEditorWorldWrapper(ViewportClient.GetWorld())->GetVREditorMode();
 	UVREditorInteractor* VRInteractor = Cast<UVREditorInteractor>( Interactor );
-
-	if( VREditorMode != nullptr && VRInteractor != nullptr )
+	if( VREditorMode != nullptr && VREditorMode->IsActive() && VRInteractor != nullptr )
 	{
 		if( Action.ActionType == ViewportWorldActionTypes::SelectAndMove_LightlyPressed )
 		{
diff --git a/Engine/Source/Editor/MeshPaint/Public/MeshPaintEdMode.h b/Engine/Source/Editor/MeshPaint/Public/MeshPaintEdMode.h
--- a/Engine/Source/Editor/MeshPaint/Public/MeshPaintEdMode.h
+++ b/Engine/Source/Editor/MeshPaint/Public/MeshPaintEdMode.h
@@ -400,10 +400,6 @@ public:
 	virtual void Tick( FEditorViewportClient* ViewportClient, float DeltaTime ) override;
 	// End of FEdMode interface
 
-
-	/** Called when an editor mode is entered or exited */
-	void OnEditorModeChanged( FEdMode* EditorMode, bool bEntered );
-
 	/** Returns true if we need to force a render/update through based fill/copy */
 	bool IsForceRendered (void) const;
 
diff --git a/Engine/Source/Editor/UnrealEd/Classes/Editor/EditorEngine.h b/Engine/Source/Editor/UnrealEd/Classes/Editor/EditorEngine.h
--- a/Engine/Source/Editor/UnrealEd/Classes/Editor/EditorEngine.h
+++ b/Engine/Source/Editor/UnrealEd/Classes/Editor/EditorEngine.h
@@ -1769,6 +1769,11 @@ public:
 	 */
 	TSharedRef<class FTimerManager> GetTimerManager() { return TimerManager.ToSharedRef(); }
 
+	/**
+	*  Returns the Editors world manager instance.
+	*/
+	TSharedRef<class FEditorWorldManager> GetEditorWorldManager() { return EditorWorldManager.ToSharedRef(); }
+
 	// Editor specific
 
 	/**
@@ -2771,6 +2776,9 @@ private:
 	/** The Timer manager for all timer delegates */
 	TSharedPtr<class FTimerManager> TimerManager;
 
+	/** Manager that holds all the worlds contexts with viewport interaction */
+	TSharedPtr<class FEditorWorldManager> EditorWorldManager;
+
 	/** The output log -> message log redirector for use during PIE */
 	TSharedPtr<class FOutputLogErrorsToMessageLogProxy> OutputLogErrorsToMessageLogProxyPtr;
 
diff --git a/Engine/Source/Editor/UnrealEd/Private/EditorEngine.cpp b/Engine/Source/Editor/UnrealEd/Private/EditorEngine.cpp
--- a/Engine/Source/Editor/UnrealEd/Private/EditorEngine.cpp
+++ b/Engine/Source/Editor/UnrealEd/Private/EditorEngine.cpp
@@ -60,6 +60,8 @@
 #include "EngineModule.h"
 #include "RendererInterface.h"
 
+#include "EditorWorldManager.h"
+
 #if PLATFORM_WINDOWS
 // For WAVEFORMATEXTENSIBLE
 	#include "AllowWindowsPlatformTypes.h"
@@ -549,6 +551,9 @@ void UEditorEngine::InitEditor(IEngineLoop* InEngineLoop)
 	// create the timer manager
 	TimerManager = MakeShareable(new FTimerManager());
 
+	// create the editor world manager
+	EditorWorldManager = MakeShareable(new FEditorWorldManager());
+
 	// Settings.
 	FBSPOps::GFastRebuild = 0;
 
@@ -1472,7 +1477,10 @@ void UEditorEngine::Tick( float DeltaSeconds, bool bIdleMode )
 			ViewportClient->Tick(DeltaSeconds);
 		}
 	}
-	
+
+	// Updates the ViewportWorldInteraction
+	EditorWorldManager->Tick( DeltaSeconds );
+
 	bool bIsMouseOverAnyLevelViewport = false;
 
 	//Do this check separate to the above loop as the ViewportClient may no longer be valid after we have ticked it
diff --git a/Engine/Source/Editor/UnrealEd/Private/EditorWorldManager.cpp b/Engine/Source/Editor/UnrealEd/Private/EditorWorldManager.cpp
new file mode 100644
--- /dev/null
+++ b/Engine/Source/Editor/UnrealEd/Private/EditorWorldManager.cpp
@@ -0,0 +1,134 @@
+// Copyright 1998-2016 Epic Games, Inc. All Rights Reserved.
+
+#include "UnrealEd.h"
+#include "EditorWorldManager.h"
+#include "ViewportWorldInteractionInterface.h"
+#include "ViewportWorldInteraction.h"
+#include "VREditorMode.h"
+#include "IViewportInteractionModule.h"
+
+/************************************************************************/
+/* FEditorWorldWrapper                                                  */
+/************************************************************************/
+FEditorWorldWrapper::FEditorWorldWrapper(FWorldContext& InWorldContext):
+	WorldContext (InWorldContext)
+{
+	WorldInteraction = NewObject<UViewportWorldInteraction>();
+	WorldInteraction->Init(InWorldContext.World());
+
+	VREditorMode = NewObject<UVREditorMode>();
+	VREditorMode->Init(WorldInteraction);
+}
+
+FEditorWorldWrapper::~FEditorWorldWrapper()
+{
+	if(VREditorMode != nullptr && !VREditorMode->IsPendingKill() )
+	{
+		VREditorMode->Shutdown();
+		VREditorMode->MarkPendingKill();
+		VREditorMode = nullptr;
+	}
+
+	if(WorldInteraction != nullptr && !WorldInteraction->IsPendingKill())
+	{
+		WorldInteraction->Shutdown();
+		WorldInteraction->MarkPendingKill();
+		WorldInteraction = nullptr;
+	}
+}
+
+UWorld* FEditorWorldWrapper::GetWorld() const
+{
+	return WorldContext.World();
+}
+
+FWorldContext& FEditorWorldWrapper::GetWorldContext() const
+{
+	return WorldContext;
+}
+
+UViewportWorldInteraction* FEditorWorldWrapper::GetViewportWorldInteraction() const
+{
+	return WorldInteraction;
+}
+
+UVREditorMode* FEditorWorldWrapper::GetVREditorMode() const
+{
+	return VREditorMode;
+}
+
+void FEditorWorldWrapper::AddReferencedObjects( FReferenceCollector& Collector )
+{
+	Collector.AddReferencedObject( VREditorMode );
+	Collector.AddReferencedObject( WorldInteraction );
+}
+
+/************************************************************************/
+/* FEditorWorldManager                                                  */
+/************************************************************************/
+FEditorWorldManager::FEditorWorldManager()
+{
+	if(GEngine)
+	{
+		GEngine->OnWorldContextDestroyed().AddRaw(this, &FEditorWorldManager::OnWorldContextRemove);
+	}
+}
+
+FEditorWorldManager::~FEditorWorldManager()
+{
+	if(GEngine)
+	{
+		GEngine->OnWorldContextDestroyed().RemoveAll(this);
+	}
+
+	EditorWorldMap.Empty();
+}
+
+TSharedPtr<FEditorWorldWrapper> FEditorWorldManager::GetEditorWorldWrapper(const UWorld* InWorld)
+{
+	// Try to find this world in the map and return it or create and add one if nothing found
+	TSharedPtr<FEditorWorldWrapper> Result;
+	if(InWorld)
+	{
+		TSharedPtr<FEditorWorldWrapper>* FoundWorld = EditorWorldMap.Find(InWorld->GetUniqueID());
+		if(FoundWorld != nullptr)
+		{
+			Result = *FoundWorld;
+		}
+		else
+		{
+			FWorldContext* WorldContext = GEditor->GetWorldContextFromWorld(InWorld);
+			Result = OnWorldContextAdd(*WorldContext);
+		}
+	}
+	return Result;
+}
+
+TSharedPtr<FEditorWorldWrapper> FEditorWorldManager::OnWorldContextAdd(FWorldContext& InWorldContext)
+{
+	//Only add editor type world to the map
+	UWorld* World = InWorldContext.World();
+	TSharedPtr<FEditorWorldWrapper> Result;
+	if(World && (InWorldContext.WorldType == EWorldType::Editor || InWorldContext.WorldType == EWorldType::Preview))
+	{
+		TSharedPtr<FEditorWorldWrapper> EditorWorld(new FEditorWorldWrapper(InWorldContext));
+		Result = EditorWorld;
+		EditorWorldMap.Add(World->GetUniqueID(), Result);
+	}
+	return Result;
+}
+
+void FEditorWorldManager::OnWorldContextRemove(FWorldContext& InWorldContext)
+{
+	UWorld* World = InWorldContext.World();
+	if(World)
+	{
+		EditorWorldMap.Remove(World->GetUniqueID());
+	}
+}
+
+void FEditorWorldManager::Tick( float DeltaTime )
+{
+	IViewportInteractionModule& ViewportInteraction = IViewportInteractionModule::Get();
+	ViewportInteraction.Tick( DeltaTime );
+}
\ No newline at end of file
diff --git a/Engine/Source/Editor/UnrealEd/Public/EditorWorldManager.h b/Engine/Source/Editor/UnrealEd/Public/EditorWorldManager.h
new file mode 100644
--- /dev/null
+++ b/Engine/Source/Editor/UnrealEd/Public/EditorWorldManager.h
@@ -0,0 +1,79 @@
+// Copyright 1998-2016 Epic Games, Inc. All Rights Reserved.
+
+#pragma once
+
+#include "Core.h"
+
+class UViewportWorldInteraction;
+class UVREditorMode;
+
+/**
+ * Holds ViewportWorldInteraction for every world with corresponding VR Editor mode
+ */
+class UNREALED_API FEditorWorldWrapper : public FGCObject
+{
+public:
+
+	/** Constructor creates ViewportWorldInteration and VREditorMode with the passed world */
+	FEditorWorldWrapper(FWorldContext& InWorldContext);
+
+	/** Default destructor */
+	virtual ~FEditorWorldWrapper();
+
+	/** Gets the world of the world context */
+	UWorld* GetWorld() const;
+
+	/** Gets the world context */
+	FWorldContext& GetWorldContext() const;
+
+	/** Gets the viewport world interaction */
+	UViewportWorldInteraction* GetViewportWorldInteraction() const;
+
+	/** Gets the VR Editor mode */
+	UVREditorMode* GetVREditorMode() const;
+
+	/** FGCObject */
+	virtual void AddReferencedObjects( FReferenceCollector& Collector ) override;
+
+private:
+
+	/** The world context */
+	FWorldContext& WorldContext;
+
+	/** The viewport world interaction */
+	UViewportWorldInteraction* WorldInteraction;
+
+	/** The VR Editor mode */
+	UVREditorMode* VREditorMode;
+};
+
+/**
+ * Has a map of FEditorWorldWrapper
+ */
+class UNREALED_API FEditorWorldManager	
+{
+public:
+
+	/** Default constructor */
+	FEditorWorldManager();
+
+	/** Default destructor */
+	virtual ~FEditorWorldManager();
+
+	/** Gets the editor world wrapper that is found with the world passed.
+	 * Adds one for this world if there was non found. */
+	TSharedPtr<FEditorWorldWrapper> GetEditorWorldWrapper(const UWorld* InWorld);
+
+	void Tick( float DeltaTime );
+
+private:
+
+	/** Adds a new editor world wrapper when a new world context was created */
+	TSharedPtr<FEditorWorldWrapper> OnWorldContextAdd(FWorldContext& InWorldContext);
+
+	/** Adds a new editor world wrapper when a new world context was created */
+	void OnWorldContextRemove(FWorldContext& InWorldContext);
+
+	/** Map of all the editor world maps */
+	TMap<uint32, TSharedPtr<FEditorWorldWrapper>> EditorWorldMap;
+};
\ No newline at end of file
diff --git a/Engine/Source/Editor/UnrealEd/UnrealEd.Build.cs b/Engine/Source/Editor/UnrealEd/UnrealEd.Build.cs
--- a/Engine/Source/Editor/UnrealEd/UnrealEd.Build.cs
+++ b/Engine/Source/Editor/UnrealEd/UnrealEd.Build.cs
@@ -61,6 +61,8 @@ public class UnrealEd : ModuleRules
                 "PortalRpc",
                 "PortalServices",
                 "BlueprintNativeCodeGen",
+                "ViewportInteraction",
+                "VREditor"
 			}
 		);
 
@@ -144,6 +146,8 @@ public class UnrealEd : ModuleRules
                 "PixelInspectorModule",
 				"MovieScene",
 				"MovieSceneTracks",
+                "ViewportInteraction",
+                "VREditor"
             }
 		);
 
@@ -215,7 +219,6 @@ public class UnrealEd : ModuleRules
                 "PortalServices",
                 "GeometryCacheEd",
                 "BlueprintNativeCodeGen",
-				"VREditor"
 			}
 		);
 
diff --git a/Engine/Source/Editor/VREditor/Public/IVREditorMode.h b/Engine/Source/Editor/VREditor/Public/IVREditorMode.h
deleted file mode 100644
diff --git a/Engine/Source/Editor/VREditor/Public/IVREditorModule.h b/Engine/Source/Editor/VREditor/Public/IVREditorModule.h
--- a/Engine/Source/Editor/VREditor/Public/IVREditorModule.h
+++ b/Engine/Source/Editor/VREditor/Public/IVREditorModule.h
@@ -3,7 +3,6 @@
 #pragma once
 
 #include "ModuleManager.h"
-#include "IVREditorMode.h"
 
 /**
  * The public interface to this module
@@ -34,13 +33,6 @@ public:
 		return FModuleManager::Get().IsModuleLoaded( "VREditor" );
 	}
 
-	/**
-	 * Gets the editor mode ID for the VREditor mode
-	 *
-	 * @return	Editor mode ID (for FEditorModeTools) of the VREditor mode
-	 */
-	virtual FEditorModeID GetVREditorModeID() const = 0;
-
 	/**
 	 * Checks whether or not editor VR features are enabled
 	 *
@@ -64,5 +56,12 @@ public:
 	 * @param	bForceWithoutHMD	If set to true, will enter VR mode without switching to HMD/stereo.  This can be useful for testing.
 	 */
 	virtual void EnableVREditor( const bool bEnable, const bool bForceWithoutHMD = false ) = 0;
+
+	/** 
+	 * Check if the VR Editor is currently running
+	 *
+	 * @return True if the VREditor is currently running 
+	 */
+	virtual bool IsVREditorModeActive() = 0;
 };
 
diff --git a/Engine/Source/Editor/VREditor/Public/VREditorInteractor.h b/Engine/Source/Editor/VREditor/Public/VREditorInteractor.h
--- a/Engine/Source/Editor/VREditor/Public/VREditorInteractor.h
+++ b/Engine/Source/Editor/VREditor/Public/VREditorInteractor.h
@@ -2,7 +2,7 @@
 
 #pragma once
 
-#include "Editor/ViewportInteraction/Public/ViewportInteractor.h"
+#include "ViewportInteractor.h"
 #include "VREditorInteractor.generated.h"
 
 /**
@@ -18,16 +18,16 @@ public:
 	virtual ~UVREditorInteractor();
 
 	/** Initialize default values */
-	virtual void Init( class FVREditorMode* InVRMode );
+	virtual void Init( class UVREditorMode* InVRMode );
 
 	/** Gets the owner of this system */
-	class FVREditorMode& GetVRMode()
+	class UVREditorMode& GetVRMode()
 	{
 		return *VRMode;
 	}
 
 	/** Gets the owner of this system (const) */
-	const class FVREditorMode& GetVRMode() const
+	const class UVREditorMode& GetVRMode() const
 	{
 		return *VRMode;
 	}
@@ -107,7 +107,8 @@ public:
 protected:
 
 	/** The mode that owns this interactor */
-	class FVREditorMode* VRMode;
+	UPROPERTY()
+	class UVREditorMode* VRMode;
 
 	//
 	// General input @todo: VREditor: Should this be general (non-UI) in interactordata ?
diff --git a/Engine/Source/Editor/VREditor/Public/VREditorMode.h b/Engine/Source/Editor/VREditor/Public/VREditorMode.h
new file mode 100644
--- /dev/null
+++ b/Engine/Source/Editor/VREditor/Public/VREditorMode.h
@@ -0,0 +1,429 @@
+// Copyright 1998-2015 Epic Games, Inc. All Rights Reserved.
+
+#pragma once
+
+#include "IVREditorModule.h"
+#include "HeadMountedDisplayTypes.h"					// For EHMDDeviceType::Type
+#include "Kismet/HeadMountedDisplayFunctionLibrary.h"	// For EHMDTrackingOrigin::Type
+#include "VREditorMode.generated.h"
+
+// Forward declare the GizmoHandleTypes that is defined in VIBaseTransformGizmo.h
+enum class EGizmoHandleTypes : uint8;
+
+/** Reason for exiting the VR editor */
+enum EVREditorExitType
+{
+	Normal = 0,
+	PIE_VR = 1,		//Play in editor with VR
+	SIE_VR = 2		//Simulate in editor with VR
+};
+
+/**
+* Types of actions that can be performed with VR controller devices
+*/
+namespace VRActionTypes
+{
+	static const FName Touch( "Touch" );
+	static const FName Modifier( "Modifier" );
+	static const FName ConfirmRadialSelection( "ConfirmRadialSelection" );
+	static const FName TrackpadPositionX( "TrackpadPositionX" );
+	static const FName TrackpadPositionY( "TrackpadPositionY" );
+	static const FName TriggerAxis( "TriggerAxis" );
+}
+
+/**
+ * VR Editor Mode.  Extends editor viewports with functionality for VR controls and object manipulation
+ */
+UCLASS()
+class VREDITOR_API UVREditorMode : public UObject
+{
+
+	GENERATED_UCLASS_BODY()
+
+public:
+
+	/** Cleans up this mode, called when the editor is shutting down */
+	virtual ~UVREditorMode();
+
+	/** Static: Sets whether we should actually use an HMD.  Call this before activating VR mode */
+	void SetActuallyUsingVR( const bool bShouldActuallyUseVR )
+	{
+		bActuallyUsingVR = bShouldActuallyUseVR;
+	}
+
+	/** Returns true if we're actually using VR, or false if we're faking it */
+	bool IsActuallyUsingVR()
+	{
+		return bActuallyUsingVR;
+	}
+
+	/** Returns true if the user wants to exit this mode */
+	bool WantsToExitMode() const
+	{
+		return bWantsToExitMode;
+	}
+
+	/** Call this to start exiting VR mode */
+	void StartExitingVRMode( const EVREditorExitType InExitType = EVREditorExitType::Normal );
+
+	/** Gets the reason for exiting the mode */
+	EVREditorExitType GetExitType() const
+	{
+		return ExitType;
+	}
+
+	/** EditorWorldManager initializes the start of the editor */
+	void Init( class UViewportWorldInteraction* InViewportWorldInteraction );
+
+	/** EditorWorldManager shuts down the VREditor when closing the editor */
+	void Shutdown();
+
+	/** When the user actually enters the VR Editor mode */
+	void Enter();
+
+	/** When the user leaves the VR Editor mode */
+	void Exit();
+
+	/** Tick before the ViewportWorldInteraction is ticked */
+	void PreTick( const float DeltaTime );
+
+	/** Tick after the ViewportWorldInteraction is ticked */
+	void Tick( const float DeltaTime );
+	
+	//bool InputKey( FEditorViewportClient* ViewportClient, FViewport* Viewport, FKey Key, EInputEvent Event );
+	//bool InputAxis( FEditorViewportClient* ViewportClient, FViewport* Viewport, int32 ControllerId, FKey Key, float Delta, float DeltaTime );
+	//bool IsCompatibleWith( FEditorModeID OtherModeID ) const;
+	//void Render( const FSceneView* SceneView, FViewport* Viewport, FPrimitiveDrawInterface* PDI );	
+
+	/**
+	 * Gets our avatar's mesh actor
+	 *
+	 * @return	The mesh actor
+	 */
+	AActor* GetAvatarMeshActor();
+
+	/** Gets the World from the ViewportWorldInteraction used by this mode */
+	UWorld* GetWorld() const;
+
+	/** Gets the world space transform of the calibrated VR room origin.  When using a seated VR device, this will feel like the
+	camera's world transform (before any HMD positional or rotation adjustments are applied.) */
+	FTransform GetRoomTransform() const;
+
+	/** Sets a new transform for the room, in world space.  This is basically setting the editor's camera transform for the viewport */
+	void SetRoomTransform( const FTransform& NewRoomTransform );
+
+	/** Gets the transform of the user's HMD in room space */
+	FTransform GetRoomSpaceHeadTransform() const;
+
+	/**
+	 * Gets the world space transform of the HMD (head)
+	 *
+	 * @return	World space space HMD transform
+	 */
+	FTransform GetHeadTransform() const;
+
+	/** Gets access to the world interaction system (const) */
+	const class UViewportWorldInteraction& GetWorldInteraction() const;
+
+	/** Gets access to the world interaction system */
+	class UViewportWorldInteraction& GetWorldInteraction();
+
+	/** If the mode was completely initialized */
+	bool IsFullyInitialized() const;
+
+	/** If the mode is currently running */
+	bool IsActive() const;
+	
+	/** * Gets the tick handle to give external systems the change to be ticked right after the ViewportWorldInteraction is ticked */
+	DECLARE_EVENT_OneParam(UVREditorMode, FOnVRTickHandle, const float /* DeltaTime */);
+	FOnVRTickHandle& OnTickHandle()
+	{
+		return TickHandle;
+	}
+
+	/** Returns the Unreal controller ID for the motion controllers we're using */
+	int32 GetMotionControllerID() const
+	{
+		return MotionControllerID;
+	}
+
+	/** Returns whether or not the flashlight is visible */
+	bool IsFlashlightOn() const
+	{
+		return bIsFlashlightOn;
+	}
+
+	/** Returns the time since the VR Editor mode was last entered */
+	inline FTimespan GetTimeSinceModeEntered() const
+	{
+		return FTimespan::FromSeconds( FApp::GetCurrentTime() ) - AppTimeModeEntered;
+	}
+
+	// @todo vreditor: Move this to FMath so we can use it everywhere
+	// NOTE: OvershootAmount is usually between 0.5 and 2.0, but can go lower and higher for extreme overshots
+	template<class T>
+	static T OvershootEaseOut( T Alpha, const float OvershootAmount = 1.0f )
+	{
+		Alpha--;
+		return 1.0f - ( ( Alpha * ( ( OvershootAmount + 1 ) * Alpha + OvershootAmount ) + 1 ) - 1.0f );
+	}
+
+	/** Spawns a transient actor that we can use in the editor world (templated for convenience) */
+	template<class T>
+	T* SpawnTransientSceneActor( const FString& ActorName, const bool bWithSceneComponent ) const
+	{
+		return CastChecked<T>( SpawnTransientSceneActor( T::StaticClass(), ActorName, bWithSceneComponent ) );
+	}
+
+	/** Spawns a transient actor that we can use in the editor world */
+	AActor* SpawnTransientSceneActor( TSubclassOf<class AActor> ActorClass, const FString& ActorName, const bool bWithSceneComponent ) const;
+
+	/** Destroys a transient actor we created earlier and nulls out the pointer */
+	void DestroyTransientActor( AActor* Actor ) const;
+
+	/** Gets access to the VR UI system (const) */
+	const class UVREditorUISystem& GetUISystem() const
+	{
+		return *UISystem;
+	}
+
+	/** Gets access to the VR UI system */
+	class UVREditorUISystem& GetUISystem()
+	{
+		return *UISystem;
+	}
+
+	/** Gets the viewport that VR Mode is activated in.  Even though editor modes are available in all
+	    level viewports simultaneously, only one viewport is "possessed" by the HMD.  Generally try to avoid using
+		this function and instead use the ViewportClient that is passed around through various FEdMode overrides */
+	const class SLevelViewport& GetLevelViewportPossessedForVR() const;
+
+	/** Mutable version of above. */
+	class SLevelViewport& GetLevelViewportPossessedForVR();
+
+	/** Gets the world scale factor, which can be multiplied by a scale vector to convert to room space */
+	float GetWorldScaleFactor() const;
+
+	/** Sets up our avatar mesh, if not already spawned */
+	void SpawnAvatarMeshActor();
+
+	/** Called internally when the user changes maps, enters/exits PIE or SIE, or switched between PIE/SIE */
+	void CleanUpActorsBeforeMapChangeOrSimulate();
+
+	/** Spawns a flashlight on the specified hand */
+	void ToggleFlashlight( class UVREditorInteractor* Interactor );
+
+	/** Will update the TransformGizmo Actor with the next Gizmo type  */
+	void CycleTransformGizmoHandleType();
+
+	/** Gets the current Gizmo handle type */
+	EGizmoHandleTypes GetCurrentGizmoType() const;
+
+	/** @return Returns the type of HMD we're dealing with */
+	EHMDDeviceType::Type GetHMDDeviceType() const;
+
+	/** @return Checks to see if the specified hand is aiming roughly toward the specified capsule */
+	bool IsHandAimingTowardsCapsule( class UViewportInteractor* Interactor, const FTransform& CapsuleTransform, const FVector CapsuleStart, const FVector CapsuleEnd, const float CapsuleRadius, const float MinDistanceToCapsule, const FVector CapsuleFrontDirection, const float MinDotForAimingAtCapsule ) const;
+	
+	/** Gets the hand interactor  */
+	class UVREditorInteractor* GetHandInteractor( const EControllerHand ControllerHand ) const;
+
+	/** Stops the haptic feedback for the left and right hand interactors for the oculus */
+	void StopOldHapticEffects();
+
+	/** Snaps the current selected actor to the ground */
+	void SnapSelectedActorsToGround();
+
+	void GoToDefaultScale();
+
+private:
+	//Handles closing the VR mode by escape key
+	void InputKey(const FEditorViewportClient& InViewportClient, const FKey InKey, const EInputEvent InEvent, bool& bOutWasHandled);
+
+	/** Called when someone closes a standalone VR Editor window */
+	void OnVREditorWindowClosed( const TSharedRef<SWindow>& ClosedWindow );
+
+	/** FEditorDelegates callbacks */
+	void OnBeginPIE( const bool bIsSimulatingInEditor );
+	void OnEndPIE( const bool bIsSimulatingInEditor );
+	void OnSwitchBetweenPIEAndSIE( const bool bIsSimulatingInEditor );
+
+protected:
+
+	//
+	// Startup/Shutdown
+	//
+
+	/** The VR editor window, if it's open right now */
+	TWeakPtr< class SWindow > VREditorWindowWeakPtr;
+
+	/** The VR level viewport, if we're in VR mode */
+	TWeakPtr< class SLevelViewport > VREditorLevelViewportWeakPtr;
+
+	/** Saved information about the editor and viewport we possessed, so we can restore it after exiting VR mode */
+	struct FSavedEditorState
+	{
+		ELevelViewportType ViewportType;
+		FVector ViewLocation;
+		FRotator ViewRotation;
+		FEngineShowFlags ShowFlags;
+		bool bLockedPitch;
+		bool bGameView;
+		bool bAlwaysShowModeWidgetAfterSelectionChanges;
+		float NearClipPlane;
+		bool bRealTime;
+		float DragTriggerDistance;
+		bool bOnScreenMessages;
+		EHMDTrackingOrigin::Type TrackingOrigin;
+
+		FSavedEditorState()
+			: ViewportType( LVT_Perspective ),
+			  ViewLocation( FVector::ZeroVector ),
+			  ViewRotation( FRotator::ZeroRotator ),
+			  ShowFlags( ESFIM_Editor ),
+			  bLockedPitch( false ),
+			  bGameView( false ),
+			  bAlwaysShowModeWidgetAfterSelectionChanges( false ),
+			  NearClipPlane( 0.0f ),
+			  bRealTime( false ),
+			  DragTriggerDistance( 0.0f ),
+			  bOnScreenMessages( false ),
+			  TrackingOrigin( EHMDTrackingOrigin::Eye )
+		{
+		}
+		
+	} SavedEditorState;
+
+	/** True if we're in using an actual HMD in this mode, or false if we're "faking" VR mode for testing */
+	bool bActuallyUsingVR;
+
+	/** True if we currently want to exit VR mode.  This is used to defer exiting until it is safe to do that */
+	bool bWantsToExitMode;
+
+	/** True if VR mode is fully initialized and ready to render */
+	bool bIsFullyInitialized;
+
+	/** App time that we entered this mode */
+	FTimespan AppTimeModeEntered;
+
+
+	//
+	// Avatar visuals
+	//
+
+	/** Actor with components to represent the VR avatar in the world, including motion controller meshes */
+	class AVREditorAvatarActor* AvatarActor;
+
+
+	//
+	// Flashlight
+	//
+
+	/** Spotlight for the flashlight */
+	class USpotLightComponent* FlashlightComponent;
+
+	/** If there is currently a flashlight in the scene */
+	bool bIsFlashlightOn;
+
+	//
+	// Input
+	//
+
+	/** The Unreal controller ID for the motion controllers we're using */
+	int32 MotionControllerID;
+
+
+	//
+	// Subsystems registered
+	//
+
+	FOnVRTickHandle TickHandle;
+
+	//
+	// Subsystems
+	//
+
+	/** VR UI system */
+	UPROPERTY()
+	class UVREditorUISystem* UISystem;
+
+	/** Teleporter system */
+	UPROPERTY()
+	class UVREditorTeleporter* TeleporterSystem;
+
+	//
+	// World interaction
+	//
+
+	/** World interaction manager */
+	UPROPERTY()
+	class UViewportWorldInteraction* WorldInteraction;
+
+	/** The current Gizmo type that is used for the TransformGizmo Actor */
+	EGizmoHandleTypes CurrentGizmoType;
+
+	UPROPERTY()
+	class UVREditorWorldInteraction* VRWorldInteractionExtension;
+
+	//
+	// Interactors
+	//
+
+	/** The mouse cursor interactor (currently only available when not in VR) */
+	UPROPERTY()
+	class UMouseCursorInteractor* MouseCursorInteractor;
+
+	/** The right motion controller */
+	UPROPERTY()
+	class UVREditorMotionControllerInteractor* LeftHandInteractor; //@todo vreditor: Hardcoded interactors
+	
+	/** The right motion controller */
+	UPROPERTY()
+	class UVREditorMotionControllerInteractor* RightHandInteractor; 
+
+	//
+	// Colors
+	//
+public:
+	// Color types
+	enum EColors
+	{
+		DefaultColor,
+		SelectionColor,
+		TeleportColor,
+		WorldDraggingColor_OneHand,
+		WorldDraggingColor_TwoHands,
+		RedGizmoColor,
+		GreenGizmoColor,
+		BlueGizmoColor,
+		WhiteGizmoColor,
+		HoverGizmoColor,
+		DraggingGizmoColor,
+		UISelectionBarColor,
+		UISelectionBarHoverColor,
+		UICloseButtonColor,
+		UICloseButtonHoverColor,
+		TotalCount
+	};
+
+	// Gets the color
+	FLinearColor GetColor( const EColors Color ) const;
+
+private:
+
+	// All the colors for this mode
+	TArray<FLinearColor> Colors;
+
+	/** If this is the first tick or before */
+	bool bFirstTick;
+
+	/** If the coordinate system was in world space before entering (local) scale mode */
+	bool bWasInWorldSpaceBeforeScaleMode;
+
+	/** The reason for exiting the mode, so the module can close the mode and take further actions on what should happen next */
+	EVREditorExitType ExitType;
+
+	/** If this current mode is running */
+	bool bIsActive;
+};
+
diff --git a/Engine/Source/Editor/VREditor/Teleporter/VREditorTeleporter.cpp b/Engine/Source/Editor/VREditor/Teleporter/VREditorTeleporter.cpp
--- a/Engine/Source/Editor/VREditor/Teleporter/VREditorTeleporter.cpp
+++ b/Engine/Source/Editor/VREditor/Teleporter/VREditorTeleporter.cpp
@@ -2,7 +2,7 @@
 
 #include "VREditorModule.h"
 #include "VREditorTeleporter.h"
-#include "IVREditorMode.h"
+#include "VREditorMode.h"
 #include "ViewportWorldInteraction.h"
 #include "VREditorInteractor.h"
 
@@ -29,7 +29,7 @@ UVREditorTeleporter::UVREditorTeleporter( const FObjectInitializer& Initializer
 	check( TeleportSound != nullptr );
 }
 
-void UVREditorTeleporter::Init( IVREditorMode* InMode )
+void UVREditorTeleporter::Init( UVREditorMode* InMode )
 {
 	VRMode = InMode;
 	VRMode->OnTickHandle().AddUObject( this, &UVREditorTeleporter::Tick );
@@ -41,6 +41,7 @@ void UVREditorTeleporter::Shutdown()
 	VRMode->OnTickHandle().RemoveAll( this );
 	VRMode->GetWorldInteraction().OnViewportInteractionInputAction().RemoveAll( this );
 	TeleportSound = nullptr;
+	VRMode = nullptr;
 }
 
 void UVREditorTeleporter::Tick( const float DeltaTime )
diff --git a/Engine/Source/Editor/VREditor/Teleporter/VREditorTeleporter.h b/Engine/Source/Editor/VREditor/Teleporter/VREditorTeleporter.h
--- a/Engine/Source/Editor/VREditor/Teleporter/VREditorTeleporter.h
+++ b/Engine/Source/Editor/VREditor/Teleporter/VREditorTeleporter.h
@@ -15,7 +15,7 @@ class UVREditorTeleporter : public UObject
 public:
 
 	/** Initializes the teleport system */
-	void Init( class IVREditorMode* InMode );
+	void Init( class UVREditorMode* InMode );
 
 	/** Shuts down the teleport system */
 	void Shutdown();
@@ -46,7 +46,8 @@ private:
 	//
 
 	/** Owning mode */
-	IVREditorMode* VRMode;
+	UPROPERTY()
+	UVREditorMode* VRMode;
 
 	/** If we are currently teleporting */
 	bool bIsTeleporting;
diff --git a/Engine/Source/Editor/VREditor/VREditorAvatarActor.cpp b/Engine/Source/Editor/VREditor/VREditorAvatarActor.cpp
--- a/Engine/Source/Editor/VREditor/VREditorAvatarActor.cpp
+++ b/Engine/Source/Editor/VREditor/VREditorAvatarActor.cpp
@@ -34,6 +34,12 @@ AVREditorAvatarActor::AVREditorAvatarActor( const FObjectInitializer& ObjectInit
 	bIsDrawingWorldMovementPostProcess( false ),
 	VRMode( nullptr )
 {
+	{
+		USceneComponent* SceneRootComponent = CreateDefaultSubobject<USceneComponent>( TEXT( "RootComponent" ) );
+		AddOwnedComponent( SceneRootComponent );
+		SetRootComponent( SceneRootComponent );
+	}
+
 	// Give us a head mesh
 	{
 		HeadMeshComponent = CreateDefaultSubobject<UStaticMeshComponent>( TEXT( "HeadMeshComponent" ) );
@@ -195,23 +201,23 @@ AVREditorAvatarActor::~AVREditorAvatarActor()
 	UserScaleIndicatorText = nullptr;
 }
 
-void AVREditorAvatarActor::Init( FVREditorMode* InVRMode  )
+void AVREditorAvatarActor::Init( UVREditorMode* InVRMode  )
 {
 	VRMode = InVRMode;
 
 	// Set the default color for the progress bar
 	{
 		static FName StaticLaserColorName( "LaserColor" );
-		const FLinearColor FixedProgressbarColor = VRMode->GetColor( FVREditorMode::EColors::WorldDraggingColor_TwoHands );
+		const FLinearColor FixedProgressbarColor = VRMode->GetColor( UVREditorMode::EColors::WorldDraggingColor_TwoHands );
 		FixedUserScaleMID->SetVectorParameterValue( StaticLaserColorName, FixedProgressbarColor );
 		TranslucentFixedUserScaleMID->SetVectorParameterValue( StaticLaserColorName, FixedProgressbarColor );
 
-		const FLinearColor CurrentProgressbarColor = VRMode->GetColor( FVREditorMode::EColors::GreenGizmoColor );
+		const FLinearColor CurrentProgressbarColor = VRMode->GetColor( UVREditorMode::EColors::GreenGizmoColor );
 		CurrentUserScaleMID->SetVectorParameterValue( StaticLaserColorName, CurrentProgressbarColor );
 		TranslucentCurrentUserScaleMID->SetVectorParameterValue( StaticLaserColorName, CurrentProgressbarColor );
 	}
 
-	UserScaleIndicatorText->SetTextRenderColor( VRMode->GetColor( FVREditorMode::EColors::WorldDraggingColor_TwoHands ).ToFColor( false ) );
+	UserScaleIndicatorText->SetTextRenderColor( VRMode->GetColor( UVREditorMode::EColors::WorldDraggingColor_TwoHands ).ToFColor( false ) );
 
 	// Tell the grid to stay relative to the rootcomponent
 	const FAttachmentTransformRules AttachmentTransformRules = FAttachmentTransformRules( EAttachmentRule::KeepRelative, true );
diff --git a/Engine/Source/Editor/VREditor/VREditorAvatarActor.h b/Engine/Source/Editor/VREditor/VREditorAvatarActor.h
--- a/Engine/Source/Editor/VREditor/VREditorAvatarActor.h
+++ b/Engine/Source/Editor/VREditor/VREditorAvatarActor.h
@@ -18,7 +18,7 @@ public:
 	~AVREditorAvatarActor();
 
 	/** Called by VREditorMode::Enter to initialize all post constructor components and to set the VRMode */
-	void Init( class FVREditorMode* InVRMode );
+	void Init( class UVREditorMode* InVRMode );
 
 	/** Called by VREditorMode to update us every frame */
 	void TickManually( const float DeltaTime );
@@ -94,5 +94,6 @@ private:
 	class UPostProcessComponent* PostProcessComponent;
 
 	/** Owning object */
-	class FVREditorMode* VRMode;
+	UPROPERTY()
+	class UVREditorMode* VRMode;
 };
\ No newline at end of file
diff --git a/Engine/Source/Editor/VREditor/VREditorBaseActor.cpp b/Engine/Source/Editor/VREditor/VREditorBaseActor.cpp
--- a/Engine/Source/Editor/VREditor/VREditorBaseActor.cpp
+++ b/Engine/Source/Editor/VREditor/VREditorBaseActor.cpp
@@ -27,9 +27,9 @@ AVREditorBaseActor::AVREditorBaseActor()
 
 }
 
-void AVREditorBaseActor::SetVRMode( FVREditorMode* InOwner )
+void AVREditorBaseActor::SetVRMode( UVREditorMode* InVRMode )
 {
-	Owner = InOwner;
+	VRMode = InVRMode;
 }
 
 void AVREditorBaseActor::TickManually( float DeltaTime )
@@ -87,7 +87,7 @@ void AVREditorBaseActor::StopMoveTo()
 
 FTransform AVREditorBaseActor::MakeUITransformLockedToHand( UViewportInteractor* Interactor, const bool bOnArm, const FVector& InRelativeOffset, const FRotator& InLocalRotation )
 {
-	const float WorldScaleFactor = Owner->GetWorldScaleFactor();
+	const float WorldScaleFactor = VRMode->GetWorldScaleFactor();
 
 	FTransform UIToHandTransform( InLocalRotation, InRelativeOffset * WorldScaleFactor );
 	if ( !bOnArm )
@@ -111,22 +111,22 @@ void AVREditorBaseActor::UpdateTransformIfDocked()
 		if ( DockedTo == EDockedTo::LeftHand )
 		{
 			const bool bOnArm = false;
-			NewTransform = MakeUITransformLockedToHand( Owner->GetHandInteractor( EControllerHand::Left ), bOnArm );
+			NewTransform = MakeUITransformLockedToHand( VRMode->GetHandInteractor( EControllerHand::Left ), bOnArm );
 		}
 		else if ( DockedTo == EDockedTo::RightHand )
 		{
 			const bool bOnArm = false;
-			NewTransform = MakeUITransformLockedToHand( Owner->GetHandInteractor( EControllerHand::Right ), bOnArm );
+			NewTransform = MakeUITransformLockedToHand( VRMode->GetHandInteractor( EControllerHand::Right ), bOnArm );
 		}
 		else if ( DockedTo == EDockedTo::LeftArm )
 		{
 			const bool bOnArm = true;
-			NewTransform = MakeUITransformLockedToHand( Owner->GetHandInteractor( EControllerHand::Left ), bOnArm );
+			NewTransform = MakeUITransformLockedToHand( VRMode->GetHandInteractor( EControllerHand::Left ), bOnArm );
 		}
 		else if ( DockedTo == EDockedTo::RightArm )
 		{
 			const bool bOnArm = true;
-			NewTransform = MakeUITransformLockedToHand( Owner->GetHandInteractor( EControllerHand::Right ), bOnArm );
+			NewTransform = MakeUITransformLockedToHand( VRMode->GetHandInteractor( EControllerHand::Right ), bOnArm );
 		}
 		else if ( DockedTo == EDockedTo::Room )
 		{
@@ -152,11 +152,11 @@ FTransform AVREditorBaseActor::MakeUITransformLockedToHand( UViewportInteractor*
 
 FTransform AVREditorBaseActor::MakeUITransformLockedToRoom()
 {
-	const float WorldScaleFactor = Owner->GetWorldScaleFactor();
+	const float WorldScaleFactor = VRMode->GetWorldScaleFactor();
 
 	const FTransform UIToRoomTransform( LocalRotation, RelativeOffset * WorldScaleFactor );
 
-	const FTransform RoomToWorldTransform = Owner->GetRoomTransform();
+	const FTransform RoomToWorldTransform = VRMode->GetRoomTransform();
 
 	FTransform UIToWorldTransform = UIToRoomTransform * RoomToWorldTransform;
 	UIToWorldTransform.SetScale3D( FVector( Scale * WorldScaleFactor ) );
@@ -166,7 +166,7 @@ FTransform AVREditorBaseActor::MakeUITransformLockedToRoom()
 
 void AVREditorBaseActor::TickMoveTo( const float DeltaTime )
 {
-	const float WorldScaleFactor = Owner->GetWorldScaleFactor();
+	const float WorldScaleFactor = VRMode->GetWorldScaleFactor();
 
 	MoveToAlpha += DeltaTime;
 	const float LerpTime = MoveToTime;
diff --git a/Engine/Source/Editor/VREditor/VREditorBaseActor.h b/Engine/Source/Editor/VREditor/VREditorBaseActor.h
--- a/Engine/Source/Editor/VREditor/VREditorBaseActor.h
+++ b/Engine/Source/Editor/VREditor/VREditorBaseActor.h
@@ -20,7 +20,7 @@ public:
 	AVREditorBaseActor();
 
 	/** Sets the owning VR mode */
-	void SetVRMode( class FVREditorMode* InOwner );
+	void SetVRMode( class UVREditorMode* InVRMode );
 
 	/** Possible UI attachment points */
 	enum class EDockedTo
@@ -99,7 +99,8 @@ protected:
 private:
 
 	/** The VR mode that owns this actor */
-	class FVREditorMode* Owner;
+	UPROPERTY()
+	class UVREditorMode* VRMode;
 
 	/** Local rotation of the UI */
 	FRotator LocalRotation;
diff --git a/Engine/Source/Editor/VREditor/VREditorDockableWindow.cpp b/Engine/Source/Editor/VREditor/VREditorDockableWindow.cpp
--- a/Engine/Source/Editor/VREditor/VREditorDockableWindow.cpp
+++ b/Engine/Source/Editor/VREditor/VREditorDockableWindow.cpp
@@ -7,7 +7,7 @@
 #include "VREditorBaseUserWidget.h"
 #include "WidgetComponent.h"
 #include "VREditorWidgetComponent.h"
-#include "VREditorWorldInteraction.h"
+#include "ViewportWorldInteraction.h"
 #include "VREditorInteractor.h"
 
 namespace VREd
@@ -167,8 +167,8 @@ void AVREditorDockableWindow::SetupWidgetComponent()
 {
 	Super::SetupWidgetComponent();
 
-	SetSelectionBarColor( GetOwner().GetOwner().GetColor( FVREditorMode::EColors::UISelectionBarColor ) );
-	SetCloseButtonColor( GetOwner().GetOwner().GetColor( FVREditorMode::EColors::UICloseButtonColor ) );
+	SetSelectionBarColor( GetOwner().GetOwner().GetColor( UVREditorMode::EColors::UISelectionBarColor ) );
+	SetCloseButtonColor( GetOwner().GetOwner().GetColor( UVREditorMode::EColors::UICloseButtonColor ) );
 }
 
 
@@ -214,7 +214,7 @@ void AVREditorDockableWindow::TickManually( float DeltaTime )
 		}
 
 		const float AnimationOvershootAmount = 1.0f;	// @todo vreditor tweak
-		float EasedAimingAtMeFadeAlpha = FVREditorMode::OvershootEaseOut( AimingAtMeFadeAlpha, AnimationOvershootAmount );
+		float EasedAimingAtMeFadeAlpha = UVREditorMode::OvershootEaseOut( AimingAtMeFadeAlpha, AnimationOvershootAmount );
 
 		// Only show our extra buttons and controls if the user is roughly aiming toward us.  This just reduces clutter.
 		SelectionBarMeshComponent->SetVisibility( EasedAimingAtMeFadeAlpha > KINDA_SMALL_NUMBER ? true : false );
@@ -260,7 +260,7 @@ void AVREditorDockableWindow::TickManually( float DeltaTime )
 				-( GetSize().Y * 0.5f + SelectionBarSize.Z * 0.5f + VREd::DockUISelectionBarVerticalOffset->GetFloat() ) ) * AnimatedScale * WorldScaleFactor;
 			SelectionBarMeshComponent->SetRelativeLocation( SelectionBarRelativeLocation );
 
-			SetSelectionBarColor( GetOwner().GetOwner().GetColor( bIsHoveringOverSelectionBar ? FVREditorMode::EColors::UISelectionBarHoverColor : FVREditorMode::EColors::UISelectionBarColor ) );
+			SetSelectionBarColor( GetOwner().GetOwner().GetColor( bIsHoveringOverSelectionBar ? UVREditorMode::EColors::UISelectionBarHoverColor : UVREditorMode::EColors::UISelectionBarColor ) );
 		}
 
 		// Update the close button
@@ -288,7 +288,7 @@ void AVREditorDockableWindow::TickManually( float DeltaTime )
 				( GetSize().Y * 0.5f + CloseButtonSize.Z * 0.5f + VREd::DockUICloseButtonOffsetFromCorner->GetFloat() ) ) * AnimatedScale * WorldScaleFactor;
 			CloseButtonMeshComponent->SetRelativeLocation( CloseButtonRelativeLocation );
 
-			SetCloseButtonColor( GetOwner().GetOwner().GetColor( bIsHoveringOverCloseButton ? FVREditorMode::EColors::UICloseButtonHoverColor : FVREditorMode::EColors::UICloseButtonColor ) );
+			SetCloseButtonColor( GetOwner().GetOwner().GetColor( bIsHoveringOverCloseButton ? UVREditorMode::EColors::UICloseButtonHoverColor : UVREditorMode::EColors::UICloseButtonColor ) );
 		}
 	}
 } 
diff --git a/Engine/Source/Editor/VREditor/VREditorFloatingUI.cpp b/Engine/Source/Editor/VREditor/VREditorFloatingUI.cpp
--- a/Engine/Source/Editor/VREditor/VREditorFloatingUI.cpp
+++ b/Engine/Source/Editor/VREditor/VREditorFloatingUI.cpp
@@ -228,7 +228,7 @@ void AVREditorFloatingUI::UpdateFadingState( const float DeltaTime )
 FVector AVREditorFloatingUI::CalculateAnimatedScale() const
 {
 	const float AnimationOvershootAmount = 0.7f;	// @todo vreditor tweak
-	const float EasedAlpha = FVREditorMode::OvershootEaseOut( FadeAlpha, AnimationOvershootAmount );
+	const float EasedAlpha = UVREditorMode::OvershootEaseOut( FadeAlpha, AnimationOvershootAmount );
 
 	// Animate vertically more than horizontally; just looks a little better
 	const float ZScale = FMath::Max( 0.001f, EasedAlpha );
diff --git a/Engine/Source/Editor/VREditor/VREditorInteractor.cpp b/Engine/Source/Editor/VREditor/VREditorInteractor.cpp
--- a/Engine/Source/Editor/VREditor/VREditorInteractor.cpp
+++ b/Engine/Source/Editor/VREditor/VREditorInteractor.cpp
@@ -2,7 +2,7 @@
 
 #include "VREditorModule.h"
 #include "VREditorInteractor.h"
-#include "VREditorWorldInteraction.h"
+#include "ViewportWorldInteraction.h"
 
 #include "VREditorMode.h"
 #include "VREditorFloatingText.h"
@@ -35,7 +35,7 @@ UVREditorInteractor::~UVREditorInteractor()
 	Shutdown();
 }
 
-void UVREditorInteractor::Init( FVREditorMode* InVRMode )
+void UVREditorInteractor::Init( UVREditorMode* InVRMode )
 {
 	VRMode = InVRMode;
 	KeyToActionMap.Reset();
@@ -99,7 +99,7 @@ FHitResult UVREditorInteractor::GetHitResultFromLaserPointer( TArray<AActor*>* O
 void UVREditorInteractor::ResetHoverState( const float DeltaTime )
 {
 	bIsHoveringOverUI = false;
-}
+}	
 
 void UVREditorInteractor::OnStartDragging( UActorComponent* ClickedComponent, const FVector& HitLocation, const bool bIsPlacingActors )
 {
diff --git a/Engine/Source/Editor/VREditor/VREditorMode.cpp b/Engine/Source/Editor/VREditor/VREditorMode.cpp
--- a/Engine/Source/Editor/VREditor/VREditorMode.cpp
+++ b/Engine/Source/Editor/VREditor/VREditorMode.cpp
@@ -1,4 +1,4 @@
-// Copyright 1998-2015 Epic Games, Inc. All Rights Reserved.
+// Copyrigh 1998-2015 Epic Games, Inc. All Rights Reserved.
 
 #include "VREditorModule.h"
 #include "VREditorMode.h"
@@ -23,12 +23,14 @@
 #include "IAnalyticsProvider.h"
 
 #include "Editor/LevelEditor/Public/LevelEditorActions.h"
-#include "Editor/ViewportInteraction/Public/ViewportInteraction.h"
+#include "ViewportInteraction.h"
+#include "ViewportWorldInteraction.h"
 #include "VREditorInteractor.h"
 #include "MouseCursorInteractor.h"
 #include "VREditorMotionControllerInteractor.h"
-
 #include "Interactables/VREditorButton.h"
+#include "IViewportInteractionModule.h"
+#include "ViewportWorldInteractionManager.h"
 
 #define LOCTEXT_NAMESPACE "VREditorMode"
 
@@ -38,44 +40,42 @@ namespace VREd
 	static FAutoConsoleVariable ForceOculusMirrorMode( TEXT( "VREd.ForceOculusMirrorMode" ), 3, TEXT( "Which Oculus display mirroring mode to use (see MirrorWindowModeType in OculusRiftHMD.h)" ) );
 }
 
-FEditorModeID FVREditorMode::VREditorModeID( "VREditor" );
-bool FVREditorMode::bActuallyUsingVR = true;
-
 // @todo vreditor: Hacky that we have to import these this way. (Plugin has them in a .cpp, not exported)
 
-FVREditorMode::FVREditorMode()
-	: bWantsToExitMode( false ),
-	  bIsFullyInitialized( false ),
-	  AppTimeModeEntered( FTimespan::Zero() ),
-	  AvatarActor( nullptr ),
-      FlashlightComponent( nullptr ),
-	  bIsFlashlightOn( false ),
-	  MotionControllerID( 0 ),	// @todo vreditor minor: We only support a single controller, and we assume the first controller are the motion controls
-	  UISystem( nullptr ),
-	  TeleporterSystem( nullptr ),
-	  WorldInteraction( nullptr ),
-	  MouseCursorInteractor( nullptr ),
-	  LeftHandInteractor( nullptr ),
-	  RightHandInteractor( nullptr ),
-	  bFirstTick( true ),
-	  bWasInWorldSpaceBeforeScaleMode( false )
+UVREditorMode::UVREditorMode( const FObjectInitializer& ObjectInitializer ) : 
+	Super( ObjectInitializer ),
+	bWantsToExitMode( false ),
+	bIsFullyInitialized( false ),
+	AppTimeModeEntered( FTimespan::Zero() ),
+	AvatarActor( nullptr ),
+    FlashlightComponent( nullptr ),
+	bIsFlashlightOn( false ),
+	MotionControllerID( 0 ),	// @todo vreditor minor: We only support a single controller, and we assume the first controller are the motion controls
+	UISystem( nullptr ),
+	TeleporterSystem( nullptr ),
+	WorldInteraction( nullptr ),
+	MouseCursorInteractor( nullptr ),
+	LeftHandInteractor( nullptr ),
+	RightHandInteractor( nullptr ),
+	bFirstTick( true ),
+	bWasInWorldSpaceBeforeScaleMode( false )
 {
-	FEditorDelegates::MapChange.AddRaw( this, &FVREditorMode::OnMapChange );
-	FEditorDelegates::BeginPIE.AddRaw( this, &FVREditorMode::OnBeginPIE );
-	FEditorDelegates::EndPIE.AddRaw( this, &FVREditorMode::OnEndPIE );
-	FEditorDelegates::OnSwitchBeginPIEAndSIE.AddRaw( this, &FVREditorMode::OnSwitchBetweenPIEAndSIE );
 }
 
 
-FVREditorMode::~FVREditorMode()
+UVREditorMode::~UVREditorMode()
 {
-	FEditorDelegates::MapChange.RemoveAll( this );
-	FEditorDelegates::BeginPIE.RemoveAll( this );
-	FEditorDelegates::EndPIE.RemoveAll( this );
-	FEditorDelegates::OnSwitchBeginPIEAndSIE.RemoveAll( this );
+	AvatarActor = nullptr;
+	FlashlightComponent = nullptr;
+	UISystem = nullptr;
+	TeleporterSystem = nullptr;
+	WorldInteraction = nullptr;
+	MouseCursorInteractor = nullptr;
+	LeftHandInteractor = nullptr;
+	RightHandInteractor = nullptr;
 }
 
-void FVREditorMode::Enter()
+void UVREditorMode::Init( UViewportWorldInteraction* InViewportWorldInteraction )
 {
 	// @todo vreditor urgent: Turn on global editor hacks for VR Editor mode
 	GEnableVREditorHacks = true;
@@ -111,13 +111,41 @@ void FVREditorMode::Enter()
 		Colors[ (int32)EColors::UICloseButtonHoverColor ] = FLinearColor( 1.0f, 1.0f, 1.0f, 1.0f );
 	}
 
+	WorldInteraction = InViewportWorldInteraction;
+
+	bIsFullyInitialized = true;
+}
+
+void UVREditorMode::Shutdown()
+{
+	bIsFullyInitialized = false;
+
+	//Exit();
+
+	WorldInteraction = nullptr;
+
+	// @todo vreditor urgent: Disable global editor hacks for VR Editor mode
+	GEnableVREditorHacks = false;
+}
+
+void UVREditorMode::Enter()
+{
+	bWantsToExitMode = false;
+	ExitType = EVREditorExitType::Normal;
+
+	{
+		IViewportWorldInteractionManager& ViewportWorldInteraction = IViewportInteractionModule::Get().GetWorldInteractionManager();
+		ViewportWorldInteraction.OnPreWorldInteractionTick().AddUObject( this, &UVREditorMode::PreTick );
+		ViewportWorldInteraction.OnPostWorldInteractionTick().AddUObject( this, &UVREditorMode::Tick );
+	}
+
 	// @todo vreditor: We need to make sure the user can never switch to orthographic mode, or activate settings that
 	// would disrupt the user's ability to view the VR scene.
-		
+
 	// @todo vreditor: Don't bother drawing toolbars in VR, or other things that won't matter in VR
 
 	{
-		const TSharedRef< ILevelEditor >& LevelEditor = FModuleManager::GetModuleChecked<FLevelEditorModule>( "LevelEditor" ).GetFirstLevelEditor().ToSharedRef();
+		const TSharedRef< ILevelEditor >& LevelEditor = FModuleManager::GetModuleChecked<FLevelEditorModule>("LevelEditor").GetFirstLevelEditor().ToSharedRef();
 
 		bool bSummonNewWindow = true;
 
@@ -125,17 +153,17 @@ void FVREditorMode::Enter()
 		TSharedPtr<SLevelViewport> ExistingActiveLevelViewport;
 		{
 			TSharedPtr<ILevelViewport> ActiveLevelViewport = LevelEditor->GetActiveViewportInterface();
-			if( ActiveLevelViewport.IsValid() )
+			if(ActiveLevelViewport.IsValid())
 			{
-				ExistingActiveLevelViewport = StaticCastSharedRef< SLevelViewport >( ActiveLevelViewport->AsWidget() );
+				ExistingActiveLevelViewport = StaticCastSharedRef< SLevelViewport >(ActiveLevelViewport->AsWidget());
 
 				// Use the currently active window instead
 				bSummonNewWindow = false;
 			}
 		}
-		
+
 		TSharedPtr< SLevelViewport > VREditorLevelViewport;
-		if( bSummonNewWindow )
+		if(bSummonNewWindow)
 		{
 			// @todo vreditor: The resolution we set here doesn't matter, as HMDs will draw at their native resolution
 			// no matter what.  We should probably allow the window to be freely resizable by the user
@@ -143,48 +171,48 @@ void FVREditorMode::Enter()
 			FVector2D WindowSize;
 			{
 				IHeadMountedDisplay::MonitorInfo HMDMonitorInfo;
-				if( bActuallyUsingVR && GEngine->HMDDevice->GetHMDMonitorInfo( HMDMonitorInfo ) )
+				if(bActuallyUsingVR && GEngine->HMDDevice->GetHMDMonitorInfo(HMDMonitorInfo))
 				{
-					WindowSize = FVector2D( HMDMonitorInfo.ResolutionX, HMDMonitorInfo.ResolutionY );
+					WindowSize = FVector2D(HMDMonitorInfo.ResolutionX, HMDMonitorInfo.ResolutionY);
 				}
 				else
 				{
 					// @todo vreditor: Hard-coded failsafe window size
-					WindowSize = FVector2D( 1920.0f, 1080.0f );
+					WindowSize = FVector2D(1920.0f, 1080.0f);
 				}
 			}
 
 			// @todo vreditor: Use SLevelEditor::GetTableTitle() for the VR window title (needs dynamic update)
-			const FText VREditorWindowTitle = NSLOCTEXT( "VREditor", "VRWindowTitle", "Unreal Editor VR" );
-
-			TSharedRef< SWindow > VREditorWindow = SNew( SWindow )
-				.Title( VREditorWindowTitle )
-				.ClientSize( WindowSize )
-				.AutoCenter( EAutoCenter::PreferredWorkArea )
-				.UseOSWindowBorder( true )	// @todo vreditor: Allow window to be freely resized?  Shouldn't really hurt anything.  We should save position/size too.
-				.SizingRule( ESizingRule::UserSized );
+			const FText VREditorWindowTitle = NSLOCTEXT("VREditor", "VRWindowTitle", "Unreal Editor VR");
+
+			TSharedRef< SWindow > VREditorWindow = SNew(SWindow)
+				.Title(VREditorWindowTitle)
+				.ClientSize(WindowSize)
+				.AutoCenter(EAutoCenter::PreferredWorkArea)
+				.UseOSWindowBorder(true)	// @todo vreditor: Allow window to be freely resized?  Shouldn't really hurt anything.  We should save position/size too.
+				.SizingRule(ESizingRule::UserSized);
 			this->VREditorWindowWeakPtr = VREditorWindow;
 
 			VREditorLevelViewport =
-				SNew( SLevelViewport )
-				.ViewportType( LVT_Perspective ) // Perspective
-				.Realtime( true )
+				SNew(SLevelViewport)
+				.ViewportType(LVT_Perspective) // Perspective
+				.Realtime(true)
 				//				.ParentLayout( AsShared() )	// @todo vreditor: We don't have one and we probably don't need one, right?  Make sure a null parent layout is handled properly everywhere.
-				.ParentLevelEditor( LevelEditor )
+				.ParentLevelEditor(LevelEditor)
 				//				.ConfigKey( BottomLeftKey )	// @todo vreditor: This is for saving/loading layout.  We would need this in order to remember viewport settings like show flags, etc.
-				.IsEnabled( FSlateApplication::Get().GetNormalExecutionAttribute() );
+				.IsEnabled(FSlateApplication::Get().GetNormalExecutionAttribute());
 
 			// Allow the editor to keep track of this editor viewport.  Because it's not inside of a normal tab, 
 			// we need to explicitly tell the level editor about it
-			LevelEditor->AddStandaloneLevelViewport( VREditorLevelViewport.ToSharedRef() );
+			LevelEditor->AddStandaloneLevelViewport(VREditorLevelViewport.ToSharedRef());
 
-			VREditorWindow->SetContent( VREditorLevelViewport.ToSharedRef() );
+			VREditorWindow->SetContent(VREditorLevelViewport.ToSharedRef());
 
 			// NOTE: We're intentionally not adding this window natively parented to the main frame window, because we don't want it
 			// to minimize/restore when the main frame is minimized/restored
-			FSlateApplication::Get().AddWindow( VREditorWindow );
+			FSlateApplication::Get().AddWindow(VREditorWindow);
 
-			VREditorWindow->SetOnWindowClosed( FOnWindowClosed::CreateRaw( this, &FVREditorMode::OnVREditorWindowClosed ) );
+			VREditorWindow->SetOnWindowClosed(FOnWindowClosed::CreateUObject(this, &UVREditorMode::OnVREditorWindowClosed));
 
 			VREditorWindow->BringToFront();	// @todo vreditor: Not sure if this is needed, especially if we decide the window should be hidden (copied this from PIE code)
 		}
@@ -192,16 +220,17 @@ void FVREditorMode::Enter()
 		{
 			VREditorLevelViewport = ExistingActiveLevelViewport;
 
-			if( bActuallyUsingVR )
+			if(bActuallyUsingVR)
 			{
 				// Switch to immersive mode
 				const bool bWantImmersive = true;
 				const bool bAllowAnimation = false;
-				ExistingActiveLevelViewport->MakeImmersive( bWantImmersive, bAllowAnimation );
+				ExistingActiveLevelViewport->MakeImmersive(bWantImmersive, bAllowAnimation);
 			}
 		}
+
 		this->VREditorLevelViewportWeakPtr = VREditorLevelViewport;
-		
+
 		{
 			FLevelEditorViewportClient& VRViewportClient = VREditorLevelViewport->GetLevelViewportClient();
 			FEditorViewportClient& VREditorViewportClient = VRViewportClient;
@@ -209,7 +238,7 @@ void FVREditorMode::Enter()
 			// Make sure we are in perspective mode
 			// @todo vreditor: We should never allow ortho switching while in VR
 			SavedEditorState.ViewportType = VREditorViewportClient.GetViewportType();
-			VREditorViewportClient.SetViewportType( LVT_Perspective );
+			VREditorViewportClient.SetViewportType(LVT_Perspective);
 
 			// Set the initial camera location
 			// @todo vreditor: This should instead be calculated using the currently active perspective camera's
@@ -218,26 +247,26 @@ void FVREditorMode::Enter()
 			// @todo vreditor: Should save and restore camera position and any other settings we change (viewport type, pitch locking, etc.)
 			SavedEditorState.ViewLocation = VRViewportClient.GetViewLocation();
 			SavedEditorState.ViewRotation = VRViewportClient.GetViewRotation();
-			
+
 			// Don't allow the tracking space to pitch up or down.  People hate that in VR.
 			// @todo vreditor: This doesn't seem to prevent people from pitching the camera with RMB drag
 			SavedEditorState.bLockedPitch = VRViewportClient.GetCameraController()->GetConfig().bLockedPitch;
-			if( bActuallyUsingVR )
+			if(bActuallyUsingVR)
 			{
 				VRViewportClient.GetCameraController()->AccessConfig().bLockedPitch = true;
 			}
 
 			// Set "game mode" to be enabled, to get better performance.  Also hit proxies won't work in VR, anyway
 			SavedEditorState.bGameView = VREditorViewportClient.IsInGameView();
-			VREditorViewportClient.SetGameView( true );
+			VREditorViewportClient.SetGameView(true);
 
 			SavedEditorState.bRealTime = VREditorViewportClient.IsRealtime();
-			VREditorViewportClient.SetRealtime( true );
+			VREditorViewportClient.SetRealtime(true);
 
 			SavedEditorState.ShowFlags = VREditorViewportClient.EngineShowFlags;
 
 			// Never show the traditional Unreal transform widget.  It doesn't work in VR because we don't have hit proxies.
-			VREditorViewportClient.EngineShowFlags.SetModeWidgets( false );
+			VREditorViewportClient.EngineShowFlags.SetModeWidgets(false);
 
 			// Make sure the mode widgets don't come back when users click on things
 			VRViewportClient.bAlwaysShowModeWidgetAfterSelectionChanges = false;
@@ -250,18 +279,18 @@ void FVREditorMode::Enter()
 			SavedEditorState.bOnScreenMessages = GAreScreenMessagesEnabled;
 			GAreScreenMessagesEnabled = false;
 
-			if( bActuallyUsingVR )
+			if(bActuallyUsingVR)
 			{
 				SavedEditorState.TrackingOrigin = GEngine->HMDDevice->GetTrackingOrigin();
-				GEngine->HMDDevice->SetTrackingOrigin( EHMDTrackingOrigin::Floor );
+				GEngine->HMDDevice->SetTrackingOrigin(EHMDTrackingOrigin::Floor);
 			}
 
 			// Make the new viewport the active level editing viewport right away
 			GCurrentLevelEditingViewportClient = &VRViewportClient;
 
 			// Enable selection outline right away
-			VREditorViewportClient.EngineShowFlags.SetSelection(true);
-			VREditorViewportClient.EngineShowFlags.SetSelectionOutline(true);
+			VREditorViewportClient.EngineShowFlags.SetSelection( true );
+			VREditorViewportClient.EngineShowFlags.SetSelectionOutline( true );
 		}
 
 		VREditorLevelViewport->EnableStereoRendering( bActuallyUsingVR );
@@ -292,11 +321,14 @@ void FVREditorMode::Enter()
 	// Setup sub systems
 	{
 		// Setup world interaction
-		WorldInteraction = NewObject<UVREditorWorldInteraction>();
-		WorldInteraction->SetOwner( this );
-		WorldInteraction->Init( VREditorLevelViewportWeakPtr.Pin()->GetViewportClient() );
+		TSharedPtr<FEditorViewportClient> ViewportClient = VREditorLevelViewportWeakPtr.Pin()->GetViewportClient();
+		WorldInteraction->SetViewport( ViewportClient );
+		WorldInteraction->Activate( true );
+
+		WorldInteraction->OnHandleKeyInput().AddUObject( this, &UVREditorMode::InputKey );
 
 		// Motion controllers
+		if(bActuallyUsingVR)
 		{
 			LeftHandInteractor = NewObject<UVREditorMotionControllerInteractor>( WorldInteraction );
 			LeftHandInteractor->SetControllerHandSide( EControllerHand::Left );
@@ -308,10 +340,9 @@ void FVREditorMode::Enter()
 			RightHandInteractor->Init( this );
 			WorldInteraction->AddInteractor( RightHandInteractor );
 
-			WorldInteraction->PairInteractors( LeftHandInteractor, RightHandInteractor );
+			WorldInteraction->PairInteractors(LeftHandInteractor, RightHandInteractor );
 		}
-
-		if( !bActuallyUsingVR )
+		else
 		{
 			// Register an interactor for the mouse cursor
 			MouseCursorInteractor = NewObject<UMouseCursorInteractor>( WorldInteraction );
@@ -324,93 +355,102 @@ void FVREditorMode::Enter()
 		UISystem->SetOwner( this );
 		UISystem->Init();
 
+		VRWorldInteractionExtension = NewObject<UVREditorWorldInteraction>();
+		VRWorldInteractionExtension->Init( this, WorldInteraction );
+
 		// Setup teleporter
 		TeleporterSystem = NewObject<UVREditorTeleporter>();
-		TeleporterSystem->Init( this );
-
+		TeleporterSystem->Init(this);
 	}
 
-	// Call parent implementation
-	FEdMode::Enter();
+	if(AvatarActor == nullptr)
+	{
+		SpawnAvatarMeshActor();
+	}
 
 	bFirstTick = true;
-	bIsFullyInitialized = true;
+	bIsActive = true;
 }
 
-
-void FVREditorMode::Exit()
+void UVREditorMode::Exit()
 {
-	bIsFullyInitialized = false;
-
-	FSlateApplication::Get().SetInputPreProcessor(false);
-
-	//Destroy the avatar
 	{
-		DestroyTransientActor( AvatarActor );
-		AvatarActor = nullptr;
-
-		FlashlightComponent = nullptr;
-	}
-
-	{
-		if( GEngine->HMDDevice.IsValid() )
+		//Destroy the avatar
 		{
-			// @todo vreditor switch: We don't want to do this if a VR PIE session is somehow active.  Is that even possible while immersive?
-			GEngine->HMDDevice->EnableStereo( false );
-		}
+			DestroyTransientActor( AvatarActor );
+			AvatarActor = nullptr;
 
-		if( bActuallyUsingVR )
-		{ 
-			// Restore Slate drag trigger distance
-			FSlateApplication::Get().SetDragTriggerDistance( SavedEditorState.DragTriggerDistance );
+			FlashlightComponent = nullptr;
 		}
 
-		TSharedPtr<SLevelViewport> VREditorLevelViewport( VREditorLevelViewportWeakPtr.Pin() );
-		if( VREditorLevelViewport.IsValid() )
 		{
-			VREditorLevelViewport->EnableStereoRendering( false );
-			VREditorLevelViewport->SetRenderDirectlyToWindow( false );
+			if(GEngine->HMDDevice.IsValid())
+			{
+				// @todo vreditor switch: We don't want to do this if a VR PIE session is somehow active.  Is that even possible while immersive?
+				GEngine->HMDDevice->EnableStereo( false );
+			}
 
+			if(bActuallyUsingVR)
 			{
-				FLevelEditorViewportClient& VRViewportClient = VREditorLevelViewport->GetLevelViewportClient();
-				FEditorViewportClient& VREditorViewportClient = VRViewportClient;
+				// Restore Slate drag trigger distance
+				FSlateApplication::Get().SetDragTriggerDistance( SavedEditorState.DragTriggerDistance );
+			}
 
-				// Restore settings that we changed on the viewport
-				VREditorViewportClient.SetViewportType( SavedEditorState.ViewportType );
-				VRViewportClient.GetCameraController()->AccessConfig().bLockedPitch = SavedEditorState.bLockedPitch;
-				VRViewportClient.bAlwaysShowModeWidgetAfterSelectionChanges = SavedEditorState.bAlwaysShowModeWidgetAfterSelectionChanges;
-				VRViewportClient.EngineShowFlags = SavedEditorState.ShowFlags;
-				VRViewportClient.SetGameView( SavedEditorState.bGameView );
+			TSharedPtr<SLevelViewport> VREditorLevelViewport( VREditorLevelViewportWeakPtr.Pin() );
+			if(VREditorLevelViewport.IsValid())
+			{
+				VREditorLevelViewport->EnableStereoRendering( false );
+				VREditorLevelViewport->SetRenderDirectlyToWindow( false );
 
-				if( bActuallyUsingVR )
 				{
-					VRViewportClient.SetViewLocation( GetHeadTransform().GetLocation() );
-				
-					FRotator HeadRotationNoRoll = GetHeadTransform().GetRotation().Rotator();
-					HeadRotationNoRoll.Roll = 0.0f;
-					VRViewportClient.SetViewRotation(HeadRotationNoRoll); // Use SavedEditorState.ViewRotation to go back to start rot
+					FLevelEditorViewportClient& VRViewportClient = VREditorLevelViewport->GetLevelViewportClient();
+					FEditorViewportClient& VREditorViewportClient = VRViewportClient;
+
+					// Restore settings that we changed on the viewport
+					VREditorViewportClient.SetViewportType( SavedEditorState.ViewportType );
+					VRViewportClient.GetCameraController()->AccessConfig().bLockedPitch = SavedEditorState.bLockedPitch;
+					VRViewportClient.bAlwaysShowModeWidgetAfterSelectionChanges = SavedEditorState.bAlwaysShowModeWidgetAfterSelectionChanges;
+					VRViewportClient.EngineShowFlags = SavedEditorState.ShowFlags;
+					VRViewportClient.SetGameView( SavedEditorState.bGameView );
+
+					if(bActuallyUsingVR)
+					{
+						VRViewportClient.SetViewLocation( GetHeadTransform().GetLocation() );
+
+						FRotator HeadRotationNoRoll = GetHeadTransform().GetRotation().Rotator();
+						HeadRotationNoRoll.Roll = 0.0f;
+						VRViewportClient.SetViewRotation( HeadRotationNoRoll ); // Use SavedEditorState.ViewRotation to go back to start rot
+					}
+
+					VRViewportClient.SetRealtime( SavedEditorState.bRealTime );
+
+					GNearClippingPlane = SavedEditorState.NearClipPlane;
+					GAreScreenMessagesEnabled = SavedEditorState.bOnScreenMessages;
+
+					if(bActuallyUsingVR)
+					{
+						GEngine->HMDDevice->SetTrackingOrigin( SavedEditorState.TrackingOrigin );
+					}
 				}
-				
-				VRViewportClient.SetRealtime( SavedEditorState.bRealTime );
-
-				GNearClippingPlane = SavedEditorState.NearClipPlane;
-				GAreScreenMessagesEnabled = SavedEditorState.bOnScreenMessages;
 
-				if( bActuallyUsingVR )
+				if(bActuallyUsingVR)
 				{
-					GEngine->HMDDevice->SetTrackingOrigin( SavedEditorState.TrackingOrigin );
+					// Leave immersive mode
+					const bool bWantImmersive = false;
+					const bool bAllowAnimation = false;
+					VREditorLevelViewport->MakeImmersive( bWantImmersive, bAllowAnimation );
 				}
+
+				VREditorLevelViewportWeakPtr.Reset();
 			}
 
-			if( bActuallyUsingVR )
+			// Kill the VR editor window
+			TSharedPtr<SWindow> VREditorWindow( VREditorWindowWeakPtr.Pin() );
+			if(VREditorWindow.IsValid())
 			{
-				// Leave immersive mode
-				const bool bWantImmersive = false;
-				const bool bAllowAnimation = false;
-				VREditorLevelViewport->MakeImmersive( bWantImmersive, bAllowAnimation );
+				VREditorWindow->RequestDestroyWindow();
+				VREditorWindow.Reset();
 			}
-
-			VREditorLevelViewportWeakPtr.Reset();
 		}
 
 		// Kill the VR editor window
@@ -421,43 +461,69 @@ void FVREditorMode::Exit()
 			VREditorWindow.Reset();
 		}
 	}
-
-	// Call parent implementation
-	FEdMode::Exit();
+	//Destroy the avatar
+	{
+		DestroyTransientActor( AvatarActor );
+		AvatarActor = nullptr;
+		FlashlightComponent = nullptr;
+	}
 
 	// Kill subsystems
-	if ( UISystem != nullptr )
+	if(UISystem != nullptr)
 	{
 		UISystem->Shutdown();
-		UISystem->MarkPendingKill();
+		//UISystem->MarkPendingKill();
 		UISystem = nullptr;
 	}
 
-	if ( TeleporterSystem != nullptr )
+	if(VRWorldInteractionExtension != nullptr)
+	{
+		VRWorldInteractionExtension->Shutdown();
+		VRWorldInteractionExtension->MarkPendingKill();
+		VRWorldInteractionExtension = nullptr;
+	}
+
+	if(TeleporterSystem != nullptr)
 	{
 		TeleporterSystem->Shutdown();
 		TeleporterSystem->MarkPendingKill();
 		TeleporterSystem = nullptr;
 	}
 
-	if ( WorldInteraction != nullptr )
 	{
-		WorldInteraction->Shutdown();
-		WorldInteraction->MarkPendingKill();
-		WorldInteraction = nullptr;
+		IViewportWorldInteractionManager& ViewportWorldInteraction = IViewportInteractionModule::Get().GetWorldInteractionManager();
+		ViewportWorldInteraction.OnPreWorldInteractionTick().RemoveAll( this );
+		ViewportWorldInteraction.OnPostWorldInteractionTick().RemoveAll( this );
 	}
 
-	// @todo vreditor urgent: Disable global editor hacks for VR Editor mode
-	GEnableVREditorHacks = false;
+	if( WorldInteraction != nullptr )
+	{
+		WorldInteraction->OnHandleKeyInput().RemoveAll( this );
+
+		WorldInteraction->RemoveInteractor( LeftHandInteractor );
+		LeftHandInteractor->MarkPendingKill();
+		LeftHandInteractor->Shutdown();
+		LeftHandInteractor = nullptr;
+		WorldInteraction->RemoveInteractor( RightHandInteractor );
+		RightHandInteractor->Shutdown();
+		RightHandInteractor->MarkPendingKill();
+		RightHandInteractor = nullptr;
+
+		WorldInteraction->Activate( false );
+	}
+
+	bWantsToExitMode = false;
+	bIsActive = false;
+	bFirstTick = false;
 }
 
-void FVREditorMode::StartExitingVRMode( const EVREditorExitType InExitType /*= EVREditorExitType::To_Editor */ )
+void UVREditorMode::StartExitingVRMode( const EVREditorExitType InExitType /*= EVREditorExitType::To_Editor */ )
 {
 	ExitType = InExitType;
 	bWantsToExitMode = true;
 }
 
-void FVREditorMode::SpawnAvatarMeshActor()
+void UVREditorMode::SpawnAvatarMeshActor()
 {
 	// Setup our avatar
 	if( AvatarActor == nullptr )
@@ -475,25 +541,20 @@ void FVREditorMode::SpawnAvatarMeshActor()
 }
 
 
-void FVREditorMode::OnVREditorWindowClosed( const TSharedRef<SWindow>& ClosedWindow )
+void UVREditorMode::OnVREditorWindowClosed( const TSharedRef<SWindow>& ClosedWindow )
 {
 	StartExitingVRMode();
 }
 
-
-void FVREditorMode::Tick( FEditorViewportClient* ViewportClient, float DeltaTime )
+void UVREditorMode::PreTick( const float DeltaTime )
 {
-	// Call parent implementation
-	FEdMode::Tick( ViewportClient, DeltaTime );
-
-	// Only if this is our VR viewport. Remember, editor modes get a chance to tick and receive input for each active viewport.
-	if ( ViewportClient != GetLevelViewportPossessedForVR().GetViewportClient().Get() || bWantsToExitMode )
+	if( !bIsFullyInitialized || !bIsActive || bWantsToExitMode )
 	{
 		return;
 	}
 
 	//Setting the initial position and rotation based on the editor viewport when going into VR mode
-	if ( bFirstTick && bActuallyUsingVR )
+	if( bFirstTick && bActuallyUsingVR )
 	{
 		const FTransform RoomToWorld = GetRoomTransform();
 		const FTransform WorldToRoom = RoomToWorld.Inverse();
@@ -512,17 +573,17 @@ void FVREditorMode::Tick( FEditorViewportClient* ViewportClient, float DeltaTime
 		FTransform ResultToWorld = ( HeadToRoomYaw.Inverse() * ViewportToRoomYaw ) * RoomToWorldYaw;
 		SetRoomTransform( ResultToWorld );
 	}
+}
 
-	if ( AvatarActor == nullptr )
+void UVREditorMode::Tick( float DeltaTime )
+{
+	if( !bIsFullyInitialized || !bIsActive || bWantsToExitMode || !VREditorLevelViewportWeakPtr.IsValid() )
 	{
-		SpawnAvatarMeshActor();
+		return;
 	}
 
 	TickHandle.Broadcast( DeltaTime );
-
-	WorldInteraction->Tick( ViewportClient, DeltaTime );
-
-	UISystem->Tick( ViewportClient, DeltaTime );
+	UISystem->Tick( GetLevelViewportPossessedForVR().GetViewportClient().Get(), DeltaTime );
 
 	// Update avatar meshes
 	{
@@ -546,122 +607,114 @@ void FVREditorMode::Tick( FEditorViewportClient* ViewportClient, float DeltaTime
 }
 
 
-bool FVREditorMode::InputKey(FEditorViewportClient* ViewportClient, FViewport* Viewport, FKey Key, EInputEvent Event)
+void UVREditorMode::InputKey(const FEditorViewportClient& InViewportClient, const FKey InKey, const EInputEvent InEvent, bool& bOutWasHandled)
 {
-	// Only if this is our VR viewport. Remember, editor modes get a chance to tick and receive input for each active viewport.
-	if ( ViewportClient != GetLevelViewportPossessedForVR().GetViewportClient().Get() )
-	{
-		return InputKey(GetLevelViewportPossessedForVR().GetViewportClient().Get(), Viewport, Key, Event);
-	}
-
-	if ( Key == EKeys::Escape )
+	if(InKey == EKeys::Escape)
 	{
 		// User hit escape, so bail out of VR mode
 		StartExitingVRMode();
 	}
-	else if( Key.IsMouseButton() )	// Input preprocessor cannot handle mouse buttons, so we'll route those the normal way
-	{
-		return WorldInteraction->HandleInputKey( Key, Event );
-	}
-
-	return FEdMode::InputKey(ViewportClient, Viewport, Key, Event);
 }
 
-
-bool FVREditorMode::InputAxis(FEditorViewportClient* ViewportClient, FViewport* Viewport, int32 ControllerId, FKey Key, float Delta, float DeltaTime)
-{
-	// Only if this is our VR viewport. Remember, editor modes get a chance to tick and receive input for each active viewport.
-	if ( ViewportClient != GetLevelViewportPossessedForVR().GetViewportClient().Get() )
-	{
-		return InputAxis(GetLevelViewportPossessedForVR().GetViewportClient().Get(), Viewport, ControllerId, Key, Delta, DeltaTime);
-	}
-
-	return FEdMode::InputAxis(ViewportClient, Viewport, ControllerId, Key, Delta, DeltaTime);
-}
-
-bool FVREditorMode::IsCompatibleWith(FEditorModeID OtherModeID) const
-{
-	// We are compatible with all other modes!
-	return true;
-}
-
-
-void FVREditorMode::AddReferencedObjects( FReferenceCollector& Collector )
-{
-	Collector.AddReferencedObject( AvatarActor );
-	Collector.AddReferencedObject( FlashlightComponent );	
-	Collector.AddReferencedObject( UISystem );
-	Collector.AddReferencedObject( WorldInteraction );
-	Collector.AddReferencedObject( TeleporterSystem );
-	Collector.AddReferencedObject( MouseCursorInteractor );
-	Collector.AddReferencedObject( LeftHandInteractor );
-	Collector.AddReferencedObject( RightHandInteractor );
-}
-
-
-void FVREditorMode::Render( const FSceneView* SceneView, FViewport* Viewport, FPrimitiveDrawInterface* PDI )
-{
-	//StopOldHapticEffects(); //@todo vreditor
-
-	FEdMode::Render( SceneView, Viewport, PDI );
-
-	if( bIsFullyInitialized )
-	{
-		// Let our subsystems render, too
-		UISystem->Render( SceneView, Viewport, PDI );
-	}
-}
+//bool UVREditorMode::InputKey(FEditorViewportClient* ViewportClient, FViewport* Viewport, FKey Key, EInputEvent Event)
+//{
+//	// Only if this is our VR viewport. Remember, editor modes get a chance to tick and receive input for each active viewport.
+//	if ( ViewportClient != GetLevelViewportPossessedForVR().GetViewportClient().Get() )
+//	{
+//		return InputKey(GetLevelViewportPossessedForVR().GetViewportClient().Get(), Viewport, Key, Event);
+//	}
+//
+//	if ( Key == EKeys::Escape )
+//	{
+//		// User hit escape, so bail out of VR mode
+//		StartExitingVRMode();
+//	}
+//	else if( Key.IsMouseButton() )	// Input preprocessor cannot handle mouse buttons, so we'll route those the normal way
+//	{
+//		return WorldInteraction->HandleInputKey( Key, Event );
+//	}
+//
+//	return FEdMode::InputKey(ViewportClient, Viewport, Key, Event);
+//}
+//
+//
+//bool UVREditorMode::InputAxis(FEditorViewportClient* ViewportClient, FViewport* Viewport, int32 ControllerId, FKey Key, float Delta, float DeltaTime)
+//{
+//	// Only if this is our VR viewport. Remember, editor modes get a chance to tick and receive input for each active viewport.
+//	if ( ViewportClient != GetLevelViewportPossessedForVR().GetViewportClient().Get() )
+//	{
+//		return InputAxis(GetLevelViewportPossessedForVR().GetViewportClient().Get(), Viewport, ControllerId, Key, Delta, DeltaTime);
+//	}
+//
+//	return FEdMode::InputAxis(ViewportClient, Viewport, ControllerId, Key, Delta, DeltaTime);
+//}
+
+//void UVREditorMode::Render( const FSceneView* SceneView, FViewport* Viewport, FPrimitiveDrawInterface* PDI )
+//{
+//	//StopOldHapticEffects(); //@todo vreditor
+//
+//	if( bIsFullyInitialized )
+//	{
+//		// Let our subsystems render, too
+//		UISystem->Render( SceneView, Viewport, PDI );
+//	}
+//}
 
 /************************************************************************/
 /* IVREditorMode interface                                              */
 /************************************************************************/
 
-AActor* FVREditorMode::GetAvatarMeshActor()
+AActor* UVREditorMode::GetAvatarMeshActor()
 {
 	return AvatarActor;
 }
 
-UWorld* FVREditorMode::GetWorld() const
+UWorld* UVREditorMode::GetWorld() const
 {
 	return WorldInteraction->GetViewportWorld();
 }
 
-FTransform FVREditorMode::GetRoomTransform() const
+FTransform UVREditorMode::GetRoomTransform() const
 {
 	return WorldInteraction->GetRoomTransform();
 }
 
-void FVREditorMode::SetRoomTransform( const FTransform& NewRoomTransform )
+void UVREditorMode::SetRoomTransform( const FTransform& NewRoomTransform )
 {
 	WorldInteraction->SetRoomTransform( NewRoomTransform );
 }
 
-FTransform FVREditorMode::GetRoomSpaceHeadTransform() const
+FTransform UVREditorMode::GetRoomSpaceHeadTransform() const
 {
 	return WorldInteraction->GetRoomSpaceHeadTransform();
 }
 
-FTransform FVREditorMode::GetHeadTransform() const
+FTransform UVREditorMode::GetHeadTransform() const
 {
 	return WorldInteraction->GetHeadTransform();
 }
 
-const UViewportWorldInteraction& FVREditorMode::GetWorldInteraction() const
+const UViewportWorldInteraction& UVREditorMode::GetWorldInteraction() const
 {
 	return *WorldInteraction;
 }
 
-UViewportWorldInteraction& FVREditorMode::GetWorldInteraction()
+UViewportWorldInteraction& UVREditorMode::GetWorldInteraction()
 {
 	return *WorldInteraction;
 }
 
-bool FVREditorMode::IsFullyInitialized() const
+bool UVREditorMode::IsFullyInitialized() const
 {
 	return bIsFullyInitialized;
 }
 
-AActor* FVREditorMode::SpawnTransientSceneActor( TSubclassOf<AActor> ActorClass, const FString& ActorName, const bool bWithSceneComponent ) const
+bool UVREditorMode::IsActive() const
+{
+	return bIsActive;
+}
+
+AActor* UVREditorMode::SpawnTransientSceneActor(TSubclassOf<AActor> ActorClass, const FString& ActorName, const bool bWithSceneComponent) const
 {
 	const bool bWasWorldPackageDirty = GetWorld()->GetOutermost()->IsDirty();
 
@@ -694,52 +747,29 @@ AActor* FVREditorMode::SpawnTransientSceneActor( TSubclassOf<AActor> ActorClass,
 }
 
 
-void FVREditorMode::DestroyTransientActor( AActor* Actor ) const
+void UVREditorMode::DestroyTransientActor( AActor* Actor ) const
 {
 	WorldInteraction->DestroyTransientActor( Actor );
 }
 
-const SLevelViewport& FVREditorMode::GetLevelViewportPossessedForVR() const
+const SLevelViewport& UVREditorMode::GetLevelViewportPossessedForVR() const
 {
 	return *VREditorLevelViewportWeakPtr.Pin();
 }
 
-SLevelViewport& FVREditorMode::GetLevelViewportPossessedForVR()
+SLevelViewport& UVREditorMode::GetLevelViewportPossessedForVR()
 {
 	return *VREditorLevelViewportWeakPtr.Pin();
 }
 
 
-float FVREditorMode::GetWorldScaleFactor() const
+float UVREditorMode::GetWorldScaleFactor() const
 {
 	return WorldInteraction->GetViewportWorld()->GetWorldSettings()->WorldToMeters / 100.0f;
 }
 
-void FVREditorMode::OnMapChange( uint32 MapChangeFlags )
-{
-	CleanUpActorsBeforeMapChangeOrSimulate();
-}
-
 
-void FVREditorMode::OnBeginPIE( const bool bIsSimulatingInEditor )
-{
-	CleanUpActorsBeforeMapChangeOrSimulate();
-}
-
-
-void FVREditorMode::OnEndPIE( const bool bIsSimulatingInEditor )
-{
-	CleanUpActorsBeforeMapChangeOrSimulate();
-}
-
-
-void FVREditorMode::OnSwitchBetweenPIEAndSIE( const bool bIsSimulatingInEditor )
-{
-	CleanUpActorsBeforeMapChangeOrSimulate();
-}
-
-
-void FVREditorMode::CleanUpActorsBeforeMapChangeOrSimulate()
+void UVREditorMode::CleanUpActorsBeforeMapChangeOrSimulate()
 {
 	if ( WorldInteraction != nullptr )
 	{
@@ -757,7 +787,7 @@ void FVREditorMode::CleanUpActorsBeforeMapChangeOrSimulate()
 	}
 }
 
-void FVREditorMode::ToggleFlashlight( UVREditorInteractor* Interactor )
+void UVREditorMode::ToggleFlashlight( UVREditorInteractor* Interactor )
 {
 	UVREditorMotionControllerInteractor* MotionControllerInteractor = Cast<UVREditorMotionControllerInteractor>( Interactor );
 	if ( MotionControllerInteractor )
@@ -785,7 +815,7 @@ void FVREditorMode::ToggleFlashlight( UVREditorInteractor* Interactor )
 	}
 }
 
-void FVREditorMode::CycleTransformGizmoHandleType()
+void UVREditorMode::CycleTransformGizmoHandleType()
 {
 	EGizmoHandleTypes NewGizmoType = (EGizmoHandleTypes)( (uint8)WorldInteraction->GetCurrentGizmoType() + 1 );
 	
@@ -818,17 +848,17 @@ void FVREditorMode::CycleTransformGizmoHandleType()
 	WorldInteraction->SetGizmoHandleType( NewGizmoType );
 }
 
-EHMDDeviceType::Type FVREditorMode::GetHMDDeviceType() const
+EHMDDeviceType::Type UVREditorMode::GetHMDDeviceType() const
 {
 	return GEngine->HMDDevice.IsValid() ? GEngine->HMDDevice->GetHMDDeviceType() : EHMDDeviceType::DT_SteamVR;
 }
 
-FLinearColor FVREditorMode::GetColor( const EColors Color ) const
+FLinearColor UVREditorMode::GetColor( const EColors Color ) const
 {
 	return Colors[ (int32)Color ];
 }
 
-bool FVREditorMode::IsHandAimingTowardsCapsule( UViewportInteractor* Interactor, const FTransform& CapsuleTransform, FVector CapsuleStart, const FVector CapsuleEnd, const float CapsuleRadius, const float MinDistanceToCapsule, const FVector CapsuleFrontDirection, const float MinDotForAimingAtCapsule ) const
+bool UVREditorMode::IsHandAimingTowardsCapsule( UViewportInteractor* Interactor, const FTransform& CapsuleTransform, FVector CapsuleStart, const FVector CapsuleEnd, const float CapsuleRadius, const float MinDistanceToCapsule, const FVector CapsuleFrontDirection, const float MinDotForAimingAtCapsule ) const
 {
 	bool bIsAimingTowards = false;
 	const float WorldScaleFactor = GetWorldScaleFactor();
@@ -889,17 +919,29 @@ bool FVREditorMode::IsHandAimingTowardsCapsule( UViewportInteractor* Interactor,
 	return bIsAimingTowards;
 }
 
-UVREditorInteractor* FVREditorMode::GetHandInteractor( const EControllerHand ControllerHand ) const 
+UVREditorInteractor* UVREditorMode::GetHandInteractor( const EControllerHand ControllerHand ) const 
 {
 	UVREditorInteractor* ResultInteractor = ControllerHand == EControllerHand::Left ? LeftHandInteractor : RightHandInteractor;
 	check( ResultInteractor != nullptr );
 	return ResultInteractor;
 }
 
-void FVREditorMode::StopOldHapticEffects()
+void UVREditorMode::StopOldHapticEffects()
 {
 	LeftHandInteractor->StopOldHapticEffects();
 	RightHandInteractor->StopOldHapticEffects();
 }
 
+void UVREditorMode::SnapSelectedActorsToGround()
+{
+	VRWorldInteractionExtension->SnapSelectedActorsToGround();
+}
+
+void UVREditorMode::GoToDefaultScale()
+{
+	LeftHandInteractor->SetDraggingMode( EViewportInteractionDraggingMode::Nothing );
+	RightHandInteractor->SetDraggingMode( EViewportInteractionDraggingMode::Nothing );
+	WorldInteraction->SetWorldToMetersScale( 100.0f );
+}
+
 #undef LOCTEXT_NAMESPACE
\ No newline at end of file
diff --git a/Engine/Source/Editor/VREditor/VREditorMode.h b/Engine/Source/Editor/VREditor/VREditorMode.h
deleted file mode 100644
diff --git a/Engine/Source/Editor/VREditor/VREditorModeManager.cpp b/Engine/Source/Editor/VREditor/VREditorModeManager.cpp
new file mode 100644
--- /dev/null
+++ b/Engine/Source/Editor/VREditor/VREditorModeManager.cpp
@@ -0,0 +1,205 @@
+#include "VREditorModule.h"
+#include "VREditorModeManager.h"
+#include "VREditorMode.h"
+#include "VREditorMode.h"
+
+#include "LevelEditor.h"
+#include "SLevelViewport.h"
+#include "HeadMountedDisplay.h"
+#include "Kismet/HeadMountedDisplayFunctionLibrary.h"	// For EHMDWornState::Type
+#include "EditorWorldManager.h"
+
+FVREditorModeManager::FVREditorModeManager() :
+	CurrentVREditorMode( nullptr ),
+	PreviousVREditorMode( nullptr ),
+	bEnableVRRequest( false ),
+	bPlayStartedFromVREditor( false ),
+	LastWorldToMeters( 100 ),
+	SavedWorldToMeters( 0 ),
+	HMDWornState( EHMDWornState::Unknown ),
+	TimeSinceHMDChecked( 0.0f )
+{
+	// Tell the level editor we want to be notified when selection changes
+	FLevelEditorModule& LevelEditor = FModuleManager::LoadModuleChecked<FLevelEditorModule>( "LevelEditor" );
+	LevelEditor.OnMapChanged().AddRaw( this, &FVREditorModeManager::OnMapChanged );
+}
+
+FVREditorModeManager::~FVREditorModeManager()
+{
+	//FEditorDelegates::MapChange.RemoveAll( this );
+	CurrentVREditorMode = nullptr;
+	PreviousVREditorMode = nullptr;
+}
+
+void FVREditorModeManager::Tick( const float DeltaTime )
+{
+	//@todo vreditor: Make the timer a configurable variable and/or change the polling system to an event-based one.
+	TimeSinceHMDChecked += DeltaTime;
+
+	// You can only auto-enter VR if the setting is enabled. Other criteria are that the VR Editor is enabled in experimental settings, that you are not in PIE, and that the editor is foreground.
+	bool bCanAutoEnterVR = ( GetDefault<UEditorExperimentalSettings>()->bEnableAutoVREditMode ) && ( GetDefault<UEditorExperimentalSettings>()->bEnableVREditing ) && !( GEditor->PlayWorld ) && FPlatformProcess::IsThisApplicationForeground();
+
+	if( GEngine != nullptr && GEngine->HMDDevice.IsValid() )
+	{
+		// Only check whether you are wearing the HMD every second, if you are allowed to auto-enter VR, and if your HMD state has changed since the last check. 
+		if( ( TimeSinceHMDChecked >= 1.0f ) && bCanAutoEnterVR && ( HMDWornState != GEngine->HMDDevice->GetHMDWornState() ) )
+		{
+			TimeSinceHMDChecked = 0.0f;
+			HMDWornState = GEngine->HMDDevice->GetHMDWornState();
+			if( HMDWornState == EHMDWornState::Worn )
+			{
+				EnableVREditor( true, false );
+			}
+			else if( HMDWornState == EHMDWornState::NotWorn )
+			{
+				EnableVREditor( false, false );
+			}
+		}
+	}
+
+	if( IsVREditorActive() )
+	{
+		UVREditorMode* VREditorMode = CurrentVREditorMode;
+		if( VREditorMode->WantsToExitMode() )
+		{
+			CloseVREditor();
+
+			// Start the session if that was the reason to stop the VR Editor mode
+			if( VREditorMode->GetExitType() == EVREditorExitType::PIE_VR  || 
+				VREditorMode->GetExitType() == EVREditorExitType::SIE_VR )
+			{
+				const bool bHMDIsReady = ( GEngine && GEngine->HMDDevice.IsValid() && GEngine->HMDDevice->IsHMDConnected() );
+
+				const FVector* StartLoc = NULL;
+				const FRotator* StartRot = NULL;
+
+				if( VREditorMode->GetExitType() == EVREditorExitType::PIE_VR )
+				{
+					GEditor->RequestPlaySession( true, NULL, false, StartLoc, StartRot, -1, false, bHMDIsReady, false );
+				}
+				else if( VREditorMode->GetExitType() == EVREditorExitType::SIE_VR )
+				{
+					//@todo VREditor: Where request should be to start the VR Editor in simulate mode
+					//TSharedPtr<ILevelViewport> LevelViewport = MakeShareable( &VREditorMode->GetLevelViewportPossessedForVR() ); //@todo VREditor: Must be above CloseVREditor
+					//GEditor->RequestPlaySession( true, LevelViewport, true /*bSimulateInEditor*/, StartLoc, StartRot, -1, false, bHMDIsReady, false );
+				}
+				bPlayStartedFromVREditor = true;
+			}
+		}
+	}
+	// Only check for input if we started this play session from the VR Editor
+	else if( bPlayStartedFromVREditor && GEditor->PlayWorld && !bEnableVRRequest )
+	{
+		// Shutdown PIE if we came from the VR Editor and we are not already requesting to start the VR Editor and when any of the players is holding down the required input
+		const float ShutDownInputKeyTime = 1.0f;
+		TArray<APlayerController*> PlayerControllers;
+		GEngine->GetAllLocalPlayerControllers(PlayerControllers);
+		for(APlayerController* PlayerController : PlayerControllers)
+		{
+			if((PlayerController->GetInputKeyTimeDown( EKeys::MotionController_Left_Grip1 ) > ShutDownInputKeyTime || PlayerController->GetInputKeyTimeDown( EKeys::MotionController_Left_Grip2 ) > ShutDownInputKeyTime ) &&
+				(PlayerController->GetInputKeyTimeDown( EKeys::MotionController_Right_Grip1 ) > ShutDownInputKeyTime || PlayerController->GetInputKeyTimeDown( EKeys::MotionController_Right_Grip2 ) > ShutDownInputKeyTime ) &&
+				PlayerController->GetInputKeyTimeDown( EKeys::MotionController_Right_Trigger ) > ShutDownInputKeyTime &&
+				PlayerController->GetInputKeyTimeDown( EKeys::MotionController_Left_Trigger ) > ShutDownInputKeyTime)
+			{
+				GEditor->RequestEndPlayMap();
+				bEnableVRRequest = true;
+				break;
+			}
+		}
+	}
+	else
+	{
+		// Start the VR Editor mode
+		if( bEnableVRRequest )
+		{
+			EnableVREditor( true, false );
+			bEnableVRRequest = false;
+		}
+	}
+}
+
+void FVREditorModeManager::EnableVREditor( const bool bEnable, const bool bForceWithoutHMD )
+{
+	// Don't do anything when the current VR Editor is already in the requested state
+	if( bEnable != IsVREditorActive() )
+	{
+		if( bEnable && ( IsVREditorAvailable() || bForceWithoutHMD ) )
+		{
+			StartVREditorMode( bForceWithoutHMD );
+		}
+		else if( !bEnable )
+		{
+			CloseVREditor();
+		}
+	}
+}
+
+bool FVREditorModeManager::IsVREditorActive() const
+{
+	return CurrentVREditorMode != nullptr && CurrentVREditorMode->IsActive();
+}
+
+
+bool FVREditorModeManager::IsVREditorAvailable() const
+{
+	const bool bHasHMDDevice = GEngine->HMDDevice.IsValid() && GEngine->HMDDevice->IsHMDEnabled();
+	return bHasHMDDevice;
+}
+
+
+void FVREditorModeManager::AddReferencedObjects( FReferenceCollector& Collector )
+{
+	Collector.AddReferencedObject( CurrentVREditorMode );
+}
+
+void FVREditorModeManager::StartVREditorMode( const bool bForceWithoutHMD )
+{
+	SavedWorldToMeters = GWorld->GetWorldSettings()->WorldToMeters; //@todo VREditor: Do not use GWorld
+
+	// Set the WorldToMeters scale when we stopped playing PIE that was started by the VR Editor to that VR Editor sessions latest WorldToMeters
+	if( bPlayStartedFromVREditor )
+	{
+		SetDirectWorldToMeters( LastWorldToMeters );
+		bPlayStartedFromVREditor = false;
+	}
+
+	TSharedPtr<FEditorWorldWrapper> EditorWorld = GEditor->GetEditorWorldManager()->GetEditorWorldWrapper( GWorld );
+	UVREditorMode* ModeFromWorld = EditorWorld->GetVREditorMode();
+
+	CurrentVREditorMode = ModeFromWorld;
+	CurrentVREditorMode->SetActuallyUsingVR( !bForceWithoutHMD );
+	CurrentVREditorMode->Enter();
+}
+
+void FVREditorModeManager::CloseVREditor()
+{
+	//Store the current WorldToMeters before exiting the mode
+	LastWorldToMeters = GWorld->GetWorldSettings()->WorldToMeters;
+
+	if( CurrentVREditorMode != nullptr )
+	{
+		CurrentVREditorMode->Exit();
+		PreviousVREditorMode = CurrentVREditorMode;
+		CurrentVREditorMode = nullptr;
+	}
+
+	SetDirectWorldToMeters( SavedWorldToMeters );
+	SavedWorldToMeters = 0.0f;
+}
+
+void FVREditorModeManager::SetDirectWorldToMeters( const float NewWorldToMeters )
+{
+	GWorld->GetWorldSettings()->WorldToMeters = NewWorldToMeters; //@todo VREditor: Do not use GWorld
+	ENGINE_API extern float GNewWorldToMetersScale;
+	GNewWorldToMetersScale = 0.0f;
+}
+
+void FVREditorModeManager::OnMapChanged( UWorld* World, EMapChangeType MapChangeType )
+{
+	if( CurrentVREditorMode && CurrentVREditorMode->IsActive() )
+	{
+		CloseVREditor();
+		bEnableVRRequest = true;
+	}
+	CurrentVREditorMode = nullptr;
+}
\ No newline at end of file
diff --git a/Engine/Source/Editor/VREditor/VREditorModeManager.h b/Engine/Source/Editor/VREditor/VREditorModeManager.h
new file mode 100644
--- /dev/null
+++ b/Engine/Source/Editor/VREditor/VREditorModeManager.h
@@ -0,0 +1,71 @@
+// Copyright 1998-2016 Epic Games, Inc. All Rights Reserved.
+
+#pragma once
+
+class UVREditorMode;
+
+/**
+ * Manages starting and closing the VREditor mode
+ */
+class FVREditorModeManager : public FGCObject
+{
+
+public:
+
+	FVREditorModeManager();
+	~FVREditorModeManager();
+
+	/** Ticks to detect entering and closing the VR Editor */
+	void Tick( const float DeltaTime );
+
+	/** Start or stop the VR Editor */
+	void EnableVREditor( const bool bEnable, const bool bForceWithoutHMD );
+
+	/** If the VR Editor is currently running */
+	bool IsVREditorActive() const;
+
+	/** If the VR Editor is currently available */
+	bool IsVREditorAvailable() const;
+
+	virtual void AddReferencedObjects( FReferenceCollector& Collector );
+
+private:
+
+	/** Saves the WorldToMeters and enters the mode belonging to GWorld */
+	void StartVREditorMode( const bool bForceWithoutHMD );
+
+	/** Closes the current VR Editor if any and sets the WorldToMeters to back to the one from before entering the VR mode */
+	void CloseVREditor();
+
+	/** Directly set the GWorld WorldToMeters */
+	void SetDirectWorldToMeters( const float NewWorldToMeters );
+	
+	/** On level changed */
+	void OnMapChanged( UWorld* World, EMapChangeType MapChangeType );
+	
+	/** The current mode, nullptr if none */
+	UPROPERTY()
+	UVREditorMode* CurrentVREditorMode;
+
+	/** The previous mode, nullptr if none */
+	UPROPERTY()
+	UVREditorMode* PreviousVREditorMode;
+
+	/** If the VR Editor mode needs to be enabled next tick */
+	bool bEnableVRRequest;
+
+	/** If PIE session started from within the VR Editor, so we can go back to the VR Editor when closing PIE with the motion controllers */
+	bool bPlayStartedFromVREditor;
+
+	/** Saved last world to meters scale from last VR Editor session so we can restore it when entering the VR Editor when PIE starter from VR Editor */
+	float LastWorldToMeters;
+
+	/** Saved world to meters scale from the editor when entering the VR Editor mode, to restore it when closing the editor */
+	float SavedWorldToMeters;
+
+	/** True when we detect that the user is wearing the HMD */
+	EHMDWornState::Type HMDWornState;
+
+	/** Timer for checking if the user is wearing the HMD */
+	float TimeSinceHMDChecked;
+};
\ No newline at end of file
diff --git a/Engine/Source/Editor/VREditor/VREditorModule.cpp b/Engine/Source/Editor/VREditor/VREditorModule.cpp
--- a/Engine/Source/Editor/VREditor/VREditorModule.cpp
+++ b/Engine/Source/Editor/VREditor/VREditorModule.cpp
@@ -2,22 +2,13 @@
 
 #include "VREditorModule.h"
 #include "IVREditorModule.h"
-#include "EditorModeRegistry.h"
-#include "VREditorMode.h"
-#include "LevelEditor.h"
-#include "HeadMountedDisplay.h"
+#include "VREditorModeManager.h"
 #include "Kismet/HeadMountedDisplayFunctionLibrary.h"	// For EHMDWornState::Type
 
 class FVREditorModule : public IVREditorModule, public FTickableEditorObject
 {
 public:
-	FVREditorModule() :
-		bEnableVRRequest( false ),
-		bPlayStartedFromVREditor( false ),
-		LastWorldToMeters( 100 ),
-		SavedWorldToMeters( 0 ),
-		HMDWornState( EHMDWornState::Unknown ),
-		TimeSinceHMDChecked( 0.0f )
+	FVREditorModule()
 	{
 	}
 
@@ -31,13 +22,10 @@ public:
 	}
 
 	// IVREditorModule overrides
-	virtual FEditorModeID GetVREditorModeID() const override
-	{
-		return FVREditorMode::GetVREditorModeID();
-	}
 	virtual bool IsVREditorEnabled() const override;
 	virtual bool IsVREditorAvailable() const override;
 	virtual void EnableVREditor( const bool bEnable, const bool bForceWithoutHMD ) override;
+	virtual bool IsVREditorModeActive() override;
 
 	// FTickableEditorObject overrides
 	virtual void Tick( float DeltaTime ) override;
@@ -55,29 +43,8 @@ public:
 
 private:
 
-
-	/** Sets the world settings directly and prevents the last GNewWorldToMetersScale from the VREditor to reset it by setting it 0 */
-	void SetDirectWorldToMeters( const float NewWorldToMeters );
-
-	/** If the VR Editor mode needs to be enabled next tick */
-	bool bEnableVRRequest;
-
-	/** If PIE session started from within the VR Editor, so we can go back to the VR Editor when closing PIE with the motion controllers */
-	bool bPlayStartedFromVREditor;
-
-	/** WorldToMeters from last VR Editor session */
-	float LastWorldToMeters;
-
-	/** Saved last world to meters scale from last VR Editor session so we can restore it when entering the VR Editor when PIE starter from VR Editor */
-	float SavedWorldToMeters;
-
-	/** True when we detect that the user is wearing the HMD */
-	EHMDWornState::Type HMDWornState;
-
-	/** Timer for checking if the user is wearing the HMD */
-	float TimeSinceHMDChecked;
-
-
+	/** Handles turning VR Editor mode on and off */
+	FVREditorModeManager ModeManager;
 };
 
 namespace VREd
@@ -85,94 +52,37 @@ namespace VREd
 	static FAutoConsoleCommand ForceVRMode( TEXT( "VREd.ForceVRMode" ), TEXT( "Toggles VREditorMode, even if not in immersive VR" ), FConsoleCommandDelegate::CreateStatic( &FVREditorModule::ToggleForceVRMode ) );
 }
 
-
 void FVREditorModule::StartupModule()
 {
-	FEditorModeRegistry::Get().RegisterMode<FVREditorMode>( FVREditorMode::GetVREditorModeID(), NSLOCTEXT( "VREditor", "ModeName", "VR" ) );
 }
 
-
 void FVREditorModule::ShutdownModule()
 {
-	FEditorModeRegistry::Get().UnregisterMode( FVREditorMode::GetVREditorModeID() );
 }
 
-
 void FVREditorModule::PostLoadCallback()
 {
-	GLevelEditorModeTools().ActivateDefaultMode();
 }
 
-
 bool FVREditorModule::IsVREditorEnabled() const
 {
-	return GLevelEditorModeTools().IsModeActive( FVREditorMode::GetVREditorModeID() );
+	return ModeManager.IsVREditorActive();
 }
-
-
+													
 bool FVREditorModule::IsVREditorAvailable() const
 {
-	const bool bHasHMDDevice = GEngine->HMDDevice.IsValid() && GEngine->HMDDevice->IsHMDEnabled();
-	return bHasHMDDevice;
-}
+	return ModeManager.IsVREditorAvailable();
+}	
 
 
 void FVREditorModule::EnableVREditor( const bool bEnable, const bool bForceWithoutHMD )
 {
-	// @todo vreditor: Need to be careful about Play in VR interfering with this.  While in Editor VR, PIE should
-	// work differently (in VR play experience).  And when outside of Editor VR, this feature should be disabled in the UI
-
-	TSharedPtr< ILevelEditor > LevelEditor;
-	{
-		FLevelEditorModule* LevelEditorModulePtr = FModuleManager::GetModulePtr<FLevelEditorModule>( "LevelEditor" );
-		if( LevelEditorModulePtr != nullptr )
-		{
-			// NOTE: Can be nullptr, especially during shutdown!
-			LevelEditor = LevelEditorModulePtr->GetFirstLevelEditor();
-		}
-	}
-
-	if( bEnable != IsVREditorEnabled() )
-	{
-		if( bEnable && ( IsVREditorAvailable() || bForceWithoutHMD ) && LevelEditor.IsValid() )
-		{
-			FVREditorMode::SetActuallyUsingVR( !bForceWithoutHMD );
-
-			SavedWorldToMeters = GWorld->GetWorldSettings()->WorldToMeters;
-
-			// Set the WorldToMeters scale when we stopped playing PIE that was started by the VR Editor to that VR Editor sessions latest WorldToMeters
-			if ( bPlayStartedFromVREditor )
-			{
-				SetDirectWorldToMeters( LastWorldToMeters );
-				bPlayStartedFromVREditor = false;
-			}
-		
-			// Activate the mode right away.  We expect it to stay active forever!
-			GLevelEditorModeTools().AddDefaultMode( FVREditorMode::GetVREditorModeID() );
-			GLevelEditorModeTools().ActivateMode( FVREditorMode::GetVREditorModeID() );
-
-		}
-		else if( !bEnable )
-		{ 
-			if( LevelEditor.IsValid() )
-			{
-				// Shut off VREditor mode
-				GLevelEditorModeTools().RemoveDefaultMode( FVREditorMode::GetVREditorModeID() );
-				GLevelEditorModeTools().DeactivateMode( FVREditorMode::GetVREditorModeID() );
-			}
-
-			SetDirectWorldToMeters( SavedWorldToMeters );
-			SavedWorldToMeters = 0.0f;
-		}
-	}
+	ModeManager.EnableVREditor( bEnable, bForceWithoutHMD );
 }
 
-
-void FVREditorModule::SetDirectWorldToMeters( const float NewWorldToMeters )
+bool FVREditorModule::IsVREditorModeActive()
 {
-	GWorld->GetWorldSettings()->WorldToMeters = NewWorldToMeters;
-	ENGINE_API extern float GNewWorldToMetersScale;
-	GNewWorldToMetersScale = 0.0f;
+	return ModeManager.IsVREditorActive();
 }
 
 void FVREditorModule::ToggleForceVRMode()
@@ -182,88 +92,9 @@ void FVREditorModule::ToggleForceVRMode()
 	Self.EnableVREditor( !Self.IsVREditorEnabled(), bForceWithoutHMD );
 }
 
-
 void FVREditorModule::Tick( float DeltaTime )
 {
-	//@todo vreditor: Make the timer a configurable variable and/or change the polling system to an event-based one.
-	TimeSinceHMDChecked += DeltaTime;
-	// You can only auto-enter VR if the setting is enabled. Other criteria are that the VR Editor is enabled in experimental settings, that you are not in PIE, and that the editor is foreground.
-	bool bCanAutoEnterVR = (GetDefault<UEditorExperimentalSettings>()->bEnableAutoVREditMode) && (GetDefault<UEditorExperimentalSettings>()->bEnableVREditing) && !(GEditor->PlayWorld) && FPlatformProcess::IsThisApplicationForeground();
-
-	if ( GEngine != nullptr && GEngine->HMDDevice.IsValid() )
-	{
-		// Only check whether you are wearing the HMD every second, if you are allowed to auto-enter VR, and if your HMD state has changed since the last check. 
-		if ((TimeSinceHMDChecked >= 1.0f) && bCanAutoEnterVR && (HMDWornState != GEngine->HMDDevice->GetHMDWornState()))
-		{
-			TimeSinceHMDChecked = 0.0f;
-			HMDWornState = GEngine->HMDDevice->GetHMDWornState();
-			if (HMDWornState == EHMDWornState::Worn)
-			{
-				EnableVREditor(true, false);
-			}
-			else if (HMDWornState == EHMDWornState::NotWorn)
-			{
-				EnableVREditor(false, false);
-			}
-		}
-	}
-
-	if( IsVREditorEnabled() )
-	{
-		FVREditorMode* VREditorMode = static_cast<FVREditorMode*>(GLevelEditorModeTools().FindMode(FVREditorMode::GetVREditorModeID()));
-		check(VREditorMode != nullptr);
-
-
-		if (VREditorMode->WantsToExitMode())
-		{
-			//Store the current WorldToMeters before exiting the mode
-			LastWorldToMeters = GWorld->GetWorldSettings()->WorldToMeters;
-
-			EnableVREditor( false, false );
-
-			// Start the session if that was the reason to stop the VR Editor mode
-			if ( VREditorMode->GetExitType() == EVREditorExitType::PIE_VR )
-			{
-				const bool bHMDIsReady = ( GEngine && GEngine->HMDDevice.IsValid() && GEngine->HMDDevice->IsHMDConnected() );
-
-				const FVector* StartLoc = NULL;
-				const FRotator* StartRot = NULL;
-				GEditor->RequestPlaySession( true, NULL, false, StartLoc, StartRot, -1, false, bHMDIsReady, false );
-				bPlayStartedFromVREditor = true;
-			}
-		}
-	}
-
-	// Only check for input if we started this play session from the VR Editor
-	else if( bPlayStartedFromVREditor && GEditor->PlayWorld && !bEnableVRRequest )
-	{
-		// Shutdown PIE if we came from the VR Editor and we are not already requesting to start the VR Editor and when any of the players is holding down the required input
-		const float ShutDownInputKeyTime = 0.5f;
-		TArray<APlayerController*> PlayerControllers;
-		GEngine->GetAllLocalPlayerControllers( PlayerControllers );
-		for ( APlayerController* PlayerController : PlayerControllers )
-		{
-			if ( ( PlayerController->GetInputKeyTimeDown( EKeys::MotionController_Left_Grip1 ) > ShutDownInputKeyTime || PlayerController->GetInputKeyTimeDown( EKeys::MotionController_Left_Grip2 ) > ShutDownInputKeyTime ) &&
-				( PlayerController->GetInputKeyTimeDown( EKeys::MotionController_Right_Grip1 ) > ShutDownInputKeyTime || PlayerController->GetInputKeyTimeDown( EKeys::MotionController_Right_Grip2 ) > ShutDownInputKeyTime ) &&
-				PlayerController->GetInputKeyTimeDown( EKeys::MotionController_Right_Trigger ) > ShutDownInputKeyTime &&
-				PlayerController->GetInputKeyTimeDown( EKeys::MotionController_Left_Trigger ) > ShutDownInputKeyTime )
-			{
-				GEditor->RequestEndPlayMap();
-				bEnableVRRequest = true;
-				break;
-			}
-		}
-	}
-	else
-	{
-		// Start the VR Editor mode
-		if ( bEnableVRRequest )
-		{
-			EnableVREditor( true, false );
-			bEnableVRRequest = false;
-		}
-		
-	}
+	ModeManager.Tick( DeltaTime );
 }
 
 
diff --git a/Engine/Source/Editor/VREditor/VREditorMotionControllerInteractor.cpp b/Engine/Source/Editor/VREditor/VREditorMotionControllerInteractor.cpp
--- a/Engine/Source/Editor/VREditor/VREditorMotionControllerInteractor.cpp
+++ b/Engine/Source/Editor/VREditor/VREditorMotionControllerInteractor.cpp
@@ -7,6 +7,7 @@
 #include "VREditorUISystem.h"
 #include "VREditorFloatingText.h"
 #include "VREditorDockableWindow.h"
+#include "VIBaseTransformGizmo.h"
 
 #include "IMotionController.h"
 #include "IHeadMountedDisplay.h"
@@ -108,7 +109,7 @@ UVREditorMotionControllerInteractor::~UVREditorMotionControllerInteractor()
 }
 
 
-void UVREditorMotionControllerInteractor::Init( class FVREditorMode* InVRMode )
+void UVREditorMotionControllerInteractor::Init( class UVREditorMode* InVRMode )
 {
 	Super::Init( InVRMode );
 	bHaveMotionController = true;
@@ -141,8 +142,6 @@ void UVREditorMotionControllerInteractor::Init( class FVREditorMode* InVRMode )
 		AddKeyAction( EKeys::MotionController_Right_Grip1, FViewportActionKeyInput( ViewportWorldActionTypes::WorldMovement ) );
 		AddKeyAction( UVREditorMotionControllerInteractor::MotionController_Right_FullyPressedTriggerAxis, FViewportActionKeyInput( ViewportWorldActionTypes::SelectAndMove ) );
 		AddKeyAction( UVREditorMotionControllerInteractor::MotionController_Right_LightlyPressedTriggerAxis, FViewportActionKeyInput( ViewportWorldActionTypes::SelectAndMove_LightlyPressed ) );
-		AddKeyAction( EKeys::MotionController_Right_Thumbstick, FViewportActionKeyInput( VRActionTypes::ConfirmRadialSelection ) );
-
 		AddKeyAction( SteamVRControllerKeyNames::Touch1, FViewportActionKeyInput( VRActionTypes::Touch ) );
 		AddKeyAction( EKeys::MotionController_Right_TriggerAxis, FViewportActionKeyInput( UVREditorMotionControllerInteractor::TriggerAxis ) );
 		AddKeyAction( EKeys::MotionController_Right_Thumbstick_X, FViewportActionKeyInput( UVREditorMotionControllerInteractor::TrackpadPositionX ) );
@@ -418,7 +417,7 @@ void UVREditorMotionControllerInteractor::Tick( const float DeltaTime )
 
 	// Updating laser colors for both hands
 	{
-		FVREditorMode::EColors ResultColorType = FVREditorMode::EColors::DefaultColor;
+		UVREditorMode::EColors ResultColorType = UVREditorMode::EColors::DefaultColor;
 		float CrawlSpeed = 0.0f;
 		float CrawlFade = 0.0f;
 
@@ -431,7 +430,7 @@ void UVREditorMotionControllerInteractor::Tick( const float DeltaTime )
 
 		if ( bIsDraggingWorldWithTwoHands )
 		{
-			ResultColorType = FVREditorMode::EColors::WorldDraggingColor_TwoHands;
+			ResultColorType = UVREditorMode::EColors::WorldDraggingColor_TwoHands;
 		}
 		else if ( DraggingMode == EViewportInteractionDraggingMode::World )
 		{
@@ -445,7 +444,7 @@ void UVREditorMotionControllerInteractor::Tick( const float DeltaTime )
 			// We can teleport in this mode, so animate the laser a bit
 			CrawlFade = 1.0f;
 			CrawlSpeed = 5.0f;
-			ResultColorType = FVREditorMode::EColors::WorldDraggingColor_OneHand;
+			ResultColorType = UVREditorMode::EColors::WorldDraggingColor_OneHand;
 			//			}
 		}
 		else if ( DraggingMode == EViewportInteractionDraggingMode::ActorsAtLaserImpact ||
@@ -456,7 +455,7 @@ void UVREditorMotionControllerInteractor::Tick( const float DeltaTime )
 			DraggingMode == EViewportInteractionDraggingMode::Interactable ||
 			( GetVRMode().GetUISystem().IsInteractorDraggingDockUI( this ) && GetVRMode().GetUISystem().IsDraggingDockUI() ) )
 		{
-			ResultColorType = FVREditorMode::EColors::SelectionColor;
+			ResultColorType = UVREditorMode::EColors::SelectionColor;
 		}
 
 		const FLinearColor ResultColor = GetVRMode().GetColor( ResultColorType );
diff --git a/Engine/Source/Editor/VREditor/VREditorMotionControllerInteractor.h b/Engine/Source/Editor/VREditor/VREditorMotionControllerInteractor.h
--- a/Engine/Source/Editor/VREditor/VREditorMotionControllerInteractor.h
+++ b/Engine/Source/Editor/VREditor/VREditorMotionControllerInteractor.h
@@ -19,7 +19,7 @@ public:
 	virtual ~UVREditorMotionControllerInteractor();
 	
 	// UVREditorInteractor overrides
-	virtual void Init( class FVREditorMode* InVRMode ) override;
+	virtual void Init( class UVREditorMode* InVRMode ) override;
 	/** Gets the trackpad slide delta */
 	virtual float GetSlideDelta() override;
 
diff --git a/Engine/Source/Editor/VREditor/VREditorQuickMenu.cpp b/Engine/Source/Editor/VREditor/VREditorQuickMenu.cpp
--- a/Engine/Source/Editor/VREditor/VREditorQuickMenu.cpp
+++ b/Engine/Source/Editor/VREditor/VREditorQuickMenu.cpp
@@ -4,7 +4,7 @@
 #include "VREditorQuickMenu.h"
 #include "VREditorUISystem.h"
 #include "VREditorMode.h"
-#include "VREditorWorldInteraction.h"
+#include "ViewportWorldInteraction.h"
 #include "VREditorFloatingUI.h"
 #include "VREditorTransformGizmo.h"
 #include "SLevelViewport.h"	// For Simulate toggle
diff --git a/Engine/Source/Editor/VREditor/VREditorRadialMenu.cpp b/Engine/Source/Editor/VREditor/VREditorRadialMenu.cpp
--- a/Engine/Source/Editor/VREditor/VREditorRadialMenu.cpp
+++ b/Engine/Source/Editor/VREditor/VREditorRadialMenu.cpp
@@ -6,7 +6,7 @@
 #include "VREditorFloatingUI.h"
 #include "VREditorUISystem.h"
 #include "VREditorMode.h"
-#include "VREditorWorldInteraction.h"
+#include "ViewportWorldInteraction.h"
 #include "VREditorTransformGizmo.h"
 #include "VREditorInteractor.h"
 #include "VREditorMotionControllerInteractor.h"
@@ -222,11 +222,7 @@ void UVREditorRadialMenu::OnPasteButtonClicked()
 
 void UVREditorRadialMenu::OnSnapActorsToGroundClicked()
 {
-	UVREditorWorldInteraction* WorldInteraction = Cast<UVREditorWorldInteraction>( &( GetOwner()->GetOwner().GetOwner().GetWorldInteraction() ) );
-	if ( WorldInteraction )
-	{
-		WorldInteraction->SnapSelectedActorsToGround();
-	}
+	GetOwner()->GetOwner().GetOwner().SnapSelectedActorsToGround();
 }
 
 
diff --git a/Engine/Source/Editor/VREditor/VREditorUISystem.cpp b/Engine/Source/Editor/VREditor/VREditorUISystem.cpp
--- a/Engine/Source/Editor/VREditor/VREditorUISystem.cpp
+++ b/Engine/Source/Editor/VREditor/VREditorUISystem.cpp
@@ -9,7 +9,7 @@
 #include "VREditorRadialMenu.h"
 #include "VREditorRadialMenuItem.h"
 #include "IHeadMountedDisplay.h"
-#include "VREditorWorldInteraction.h"
+#include "ViewportWorldInteraction.h"
 #include "ViewportInteractor.h"
 #include "VREditorInteractor.h"
 #include "VREditorMotionControllerInteractor.h"
@@ -169,7 +169,6 @@ void UVREditorUISystem::Shutdown()
 
 	DefaultWindowTransforms.Empty();
 
-
 	QuickMenuWidgetClass = nullptr;
 	QuickRadialWidgetClass = nullptr;
 	TutorialWidgetClass = nullptr;
@@ -190,7 +189,11 @@ void UVREditorUISystem::OnVRAction( FEditorViewportClient& ViewportClient, UView
 		{
 			if ( Action.ActionType == VRActionTypes::ConfirmRadialSelection )
 			{
-				if ( Action.Event == IE_Pressed )
+				const EViewportInteractionDraggingMode DraggingMode = Interactor->GetDraggingMode();
+				if ( Action.Event == IE_Pressed && 
+					 ( DraggingMode != EViewportInteractionDraggingMode::ActorsFreely && DraggingMode != EViewportInteractionDraggingMode::World && 
+					   ( DraggingMode != EViewportInteractionDraggingMode::AssistingDrag &&
+						 Interactor->GetOtherInteractor() != nullptr && Interactor->GetOtherInteractor()->GetDraggingMode() == EViewportInteractionDraggingMode::World ) ) )
 				{
 					UVREditorRadialMenu* RadialMenu = QuickRadialMenu->GetUserWidget<UVREditorRadialMenu>();
 					if ( IsShowingRadialMenu( VREditorInteractor ) )
@@ -199,9 +202,8 @@ void UVREditorUISystem::OnVRAction( FEditorViewportClient& ViewportClient, UView
 					}
 
 					RadialMenu->SelectCurrentItem();
+					bWasHandled = true;
 				}
-
-				bWasHandled = true;
 			}
 			else if ( Action.ActionType == ViewportWorldActionTypes::SelectAndMove_LightlyPressed )
 			{
@@ -279,6 +281,13 @@ void UVREditorUISystem::OnVRAction( FEditorViewportClient& ViewportClient, UView
 										{
 											Reply = FSlateApplication::Get().RoutePointerDownEvent(WidgetPathUnderFinger, PointerEvent);
 										}
+
+										// In case of selecting a level in the content browser the VREditormode is closed, this makes sure nothing happens after that.
+										if( !IVREditorModule::Get().IsVREditorModeActive() )
+										{
+											bWasHandled = true;
+											return;
+										}
 										
 										VREditorInteractor->SetIsClickingOnUI( true );
 										VREditorInteractor->SetIsRightClickingOnUI( bIsRightClicking );
@@ -1040,6 +1049,7 @@ void UVREditorUISystem::TryToSpawnRadialMenu( UVREditorInteractor* Interactor )
 			DraggingMode != EViewportInteractionDraggingMode::ActorsAtLaserImpact &&	// Don't show radial menu if the hand is busy dragging something around
 			DraggingMode != EViewportInteractionDraggingMode::ActorsFreely &&
 			DraggingMode != EViewportInteractionDraggingMode::ActorsWithGizmo &&
+			DraggingMode != EViewportInteractionDraggingMode::World &&
 			DraggingMode != EViewportInteractionDraggingMode::AssistingDrag &&
 			( InteractorDraggingUI == nullptr || InteractorDraggingUI != Interactor ) &&
 			!Interactor->IsHoveringOverUI();	// Don't show radial menu when aiming at a UI  (too much clutter)
diff --git a/Engine/Source/Editor/VREditor/VREditorUISystem.h b/Engine/Source/Editor/VREditor/VREditorUISystem.h
--- a/Engine/Source/Editor/VREditor/VREditorUISystem.h
+++ b/Engine/Source/Editor/VREditor/VREditorUISystem.h
@@ -46,19 +46,19 @@ public:
 	void Render( const class FSceneView* SceneView, class FViewport* Viewport, class FPrimitiveDrawInterface* PDI );
 
 	/** Sets the owner of this system */
-	void SetOwner( class FVREditorMode* NewOwner )
+	void SetOwner( class UVREditorMode* NewOwner )
 	{
 		VRMode = NewOwner;
 	}
 
 	/** Gets the owner of this system */
-	class FVREditorMode& GetOwner()
+	class UVREditorMode& GetOwner()
 	{
 		return *VRMode;
 	}
 
 	/** Gets the owner of this system (const) */
-	const class FVREditorMode& GetOwner() const
+	const class UVREditorMode& GetOwner() const
 	{
 		return *VRMode;
 	}
@@ -161,7 +161,8 @@ protected:
 protected:
 
 	/** Owning object */
-	class FVREditorMode* VRMode;
+	UPROPERTY()
+	class UVREditorMode* VRMode;
 
 	/** All of the floating UIs.  These may or may not be visible (spawned) */
 	UPROPERTY()
diff --git a/Engine/Source/Editor/VREditor/VREditorWorldInteraction.cpp b/Engine/Source/Editor/VREditor/VREditorWorldInteraction.cpp
--- a/Engine/Source/Editor/VREditor/VREditorWorldInteraction.cpp
+++ b/Engine/Source/Editor/VREditor/VREditorWorldInteraction.cpp
@@ -8,6 +8,7 @@
 #include "VREditorDockableWindow.h"
 #include "VREditorInteractor.h"
 #include "VIBaseTransformGizmo.h"
+#include "ViewportWorldInteraction.h"
 #include "SnappingUtils.h"
 #include "ScopedTransaction.h"
 
@@ -42,19 +43,26 @@ UVREditorWorldInteraction::UVREditorWorldInteraction( const FObjectInitializer&
 	FloatingUIAssetDraggedFrom( nullptr ),
 	PlacingMaterialOrTextureAsset( nullptr )
 {
-	// Find out when the user drags stuff out of a content browser
-	FEditorDelegates::OnAssetDragStarted.AddUObject( this, &UVREditorWorldInteraction::OnAssetDragStartedFromContentBrowser );
-
 	// Load sounds
 	DropMaterialOrMaterialSound = LoadObject<USoundCue>( nullptr, TEXT( "/Engine/VREditor/Sounds/VR_grab_Cue" ) );
 	check( DropMaterialOrMaterialSound != nullptr );
 }
 
-void UVREditorWorldInteraction::Shutdown()
+void UVREditorWorldInteraction::Init( UVREditorMode* InOwner, UViewportWorldInteraction* InViewportWorldInteraction )
 {
-	Super::Shutdown();
+	Owner = InOwner;
+	ViewportWorldInteraction = InViewportWorldInteraction;
+
+	// Find out when the user drags stuff out of a content browser
+	FEditorDelegates::OnAssetDragStarted.AddUObject( this, &UVREditorWorldInteraction::OnAssetDragStartedFromContentBrowser );
+
+	ViewportWorldInteraction->OnStopDragging().AddUObject( this, &UVREditorWorldInteraction::StopDragging );
+}
 
+void UVREditorWorldInteraction::Shutdown()
+{
 	FEditorDelegates::OnAssetDragStarted.RemoveAll( this );
+	ViewportWorldInteraction->OnStopDragging().RemoveAll( this );
 
 	PlacingMaterialOrTextureAsset = nullptr;
 	FloatingUIAssetDraggedFrom = nullptr;
@@ -62,10 +70,11 @@ void UVREditorWorldInteraction::Shutdown()
 	Owner = nullptr;
 }
 
+//@todo vreditor: This function was overriding ViewportWorldInteraction
 bool UVREditorWorldInteraction::IsInteractableComponent( const UActorComponent* Component ) const
 {
 	bool bResult = false;
-	bResult = UViewportWorldInteraction::IsInteractableComponent( Component );
+	//bResult = UViewportWorldInteraction::IsInteractableComponent( Component );
 
 	// Don't bother if the base class already found out that this is a custom interactable
 	if ( bResult )
@@ -75,7 +84,7 @@ bool UVREditorWorldInteraction::IsInteractableComponent( const UActorComponent*
 		{
 			static const bool bIsVREditorDemo = FParse::Param( FCommandLine::Get(), TEXT( "VREditorDemo" ) );	// @todo vreditor: Remove this when no longer needed (console variable, too!)
 			const bool bIsFrozen = bIsVREditorDemo && Component->GetOwner()->GetActorLabel().StartsWith( TEXT( "Frozen_" ) );
-			bResult = !bIsFrozen && !GetOwner().GetUISystem().IsWidgetAnEditorUIWidget( Component ); // Don't allow user to move around our UI widgets //@todo VREditor: This wouldn't be necessary if the UI used the interactable interface
+			bResult = !bIsFrozen && !Owner->GetUISystem().IsWidgetAnEditorUIWidget( Component ); // Don't allow user to move around our UI widgets //@todo VREditor: This wouldn't be necessary if the UI used the interactable interface
 		}
 	}
 	
@@ -90,19 +99,17 @@ void UVREditorWorldInteraction::StopDragging( UViewportInteractor* Interactor )
 	if ( InteractorDraggingMode == EViewportInteractionDraggingMode::Material )
 	{
 		PlaceDraggedMaterialOrTexture( Interactor );
-	}
+	}	
 
 	// If we were placing something, bring the window back
-	if ( InteractorDraggingMode == EViewportInteractionDraggingMode::ActorsAtLaserImpact ||
-		InteractorDraggingMode == EViewportInteractionDraggingMode::Material )
+	if ( FloatingUIAssetDraggedFrom != nullptr && 
+		( InteractorDraggingMode == EViewportInteractionDraggingMode::ActorsAtLaserImpact ||
+		InteractorDraggingMode == EViewportInteractionDraggingMode::Material ) )
 	{
-		GetOwner().GetUISystem().ShowEditorUIPanel( FloatingUIAssetDraggedFrom, Cast<UVREditorInteractor>( Interactor->GetOtherInteractor() ), 
+		Owner->GetUISystem().ShowEditorUIPanel( FloatingUIAssetDraggedFrom, Cast<UVREditorInteractor>( Interactor->GetOtherInteractor() ), 
 			true, false, false, false );
 		FloatingUIAssetDraggedFrom = nullptr;
 	}
-
-
-	Super::StopDragging( Interactor );
 }
 
 void UVREditorWorldInteraction::StartDraggingMaterialOrTexture( UViewportInteractor* Interactor, const FViewportActionKeyInput& Action, const FVector HitLocation, UObject* MaterialOrTextureAsset )
@@ -142,8 +149,8 @@ void UVREditorWorldInteraction::StartDraggingMaterialOrTexture( UViewportInterac
 		InteractorData.GizmoSpaceFirstDragUpdateOffsetAlongAxis = FVector::ZeroVector;	// Will be determined on first update
 		InteractorData.GizmoSpaceDragDeltaFromStartOffset = FVector::ZeroVector;	// Set every frame while dragging
 
-		bDraggedSinceLastSelection = false;
-		LastDragGizmoStartTransform = FTransform::Identity;
+		ViewportWorldInteraction->SetDraggedSinceLastSelection( false );
+		ViewportWorldInteraction->SetLastDragGizmoStartTransform( FTransform::Identity );
 
 		// Play a haptic effect when objects aTrackingTransactionre picked up
 		Interactor->PlayHapticEffect( VREd::DragHapticFeedbackStrength->GetFloat() );
@@ -152,6 +159,11 @@ void UVREditorWorldInteraction::StartDraggingMaterialOrTexture( UViewportInterac
 
 void UVREditorWorldInteraction::OnAssetDragStartedFromContentBrowser( const TArray<FAssetData>& DraggedAssets, UActorFactory* FactoryToUse )
 {
+	if( ViewportWorldInteraction == nullptr )
+	{
+		return;
+	}
+
 	const bool bIsPreview = false;	 // @todo vreditor placement: Consider supporting a "drop preview" actor (so you can cancel placement interactively)
 
 	bool bTransactionStarted = false;
@@ -159,7 +171,9 @@ void UVREditorWorldInteraction::OnAssetDragStartedFromContentBrowser( const TArr
 	// Figure out which controller pressed the button and started dragging
 	// @todo vreditor placement: This logic could misfire.  Ideally we would be routed information from the pointer event, so we can determine the hand.
 	UVREditorInteractor* PlacingWithInteractor = nullptr;
-	for ( UViewportInteractor* Interactor : GetInteractors() )
+
+	TArray<UViewportInteractor*>& Interactors = ViewportWorldInteraction->GetInteractors();
+	for ( UViewportInteractor* Interactor : Interactors )
 	{
 		UVREditorInteractor* VRInteractor = Cast<UVREditorInteractor>( Interactor );
 		FViewportActionKeyInput* SelectAndMoveAction = Interactor->GetActionWithName( ViewportWorldActionTypes::SelectAndMove );
@@ -201,7 +215,7 @@ void UVREditorWorldInteraction::OnAssetDragStartedFromContentBrowser( const TArr
 
 		FloatingUIAssetDraggedFrom = PlacingWithInteractor->GetHoveringOverWidgetComponent();
 		// Hide the UI panel that's being used to drag
-		GetOwner().GetUISystem().ShowEditorUIPanel( FloatingUIAssetDraggedFrom, PlacingWithInteractor, false, false, true, false );
+		Owner->GetUISystem().ShowEditorUIPanel( FloatingUIAssetDraggedFrom, PlacingWithInteractor, false, false, true, false );
 
 		TArray< UObject* > DroppedObjects;
 		TArray< AActor* > AllNewActors;
@@ -216,7 +230,7 @@ void UVREditorWorldInteraction::OnAssetDragStartedFromContentBrowser( const TArr
 		if( !bShouldInterpolateFromDragLocation )
 		{
 			FVector HitLocation = FVector::ZeroVector;
-			if( FindPlacementPointUnderLaser( PlacingWithInteractor, /* Out */ HitLocation ) )
+			if( ViewportWorldInteraction->FindPlacementPointUnderLaser( PlacingWithInteractor, /* Out */ HitLocation ) )
 			{
 				PlaceAt = HitLocation;
 			}
@@ -227,7 +241,7 @@ void UVREditorWorldInteraction::OnAssetDragStartedFromContentBrowser( const TArr
 			bool bCanPlace = true;
 
 			UObject* AssetObj = AssetData.GetAsset();
-			if( !ObjectTools::IsAssetValidForPlacing( GetViewportWorld(), AssetData.ObjectPath.ToString() ) )
+			if( !ObjectTools::IsAssetValidForPlacing( ViewportWorldInteraction->GetViewportWorld(), AssetData.ObjectPath.ToString() ) )
 			{
 				bCanPlace = false;
 			}
@@ -261,6 +275,7 @@ void UVREditorWorldInteraction::OnAssetDragStartedFromContentBrowser( const TArr
 			}
 
 
+			FTrackingTransaction& TrackingTransaction = ViewportWorldInteraction->GetTrackingTransaction();
 			if( bCanPlace )
 			{
 				CA_SUPPRESS(6240);
@@ -282,7 +297,7 @@ void UVREditorWorldInteraction::OnAssetDragStartedFromContentBrowser( const TArr
 				const bool bSelectNewActors = true;
 				const EObjectFlags NewObjectFlags = bIsPreview ? RF_Transient : RF_Transactional;
 
-				TArray<AActor*> NewActors = FLevelEditorViewportClient::TryPlacingActorFromObject( GetViewportWorld()->GetCurrentLevel(), AssetObj, bSelectNewActors, NewObjectFlags, FactoryToUse );
+				TArray<AActor*> NewActors = FLevelEditorViewportClient::TryPlacingActorFromObject( ViewportWorldInteraction->GetViewportWorld()->GetCurrentLevel(), AssetObj, bSelectNewActors, NewObjectFlags, FactoryToUse );
 
 				if( NewActors.Num() > 0 )
 				{
@@ -300,7 +315,7 @@ void UVREditorWorldInteraction::OnAssetDragStartedFromContentBrowser( const TArr
 								( bShouldInterpolateFromDragLocation ? 
 									VREd::SizeOfActorsOverContentBrowserThumbnail->GetFloat() :
 									VREd::OverrideSizeOfPlacedActors->GetFloat() )
-								* GetWorldScaleFactor();
+								* ViewportWorldInteraction->GetWorldScaleFactor();
 							NewActor->SetActorScale3D( FVector( DesiredSize / LocalBoundsSize ) );
 						}
 					}
@@ -311,6 +326,7 @@ void UVREditorWorldInteraction::OnAssetDragStartedFromContentBrowser( const TArr
 		// Cancel the transaction if nothing was placed
 		if( bTransactionStarted && AllNewActors.Num() == 0)
 		{
+			FTrackingTransaction& TrackingTransaction = ViewportWorldInteraction->GetTrackingTransaction();
 			--TrackingTransaction.TransCount;
 			TrackingTransaction.Cancel();
 			GEditor->DisableDeltaModification( false );
@@ -335,14 +351,14 @@ void UVREditorWorldInteraction::OnAssetDragStartedFromContentBrowser( const TArr
 
 			// Start dragging the new actor(s)
 			const bool bIsPlacingActors = true;
-			StartDraggingActors( PlacingWithInteractor, Action, GetTransformGizmoActor()->GetRootComponent(), PlaceAt, bIsPlacingActors );
+			ViewportWorldInteraction->StartDraggingActors( PlacingWithInteractor, Action, ViewportWorldInteraction->GetTransformGizmoActor()->GetRootComponent(), PlaceAt, bIsPlacingActors );
 
 			// If we're interpolating, update the target transform of the actors to use our overridden size.  When
 			// we placed them we set their size to be 'thumbnail sized', and we want them to interpolate to
 			// their actual size in the world
 			if( bShouldInterpolateFromDragLocation )
 			{
-				for( TUniquePtr<FViewportTransformable>& TransformablePtr : Transformables )
+				for( TUniquePtr<FViewportTransformable>& TransformablePtr : ViewportWorldInteraction->Transformables )
 				{
 					FViewportTransformable& Transformable = *TransformablePtr;
 
@@ -350,7 +366,7 @@ void UVREditorWorldInteraction::OnAssetDragStartedFromContentBrowser( const TArr
 
 					if( VREd::OverrideSizeOfPlacedActors->GetFloat() > KINDA_SMALL_NUMBER )
 					{
-						ObjectSize = 1.0f * GetWorldScaleFactor();
+						ObjectSize = 1.0f * ViewportWorldInteraction->GetWorldScaleFactor();
 
 						// Got an actor?
 						AActor* Actor = Cast<AActor>( Transformable.Object.Get() );
@@ -359,7 +375,7 @@ void UVREditorWorldInteraction::OnAssetDragStartedFromContentBrowser( const TArr
 							const FBox LocalSpaceBounds = Actor->CalculateComponentsBoundingBoxInLocalSpace();
 							const float LocalBoundsSize = LocalSpaceBounds.GetSize().GetAbsMax();
 
-							const float DesiredSize = VREd::OverrideSizeOfPlacedActors->GetFloat() * GetWorldScaleFactor();
+							const float DesiredSize = VREd::OverrideSizeOfPlacedActors->GetFloat() * ViewportWorldInteraction->GetWorldScaleFactor();
 							const FVector NewScale( DesiredSize / LocalBoundsSize );
 
 							Transformable.UnsnappedTargetTransform.SetScale3D( NewScale );
@@ -396,7 +412,7 @@ void UVREditorWorldInteraction::PlaceDraggedMaterialOrTexture( UViewportInteract
 			FHitResult HitResult = Interactor->GetHitResultFromLaserPointer( nullptr, bIgnoreGizmos, nullptr, bEvenIfUIIsInFront );
 			if( HitResult.Actor.IsValid() )
 			{
-				if( IsInteractableComponent( HitResult.GetComponent() ) )	// @todo vreditor placement: We don't necessarily need to restrict to only VR-interactive components here
+				if( ViewportWorldInteraction->IsInteractableComponent( HitResult.GetComponent() ) )	// @todo vreditor placement: We don't necessarily need to restrict to only VR-interactive components here
 				{
 					// Don't place materials on UI widget handles though!
 					if( Cast<AVREditorFloatingUI>( HitResult.GetComponent()->GetOwner() ) == nullptr )
@@ -431,7 +447,7 @@ void UVREditorWorldInteraction::PlaceDraggedMaterialOrTexture( UViewportInteract
 
 			if( DroppedObjAsMaterial || DroppedObjAsTexture )
 			{
-				UGameplayStatics::PlaySound2D( GetViewportWorld(), DropMaterialOrMaterialSound );
+				UGameplayStatics::PlaySound2D( ViewportWorldInteraction->GetViewportWorld(), DropMaterialOrMaterialSound );
 			}
 		}
 	}
@@ -441,7 +457,7 @@ void UVREditorWorldInteraction::PlaceDraggedMaterialOrTexture( UViewportInteract
 
 void UVREditorWorldInteraction::SnapSelectedActorsToGround()
 {
-	TSharedPtr<SLevelViewport> LevelEditorViewport = StaticCastSharedPtr<SLevelViewport>( GetViewportClient()->GetEditorViewportWidget() );
+	TSharedPtr<SLevelViewport> LevelEditorViewport = StaticCastSharedPtr<SLevelViewport>( ViewportWorldInteraction->GetViewportClient()->GetEditorViewportWidget() );
 	if ( LevelEditorViewport.IsValid() )
 	{
 		FLevelEditorModule& LevelEditorModule = FModuleManager::GetModuleChecked<FLevelEditorModule>( TEXT( "LevelEditor" ) );
@@ -451,7 +467,7 @@ void UVREditorWorldInteraction::SnapSelectedActorsToGround()
 		CommandList->ExecuteAction( Commands.SnapBottomCenterBoundsToFloor.ToSharedRef() );
 
 		// @todo vreditor: This should not be needed after to allow transformables to stop animating the transforms of actors after they come to rest
-		SetupTransformablesForSelectedActors();
+		ViewportWorldInteraction->SetupTransformablesForSelectedActors();
 	}
 }
 
diff --git a/Engine/Source/Editor/VREditor/VREditorWorldInteraction.h b/Engine/Source/Editor/VREditor/VREditorWorldInteraction.h
--- a/Engine/Source/Editor/VREditor/VREditorWorldInteraction.h
+++ b/Engine/Source/Editor/VREditor/VREditorWorldInteraction.h
@@ -5,43 +5,34 @@
 #include "ViewportWorldInteraction.h"
 #include "VREditorWorldInteraction.generated.h"
 
+class UViewportInteractor;
+class UVREditorMode;
+
 /**
  * VR Editor interaction with the 3D world
  */
 UCLASS()
-class UVREditorWorldInteraction : public UViewportWorldInteraction
+class UVREditorWorldInteraction : public UObject
 {
 	GENERATED_UCLASS_BODY()
 
 public:
 
-	// UViewportWorldInteraction overrides
-	virtual void Shutdown() override;
-	virtual bool IsInteractableComponent( const UActorComponent* Component ) const override;
-	virtual void StopDragging( class UViewportInteractor* Interactor ) override;
-
-	/** Sets the owner of this system */
-	void SetOwner( class FVREditorMode* NewOwner )
-	{
-		Owner = NewOwner;
-	}
-
-	/** Gets the owner of this system */
-	class FVREditorMode& GetOwner()
-	{
-		return *Owner;
-	}
-
-	/** Gets the owner of this system (const) */
-	const class FVREditorMode& GetOwner() const
-	{
-		return *Owner;
-	}
-
-		/** Snaps the selected objects to the ground */
+	/** Registers to events and sets initial values */
+	void Init( UVREditorMode* NewOwner, UViewportWorldInteraction* ViewportWorldInteraction );
+	
+	/** Removes registered event */
+	void Shutdown();
+
+	/** If this component is interactable by an interactor */
+	bool IsInteractableComponent( const UActorComponent* Component ) const;
+
+	/** Snaps the selected objects to the ground */
 	void SnapSelectedActorsToGround();
 
 protected:
+	/** When an interactor stops dragging */
+	void StopDragging( UViewportInteractor* Interactor );
 
 	/** Starts dragging a material, allowing the user to drop it on an object in the scene to place it */
 	void StartDraggingMaterialOrTexture( UViewportInteractor* Interactor, const FViewportActionKeyInput& Action, const FVector HitLocation, UObject* MaterialOrTextureAsset );
@@ -55,7 +46,12 @@ protected:
 protected:
 
 	/** Owning object */
-	FVREditorMode* Owner;
+	UPROPERTY()
+	UVREditorMode* Owner;
+
+	/** The actual ViewportWorldInteraction */
+	UPROPERTY()
+	UViewportWorldInteraction* ViewportWorldInteraction;
 
 	/** Sound for dropping materials and textures */
 	UPROPERTY()
diff --git a/Engine/Source/Editor/ViewportInteraction/Public/IViewportInteractionModule.h b/Engine/Source/Editor/ViewportInteraction/Public/IViewportInteractionModule.h
--- a/Engine/Source/Editor/ViewportInteraction/Public/IViewportInteractionModule.h
+++ b/Engine/Source/Editor/ViewportInteraction/Public/IViewportInteractionModule.h
@@ -32,5 +32,15 @@ public:
 	{
 		return FModuleManager::Get().IsModuleLoaded( "ViewportInteraction" );
 	}
+
+	/**
+	 * @return The ViewportWorldInteractionManager that owns entering and exiting the ViewportWorldInteraction
+	 */
+	virtual class FViewportWorldInteractionManager& GetWorldInteractionManager() = 0;
+
+	/**
+	 * Updates the ViewportWorldInteraction
+	 */
+	virtual void Tick( float DeltaTime ) = 0;
 };
 
diff --git a/Engine/Source/Editor/ViewportInteraction/Public/IViewportWorldInteractionManager.h b/Engine/Source/Editor/ViewportInteraction/Public/IViewportWorldInteractionManager.h
new file mode 100644
--- /dev/null
+++ b/Engine/Source/Editor/ViewportInteraction/Public/IViewportWorldInteractionManager.h
@@ -0,0 +1,27 @@
+// Copyright 1998-2016 Epic Games, Inc. All Rights Reserved.
+
+#pragma once
+
+class IViewportWorldInteractionManager
+{
+public:
+	DECLARE_EVENT_OneParam( IViewportWorldInteractionManager, FOnWorldInteractionTick, const float /* DeltaTime */ );
+
+	/** Gets the event for ticking before the viewport world interaction */
+	virtual FOnWorldInteractionTick& OnPreWorldInteractionTick() = 0;
+
+	/** Gets the event for ticking after the viewport world interaction */
+	virtual FOnWorldInteractionTick& OnPostWorldInteractionTick() = 0;
+
+	/** 
+	 * Sets the current ViewportWorldInteraction
+	 *
+	 * @param UViewportWorldInteraction the next ViewportWorldInteraction used
+	 */
+	virtual void SetCurrentViewportWorldInteraction( class UViewportWorldInteraction* /*WorldInteraction*/ ) = 0;
+
+	/**
+	 * Updates the current ViewportWorldInteraction
+	 */
+	virtual void Tick( float DeltaTime ) = 0;
+};
\ No newline at end of file
diff --git a/Engine/Source/Editor/ViewportInteraction/Public/ViewportWorldInteraction.h b/Engine/Source/Editor/ViewportInteraction/Public/ViewportWorldInteraction.h
--- a/Engine/Source/Editor/ViewportInteraction/Public/ViewportWorldInteraction.h
+++ b/Engine/Source/Editor/ViewportInteraction/Public/ViewportWorldInteraction.h
@@ -69,22 +69,27 @@ class VIEWPORTINTERACTION_API UViewportWorldInteraction : public UObject, public
 	
 public:
 
-	~UViewportWorldInteraction();
+	virtual ~UViewportWorldInteraction();
 
 	// IViewportWorldInteraction overrides
-	virtual void Init( const TSharedPtr<class FEditorViewportClient>& InEditorViewportClient ) override;
+	virtual void Init( UWorld* InWorld ) override;
 	virtual void Shutdown() override;
-	virtual void Tick( class FEditorViewportClient* ViewportClient, const float DeltaTime ) override;
+	virtual void Tick( const float DeltaTime ) override;
 	virtual void AddInteractor( UViewportInteractor* Interactor ) override;
 	virtual void RemoveInteractor( UViewportInteractor* Interactor ) override;
-	virtual FOnVIHoverUpdate& OnViewportInteractionHoverUpdate() override
-	{
-		return OnHoverUpdateEvent;
-	}
-	virtual FOnVIActionHandle& OnViewportInteractionInputAction() override
-	{
-		return OnInputActionEvent;
-	}
+	virtual FOnVIHoverUpdate& OnViewportInteractionHoverUpdate() override { return OnHoverUpdateEvent; }
+	virtual FOnVIActionHandle& OnViewportInteractionInputAction() override { return OnInputActionEvent; }
+	virtual FOnHandleInputKey& OnHandleKeyInput() override { return OnKeyInputEvent; }
+	virtual FOnHandleInputAxis& OnHandleAxisInput() override { return OnAxisInputEvent; }
+
+	/** Sets the Viewport to be used, this is needed to run the ViewportWorldInteraction */
+	void SetViewport( const TSharedPtr<class FEditorViewportClient>& InEditorViewportClient );
+	
+	/** Starts if bInActive is true and closes if bInActivate is false */
+	void Activate( const bool bInActivate );
+	
+	/** Returns true if this worldinteraction is currently running */
+	bool IsActive() const;
 
 	/** Pairs to interactors by setting	the other interactor for each interactor */
 	void PairInteractors( UViewportInteractor* FirstInteractor, UViewportInteractor* SecondInteractor );
@@ -152,7 +157,9 @@ public:
 	void Deselect();
 	
 	/** Called to finish a drag action with the specified interactor */
-	virtual void StopDragging( class UViewportInteractor* Interactor );
+	void StopDragging( class UViewportInteractor* Interactor );
+
+	virtual FOnStopDragging& OnStopDragging() override { return OnStopDraggingEvent; };
 
 	/** Starts dragging selected actors around, changing selection if needed.  Called when clicking in the world, or when placing new objects */
 	void StartDraggingActors( UViewportInteractor* Interactor, const FViewportActionKeyInput& Action, UActorComponent* ClickedComponent, const FVector& HitLocation, const bool bIsPlacingActors );
@@ -176,7 +183,7 @@ public:
 	bool IsTransformingActor( AActor* Actor ) const;
 
 	/** Checks to see if we're allowed to interact with the specified component */
-	virtual bool IsInteractableComponent( const UActorComponent* Component ) const;
+	bool IsInteractableComponent( const UActorComponent* Component ) const;
 
 	/** Gets the transform gizmo actor, or returns null if we currently don't have one */
 	class ABaseTransformGizmo* GetTransformGizmoActor();
@@ -219,25 +226,37 @@ public:
 
 	/** Destroys a transient actor we created earlier and nulls out the pointer */
 	void DestroyTransientActor(AActor* Actor) const;
+	
+	/** Figures out where to place an object when tracing it against the scene using a laser pointer */
+	bool FindPlacementPointUnderLaser( UViewportInteractor* Interactor, FVector& OutHitLocation );
+
+	/** Gets the tracking transactions */
+	FTrackingTransaction& GetTrackingTransaction();
+	
+	/** Creates a Transformable for each selected actor, wiping out the existing transformables */
+	void SetupTransformablesForSelectedActors();
+	
+	/** All of the objects we're currently interacting with, such as selected actors */
+	TArray< TUniquePtr< class FViewportTransformable > > Transformables;
 
 protected:
 
 	/** Allow base class to refresh on selection changed */
 	virtual void RefreshOnSelectionChanged( const bool bNewObjectsSelected, bool bAllHandlesVisible, class UActorComponent* SingleVisibleHandle, const TArray< UActorComponent* >& HoveringOverHandles );
 
-	/** Creates a Transformable for each selected actor, wiping out the existing transformables */
-	void SetupTransformablesForSelectedActors();
+private:
 
-	/** Figures out where to place an object when tracing it against the scene using a laser pointer */
-	bool FindPlacementPointUnderLaser( UViewportInteractor* Interactor, FVector& OutHitLocation );
+	/** Starts when activated */
+	void Enter();
 
-private:
+	/** Closes when deactivated */
+	void Exit();
 
 	/** Called every frame to update hover state */
-	void HoverTick( class FEditorViewportClient* ViewportClient, const float DeltaTime );
+	void HoverTick( const float DeltaTime );
 
 	/** Updates the interaction */
-	void InteractionTick( FEditorViewportClient* ViewportClient, const float DeltaTime );
+	void InteractionTick( const float DeltaTime );
 	
 	/** Called by the world interaction system when one of our components is dragged by the user.  If null is
 	passed in then we'll treat it as dragging the whole object (rather than a specific axis/handle) */
@@ -317,6 +336,9 @@ private:
 	/** Spawns a transient actor that we can use in the editor world */
 	AActor* SpawnTransientSceneActor( TSubclassOf<AActor> ActorClass, const FString& ActorName, const bool bWithSceneComponent ) const;
 
+	/** Destroys the actors */
+	void DestroyActors();
+
 	/** Gets the snap grid mesh actor */
 	AActor* GetSnapGridActor()
 	{
@@ -359,8 +381,8 @@ private:
 
 protected:
 	
-	/** All of the objects we're currently interacting with, such as selected actors */
-	TArray< TUniquePtr< class FViewportTransformable > > Transformables;
+	/** The world for this interaction */
+	UWorld* World;
 
 	/** True if we've dragged objects with either hand since the last time we selected something */
 	bool bDraggedSinceLastSelection;
@@ -381,7 +403,6 @@ private:
 	FTimespan AppTimeEntered;
 
 	/** All the interactors registered to modify the world */
-	UPROPERTY()
 	TArray< class UViewportInteractor* > Interactors;
 
 	//
@@ -389,19 +410,12 @@ private:
 	//
 
 	/** The viewport that the worldinteraction is used for */
-	TSharedPtr<class FEditorViewportClient> EditorViewportClient;
+	TOptional<class FEditorViewportClient*> EditorViewportClient;
 
 	/** The last frame that input was polled.  We keep track of this so that we can make sure to poll for latest input right before
 	its first needed each frame */
 	uint32 LastFrameNumberInputWasPolled;
 
-	//
-	// Input
-	//
-
-	/** Slate Input Processor */
-	TSharedPtr<class FViewportInteractionInputProcessor> InputProcessor;
-
 	/** The Unreal controller ID for the motion controllers we're using */
 	int32 MotionControllerID;
 
@@ -459,7 +473,6 @@ private:
 	/** Whether the gizmo should be visible or not */
 	bool bIsTransformGizmoVisible;
 
-
 	//
 	// Snap grid
 	//
@@ -483,9 +496,25 @@ private:
 	/** The current dragged interactable */
 	class IViewportInteractableInterface* DraggedInteractable;		
 	
+	//
+	// Events
+	//
+
 	/** Event that fires every frame to update hover state based on what's under the cursor.  Set bWasHandled to true if you detect something to hover. */
 	FOnVIHoverUpdate OnHoverUpdateEvent;
 
 	/** Event that is fired for when a key is pressed by an interactor */
 	FOnVIActionHandle OnInputActionEvent;
+
+	/** Event that is fired when new key input is received by the InputProcessor */
+	FOnHandleInputKey OnKeyInputEvent;
+
+	/** Event that is fired when new axis input is received by the InputProcessor */
+	FOnHandleInputAxis OnAxisInputEvent;
+
+	/** Event this is fired when a interactor stopped dragging */
+	FOnStopDragging OnStopDraggingEvent;
+
+	/** If the world interaction is active and running the tick function */
+	bool bActive;
 };
\ No newline at end of file
diff --git a/Engine/Source/Editor/ViewportInteraction/Public/ViewportWorldInteractionInterface.h b/Engine/Source/Editor/ViewportInteraction/Public/ViewportWorldInteractionInterface.h
--- a/Engine/Source/Editor/ViewportInteraction/Public/ViewportWorldInteractionInterface.h
+++ b/Engine/Source/Editor/ViewportInteraction/Public/ViewportWorldInteractionInterface.h
@@ -16,13 +16,13 @@ class VIEWPORTINTERACTION_API IViewportWorldInteractionInterface
 
 public:
 	/** Initializes the world interaction */
-	virtual void Init( const TSharedPtr<class FEditorViewportClient>& InEditorViewportClient ) = 0;
+	virtual void Init( UWorld* InWorld ) = 0;
 
 	/** Shuts down the world interaction */
 	virtual void Shutdown() = 0;
 
 	/** Main update loop including updating interaction, actions and hover */
-	virtual void Tick( class FEditorViewportClient* ViewportClient, const float DeltaTime ) = 0;
+	virtual void Tick( const float DeltaTime ) = 0;
 	
 	/** Adds interactor to the worldinteraction */	
 	virtual void AddInteractor( class UViewportInteractor* Interactor ) = 0;
@@ -38,4 +38,16 @@ public:
 	DECLARE_EVENT_FiveParams( IViewportWorldInteractionInterface, FOnVIActionHandle, class FEditorViewportClient& /* ViewportClient */, class UViewportInteractor* /* Interactor */, 
 	const struct FViewportActionKeyInput& /* Action */, bool& /* bOutIsInputCaptured */, bool& /* bWasHandled */ );
 	virtual FOnVIActionHandle& OnViewportInteractionInputAction() = 0;
+
+	/** To handle raw key input from the Inputprocessor */
+	DECLARE_EVENT_FourParams( IViewportWorldInteractionInterface, FOnHandleInputKey, const class FEditorViewportClient& /* ViewportClient */, const FKey /* Key */, const EInputEvent /* Event */, bool& /* bWasHandled */ );
+	virtual FOnHandleInputKey& OnHandleKeyInput() = 0;
+
+	/** To handle raw axix input from the Inputprocessor */
+	DECLARE_EVENT_SixParams( IViewportWorldInteractionInterface, FOnHandleInputAxis, class FEditorViewportClient& /* ViewportClient */, const int32 /* ControllerId */, const FKey /* Key */, const float /* Delta */, const float /* DeltaTime */, bool& /* bWasHandled */ );
+	virtual FOnHandleInputAxis& OnHandleAxisInput() = 0;
+
+	/** Gets the event for when an interactor stops dragging */
+	DECLARE_EVENT_OneParam( IViewportWorldInteractionInterface, FOnStopDragging, class UViewportInteractor* /** Interactor */ );
+	virtual FOnStopDragging& OnStopDragging() = 0;
 };
\ No newline at end of file
diff --git a/Engine/Source/Editor/ViewportInteraction/ViewportInteractionInputProcessor.cpp b/Engine/Source/Editor/ViewportInteraction/ViewportInteractionInputProcessor.cpp
--- a/Engine/Source/Editor/ViewportInteraction/ViewportInteractionInputProcessor.cpp
+++ b/Engine/Source/Editor/ViewportInteraction/ViewportInteractionInputProcessor.cpp
@@ -2,16 +2,17 @@
 
 #include "ViewportInteractionModule.h"
 #include "ViewportInteractionInputProcessor.h"
-#include "ViewportWorldInteraction.h"
+#include "ViewportWorldInteractionManager.h"
 #include "SLevelViewport.h"
 
-FViewportInteractionInputProcessor::FViewportInteractionInputProcessor( UViewportWorldInteraction& InWorldInteraction  )
-	: WorldInteraction( InWorldInteraction  )
+FViewportInteractionInputProcessor::FViewportInteractionInputProcessor( FViewportWorldInteractionManager* InWorldInteractionManager )
+	: WorldInteractionManager( InWorldInteractionManager )
 {
 }
 
 FViewportInteractionInputProcessor::~FViewportInteractionInputProcessor()
 {
+	WorldInteractionManager = nullptr;
 }
 
 void FViewportInteractionInputProcessor::Tick( const float DeltaTime, FSlateApplication& SlateApp, TSharedRef<ICursor> Cursor )
@@ -20,17 +21,17 @@ void FViewportInteractionInputProcessor::Tick( const float DeltaTime, FSlateAppl
 
 bool FViewportInteractionInputProcessor::HandleKeyDownEvent( FSlateApplication& SlateApp, const FKeyEvent& InKeyEvent )
 {
-	return WorldInteraction.HandleInputKey( InKeyEvent.GetKey(), InKeyEvent.IsRepeat() ? IE_Repeat : IE_Pressed );
+	return WorldInteractionManager->HandleInputKey( InKeyEvent.GetKey(), InKeyEvent.IsRepeat() ? IE_Repeat : IE_Pressed );
 }
 
 bool FViewportInteractionInputProcessor::HandleKeyUpEvent( FSlateApplication& SlateApp, const FKeyEvent& InKeyEvent )
 {
-	return WorldInteraction.HandleInputKey( InKeyEvent.GetKey(), IE_Released );
+	return WorldInteractionManager->HandleInputKey( InKeyEvent.GetKey(), IE_Released );
 }
 
 bool FViewportInteractionInputProcessor::HandleAnalogInputEvent( FSlateApplication& SlateApp, const FAnalogInputEvent& InAnalogInputEvent )
 {
-	return WorldInteraction.HandleInputAxis( InAnalogInputEvent.GetUserIndex(), InAnalogInputEvent.GetKey(), InAnalogInputEvent.GetAnalogValue(), FApp::GetDeltaTime() );
+	return WorldInteractionManager->HandleInputAxis( InAnalogInputEvent.GetUserIndex(), InAnalogInputEvent.GetKey(), InAnalogInputEvent.GetAnalogValue(), FApp::GetDeltaTime() );
 }
 
 bool FViewportInteractionInputProcessor::HandleMouseMoveEvent( FSlateApplication& SlateApp, const FPointerEvent& MouseEvent )
diff --git a/Engine/Source/Editor/ViewportInteraction/ViewportInteractionInputProcessor.h b/Engine/Source/Editor/ViewportInteraction/ViewportInteractionInputProcessor.h
--- a/Engine/Source/Editor/ViewportInteraction/ViewportInteractionInputProcessor.h
+++ b/Engine/Source/Editor/ViewportInteraction/ViewportInteractionInputProcessor.h
@@ -4,12 +4,12 @@
 
 #include "Framework/Application/IInputProcessor.h"
 
-class UViewportWorldInteraction;
+class FViewportWorldInteractionManager;
 
 class FViewportInteractionInputProcessor : public IInputProcessor
 {
 public:
-	FViewportInteractionInputProcessor ( UViewportWorldInteraction& InWorldInteraction );
+	FViewportInteractionInputProcessor ( FViewportWorldInteractionManager* InWorldInteractionManager );
 	virtual ~FViewportInteractionInputProcessor ();
 
 	virtual void Tick( const float DeltaTime, FSlateApplication& SlateApp, TSharedRef<ICursor> Cursor ) override;
@@ -29,5 +29,6 @@ public:
 private:
 
 	/** The WorldInteraction that will receive the input */
-	UViewportWorldInteraction& WorldInteraction;
+	FViewportWorldInteractionManager* WorldInteractionManager;
+
 };
diff --git a/Engine/Source/Editor/ViewportInteraction/ViewportInteractionModule.cpp b/Engine/Source/Editor/ViewportInteraction/ViewportInteractionModule.cpp
--- a/Engine/Source/Editor/ViewportInteraction/ViewportInteractionModule.cpp
+++ b/Engine/Source/Editor/ViewportInteraction/ViewportInteractionModule.cpp
@@ -1,7 +1,7 @@
 // Copyright 1998-2015 Epic Games, Inc. All Rights Reserved.
 
 #include "ViewportInteractionModule.h"
-#include "ViewportWorldInteraction.h"
+#include "ViewportWorldInteractionManager.h"
 
 FViewportInteractionModule::FViewportInteractionModule()
 {
@@ -13,3 +13,24 @@ FViewportInteractionModule::~FViewportInteractionModule()
 
 IMPLEMENT_MODULE( FViewportInteractionModule, ViewportInteraction )
 
+void FViewportInteractionModule::StartupModule()
+{
+}
+
+void FViewportInteractionModule::ShutdownModule()
+{
+}
+
+void FViewportInteractionModule::PostLoadCallback()
+{
+}
+
+FViewportWorldInteractionManager& FViewportInteractionModule::GetWorldInteractionManager()
+{
+	return WorldInteractionManager;
+}
+
+void FViewportInteractionModule::Tick( float DeltaTime )
+{
+	WorldInteractionManager.Tick( DeltaTime );
+}
\ No newline at end of file
diff --git a/Engine/Source/Editor/ViewportInteraction/ViewportInteractionModule.h b/Engine/Source/Editor/ViewportInteraction/ViewportInteractionModule.h
--- a/Engine/Source/Editor/ViewportInteraction/ViewportInteractionModule.h
+++ b/Engine/Source/Editor/ViewportInteraction/ViewportInteractionModule.h
@@ -3,12 +3,32 @@
 #pragma once
 
 #include "UnrealEd.h"
+#include "ViewportWorldInteractionManager.h"
 #include "IViewportInteractionModule.h"
 
+class FViewportWorldInteractionManager;
+
 class FViewportInteractionModule : public IViewportInteractionModule
 {
 public:
 	
 	FViewportInteractionModule();
 	virtual ~FViewportInteractionModule();
+
+	// FModuleInterface overrides
+	virtual void StartupModule() override;
+	virtual void ShutdownModule() override;
+	virtual void PostLoadCallback() override;
+	virtual bool SupportsDynamicReloading() override
+	{
+		return true;
+	}
+	virtual FViewportWorldInteractionManager& GetWorldInteractionManager() override;
+	virtual void Tick( float DeltaTime ) override;
+
+private:
+
+	/** Manager that owns the current ViewportWorldInteraction */
+	FViewportWorldInteractionManager WorldInteractionManager;
+	
 };
\ No newline at end of file
diff --git a/Engine/Source/Editor/ViewportInteraction/ViewportInteractor.cpp b/Engine/Source/Editor/ViewportInteraction/ViewportInteractor.cpp
--- a/Engine/Source/Editor/ViewportInteraction/ViewportInteractor.cpp
+++ b/Engine/Source/Editor/ViewportInteraction/ViewportInteractor.cpp
@@ -90,11 +90,8 @@ bool UViewportInteractor::HandleInputKey( const FKey Key, const EInputEvent Even
 		// Give subsystems a chance to handle actions for this interactor
 		WorldInteraction->OnViewportInteractionInputAction().Broadcast( *WorldInteraction->GetViewportClient(), this, *Action, Action->bIsInputCaptured, bHandled );
 
-		if ( !bHandled )
-		{
-			// Give the derived classes a chance to update according to the input
-			HandleInputKey( *Action, Key, Event, bHandled );
-		}
+		// Give the derived classes a chance to update according to the input
+		HandleInputKey( *Action, Key, Event, bHandled );
 
 		// Start checking on default action implementation
 		if ( !bHandled )
diff --git a/Engine/Source/Editor/ViewportInteraction/ViewportWorldInteraction.cpp b/Engine/Source/Editor/ViewportInteraction/ViewportWorldInteraction.cpp
--- a/Engine/Source/Editor/ViewportInteraction/ViewportWorldInteraction.cpp
+++ b/Engine/Source/Editor/ViewportInteraction/ViewportWorldInteraction.cpp
@@ -9,6 +9,7 @@
 #include "VIGizmoHandle.h"
 #include "VIBaseTransformGizmo.h"
 #include "Gizmo/VIPivotTransformGizmo.h"
+#include "IViewportWorldInteractionManager.h"
 
 #include "SnappingUtils.h"
 #include "ScopedTransaction.h"
@@ -209,7 +210,7 @@ static void SegmentDistToSegmentDouble( DVector A1, DVector B1, DVector A2, DVec
 	OutP2 = A2 + S2 * T2;
 }
 
-UViewportWorldInteraction::UViewportWorldInteraction( const FObjectInitializer& Initializer ) :
+UViewportWorldInteraction::UViewportWorldInteraction( const FObjectInitializer& Initializer ):
 	Super( Initializer ),
 	bDraggedSinceLastSelection( false ),
 	LastDragGizmoStartTransform( FTransform::Identity ),
@@ -224,14 +225,15 @@ UViewportWorldInteraction::UViewportWorldInteraction( const FObjectInitializer&
 	TransformGizmoActor( nullptr ),
 	TransformGizmoClass( APivotTransformGizmo::StaticClass() ),
 	GizmoLocalBounds( FBox( 0 ) ),
-	StartDragAngleOnRotation( ),
-	StartDragHandleDirection( ),
+	StartDragAngleOnRotation(),
+	StartDragHandleDirection(),
 	CurrentGizmoType( EGizmoHandleTypes::All ),
 	bIsTransformGizmoVisible( true ),
 	SnapGridActor( nullptr ),
 	SnapGridMeshComponent( nullptr ),
 	SnapGridMID( nullptr ),
-	DraggedInteractable( nullptr )
+	DraggedInteractable( nullptr ),
+	bActive( false )
 {
 }
 
@@ -240,13 +242,9 @@ UViewportWorldInteraction::~UViewportWorldInteraction()
 	Shutdown();
 }
 
-void UViewportWorldInteraction::Init( const TSharedPtr<FEditorViewportClient>& InEditorViewportClient )
+void UViewportWorldInteraction::Init( UWorld* InWorld )
 {
-	InputProcessor = MakeShareable( new FViewportInteractionInputProcessor( *this ) );
-	FSlateApplication::Get().SetInputPreProcessor( true, InputProcessor );
-
-	// Find out about selection changes
-	USelection::SelectionChangedEvent.AddUObject( this, &UViewportWorldInteraction::OnActorSelectionChanged ); //@todo viewportinteraction
+	World = InWorld;
 
 	Colors.SetNumZeroed( (int32)EColors::TotalCount );
 	{
@@ -258,11 +256,7 @@ void UViewportWorldInteraction::Init( const TSharedPtr<FEditorViewportClient>& I
 		Colors[ (int32)EColors::Dragging ] = FLinearColor::White;
 	}
 
-	EditorViewportClient = InEditorViewportClient;
 	AppTimeEntered = FTimespan::FromSeconds( FApp::GetCurrentTime() );
-
-	//Spawn the transform gizmo at init so we do not hitch when selecting our first object
-	SpawnTransformGizmoIfNeeded();
 }
 
 void UViewportWorldInteraction::Shutdown()
@@ -272,41 +266,48 @@ void UViewportWorldInteraction::Shutdown()
 	OnHoverUpdateEvent.Clear();
 	OnInputActionEvent.Clear();
 
-	if ( TransformGizmoActor != nullptr )
-	{
-		DestroyTransientActor( TransformGizmoActor );
-		TransformGizmoActor = nullptr;
-	}
-
-	if ( SnapGridActor != nullptr )
-	{
-		DestroyTransientActor( SnapGridActor );
-		SnapGridActor = nullptr;
-		SnapGridMeshComponent = nullptr;
-	}
-
-	if ( SnapGridMID != nullptr )
-	{
-		SnapGridMID->MarkPendingKill();
-		SnapGridMID = nullptr;
-	}
-
 	for ( UViewportInteractor* Interactor : Interactors )
 	{
 		Interactor->Shutdown();
+		Interactor->MarkPendingKill();
 	}
+
 	Interactors.Empty();
+	Transformables.Empty();
+	Colors.Empty();
 
-	EditorViewportClient.Reset();
+	EditorViewportClient = nullptr;
 	DraggedInteractable = nullptr;
 
+}
+
+void UViewportWorldInteraction::Enter()
+{
+	IViewportWorldInteractionManager& WorldInteractionManager = IViewportInteractionModule::Get().GetWorldInteractionManager();
+	WorldInteractionManager.SetCurrentViewportWorldInteraction( this );
+
+	// Find out about selection changes
+	USelection::SelectionChangedEvent.AddUObject( this, &UViewportWorldInteraction::OnActorSelectionChanged ); 
+	
+    //Spawn the transform gizmo at init so we do not hitch when selecting our first object
+	SpawnTransformGizmoIfNeeded();
+	TransformGizmoActor->SetIsTemporarilyHiddenInEditor( true );
+}
+
+void UViewportWorldInteraction::Exit()
+{
+	DestroyActors();
 	USelection::SelectionChangedEvent.RemoveAll( this );
+	EditorViewportClient = nullptr;
+	
+	IViewportWorldInteractionManager& WorldInteractionManager = IViewportInteractionModule::Get().GetWorldInteractionManager();
+	WorldInteractionManager.SetCurrentViewportWorldInteraction( nullptr );
 }
 
-void UViewportWorldInteraction::Tick( FEditorViewportClient* ViewportClient, const float DeltaTime )
+void UViewportWorldInteraction::Tick( const float DeltaTime )
 {
 	// Only if this is our viewport. Remember, editor modes get a chance to tick and receive input for each active viewport.
-	if ( !EditorViewportClient.IsValid() || ViewportClient != EditorViewportClient.Get() )
+	if ( !EditorViewportClient.IsSet() && !bActive )
 	{
 		return;
 	}
@@ -318,9 +319,9 @@ void UViewportWorldInteraction::Tick( FEditorViewportClient* ViewportClient, con
 	PollInputIfNeeded();
 
 	// Update hover. Note that hover can also be updated when ticking our sub-systems below.
-	HoverTick( ViewportClient, DeltaTime );
+	HoverTick( DeltaTime );
 
-	InteractionTick( ViewportClient, DeltaTime );
+	InteractionTick( DeltaTime );
 
 	// Update all the interactors
 	for ( UViewportInteractor* Interactor : Interactors )
@@ -332,16 +333,42 @@ void UViewportWorldInteraction::Tick( FEditorViewportClient* ViewportClient, con
 void UViewportWorldInteraction::AddInteractor( UViewportInteractor* Interactor )
 {
 	Interactor->SetWorldInteraction( this );
-	Interactors.AddUnique( Interactor );
+	Interactors.AddUnique( Interactor ); 
 }
 
 void UViewportWorldInteraction::RemoveInteractor( UViewportInteractor* Interactor )
 {
-	Interactor->GetOtherInteractor()->RemoveOtherInteractor();
 	Interactor->RemoveOtherInteractor();
 	Interactors.Remove( Interactor );
 }
 
+void UViewportWorldInteraction::SetViewport(const TSharedPtr<class FEditorViewportClient>& InEditorViewportClient)
+{
+	EditorViewportClient = InEditorViewportClient.Get();
+}
+
+void UViewportWorldInteraction::Activate(const bool bInActivate)
+{
+	if( bActive != bInActivate )
+	{
+		if( bInActivate )
+		{
+			Enter();
+		}
+		else
+		{
+			Exit();
+		}
+
+		bActive = bInActivate;
+	}
+}
+
+bool UViewportWorldInteraction::IsActive() const
+{
+	return bActive;
+}
+
 void UViewportWorldInteraction::PairInteractors( UViewportInteractor* FirstInteractor, UViewportInteractor* SecondInteractor )
 {
 	FirstInteractor->SetOtherInteractor( SecondInteractor );
@@ -351,16 +378,22 @@ void UViewportWorldInteraction::PairInteractors( UViewportInteractor* FirstInter
 bool UViewportWorldInteraction::HandleInputKey( const FKey Key, const EInputEvent Event )
 {
 	bool bWasHandled = false;
-
-	for ( UViewportInteractor* Interactor : Interactors )
+	if( bActive )
 	{
-		// Stop iterating if the input was handled by an Interactor
-		if ( bWasHandled )
+		OnKeyInputEvent.Broadcast( *EditorViewportClient.GetValue(), Key, Event, bWasHandled );
+		if( !bWasHandled || !bActive )
 		{
-			break;
+			for( UViewportInteractor* Interactor : Interactors )
+			{
+				bWasHandled = Interactor->HandleInputKey( Key, Event );
+				
+				// Stop iterating if the input was handled by an Interactor
+				if(bWasHandled || !bActive)
+				{
+					break;
+				}
+			}
 		}
-
-		bWasHandled = Interactor->HandleInputKey( Key, Event );
 	}
 
 	return bWasHandled;
@@ -370,15 +403,23 @@ bool UViewportWorldInteraction::HandleInputAxis( const int32 ControllerId, const
 {
 	bool bWasHandled = false;
 
-	for ( UViewportInteractor* Interactor : Interactors )
+	if(bActive)
 	{
-		// Stop iterating if the input was handled by an interactor
-		if ( bWasHandled )
+		OnAxisInputEvent.Broadcast( *EditorViewportClient.GetValue(), ControllerId, Key, Delta, DeltaTime, bWasHandled );
+
+		if( !bWasHandled || !bActive)
 		{
-			break;
-		}
+			for( UViewportInteractor* Interactor : Interactors )
+			{
+				bWasHandled = Interactor->HandleInputAxis( Key, Delta, DeltaTime );
 
-		bWasHandled = Interactor->HandleInputAxis( Key, Delta, DeltaTime );
+				// Stop iterating if the input was handled by an interactor
+				if(bWasHandled || !bActive)
+				{
+					break;
+				}
+			}
+		}
 	}
 
 	return bWasHandled;
@@ -387,9 +428,11 @@ bool UViewportWorldInteraction::HandleInputAxis( const int32 ControllerId, const
 
 FTransform UViewportWorldInteraction::GetRoomTransform() const
 {
+	check( EditorViewportClient.GetValue() );
+
 	const FTransform RoomTransform(
-		EditorViewportClient->GetViewRotation().Quaternion(),
-		EditorViewportClient->GetViewLocation(),
+		EditorViewportClient.GetValue()->GetViewRotation().Quaternion(),
+		EditorViewportClient.GetValue()->GetViewLocation(),
 		FVector( 1.0f ) );
 	return RoomTransform;
 }
@@ -420,12 +463,12 @@ FTransform UViewportWorldInteraction::GetHeadTransform() const
 
 void UViewportWorldInteraction::SetRoomTransform( const FTransform& NewRoomTransform )
 {
-	EditorViewportClient->SetViewLocation( NewRoomTransform.GetLocation() );
-	EditorViewportClient->SetViewRotation( NewRoomTransform.GetRotation().Rotator() );
+	EditorViewportClient.GetValue()->SetViewLocation( NewRoomTransform.GetLocation() );
+	EditorViewportClient.GetValue()->SetViewRotation( NewRoomTransform.GetRotation().Rotator() );
 
 	// Forcibly dirty the viewport camera location
 	const bool bDollyCamera = false;
-	EditorViewportClient->MoveViewportCamera( FVector::ZeroVector, FRotator::ZeroRotator, bDollyCamera );
+	EditorViewportClient.GetValue()->MoveViewportCamera( FVector::ZeroVector, FRotator::ZeroRotator, bDollyCamera );
 }
 
 float UViewportWorldInteraction::GetWorldScaleFactor() const
@@ -435,16 +478,23 @@ float UViewportWorldInteraction::GetWorldScaleFactor() const
 
 UWorld* UViewportWorldInteraction::GetViewportWorld() const
 {
-	if ( EditorViewportClient.IsValid() )
+	UWorld* ResultWorld = nullptr;
+	
+	if( World )
 	{
-		return EditorViewportClient->GetWorld();
+		ResultWorld = World;
 	}
-	return nullptr;
+	else if ( EditorViewportClient.IsSet() )
+	{
+		ResultWorld = EditorViewportClient.GetValue()->GetWorld();
+	}
+
+	return ResultWorld;
 }
 
 FEditorViewportClient* UViewportWorldInteraction::GetViewportClient() const
 {
-	return EditorViewportClient.Get();
+	return EditorViewportClient.GetValue();
 }
 
 void UViewportWorldInteraction::Undo()
@@ -488,7 +538,7 @@ void UViewportWorldInteraction::Deselect()
 	GEditor->SelectNone( true, true );
 }
 
-void UViewportWorldInteraction::HoverTick( FEditorViewportClient* ViewportClient, const float DeltaTime )
+void UViewportWorldInteraction::HoverTick( const float DeltaTime )
 {
 	for ( UViewportInteractor* Interactor : Interactors )
 	{
@@ -502,7 +552,7 @@ void UViewportWorldInteraction::HoverTick( FEditorViewportClient* ViewportClient
 		Interactor->UpdateHoverResult( HitHoverResult );
 
 		FVector HoverImpactPoint = HitHoverResult.ImpactPoint;
-		OnHoverUpdateEvent.Broadcast( *ViewportClient, Interactor, /* In/Out */ HoverImpactPoint, /* In/Out */ bWasHandled);
+		OnHoverUpdateEvent.Broadcast( *EditorViewportClient.GetValue(), Interactor, /* In/Out */ HoverImpactPoint, /* In/Out */ bWasHandled);
 		
 		if ( bWasHandled )
 		{
@@ -579,10 +629,11 @@ void UViewportWorldInteraction::HoverTick( FEditorViewportClient* ViewportClient
 	}
 }
 
-void UViewportWorldInteraction::InteractionTick( FEditorViewportClient* ViewportClient, const float DeltaTime )
+void UViewportWorldInteraction::InteractionTick( const float DeltaTime )
 {
 	const FTimespan CurrentTime = FTimespan::FromSeconds( FPlatformTime::Seconds() );
 	const float WorldToMetersScale = GetViewportWorld()->GetWorldSettings()->WorldToMeters;
+	FEditorViewportClient* ViewportClient = EditorViewportClient.GetValue();
 
 	// Update viewport with any objects that are currently hovered over
 	{
@@ -719,6 +770,8 @@ void UViewportWorldInteraction::InteractionTick( FEditorViewportClient* Viewport
 								}
 							}
 
+							InteractorData.DragRayLength = ( LaserPointerStart - HitLocation ).Size();
+
 							DraggedTo = HitLocation;
 						}
 						else
@@ -2191,6 +2244,8 @@ void UViewportWorldInteraction::StartDraggingActors( UViewportInteractor* Intera
 }
 void UViewportWorldInteraction::StopDragging( UViewportInteractor* Interactor )
 {
+	OnStopDraggingEvent.Broadcast( Interactor );
+
 	FViewportInteractorData& InteractorData = Interactor->GetInteractorData();
 	if ( InteractorData.DraggingMode != EViewportInteractionDraggingMode::Nothing )
 	{
@@ -2350,6 +2405,11 @@ bool UViewportWorldInteraction::FindPlacementPointUnderLaser( UViewportInteracto
 	return bHitSomething;
 }
 
+FTrackingTransaction& UViewportWorldInteraction::GetTrackingTransaction()
+{
+	return TrackingTransaction;
+}
+
 bool UViewportWorldInteraction::IsTransformingActor( AActor* Actor ) const
 {
 	bool bFoundActor = false;
@@ -2394,7 +2454,7 @@ void UViewportWorldInteraction::PollInputIfNeeded()
 
 void UViewportWorldInteraction::OnActorSelectionChanged( UObject* ChangedObject )
 {
-	if ( EditorViewportClient.IsValid() )
+	if ( EditorViewportClient.IsSet() )
 	{
 		const bool bNewObjectsSelected = true;
 		const bool bAllHandlesVisible = true;
@@ -2421,6 +2481,7 @@ void UViewportWorldInteraction::RefreshTransformGizmo( const bool bNewObjectsSel
 
 		// Make sure the gizmo is visible
 		TransformGizmoActor->GetRootComponent()->SetVisibility( true );
+		TransformGizmoActor->SetIsTemporarilyHiddenInEditor( false );
 
 		const ECoordSystem CurrentCoordSystem = GetTransformGizmoCoordinateSpace();
 		const bool bIsWorldSpaceGizmo = ( CurrentCoordSystem == COORD_World );
@@ -2629,7 +2690,7 @@ void UViewportWorldInteraction::SpawnTransformGizmoIfNeeded()
 
 		check( TransformGizmoActor != nullptr );
 		TransformGizmoActor->SetOwnerWorldInteraction( this );
-
+		
 		if ( VI::ShowTransformGizmo->GetInt() == 0 || !bIsTransformGizmoVisible )
 		{
 			TransformGizmoActor->SetIsTemporarilyHiddenInEditor( true );
@@ -2742,7 +2803,6 @@ UViewportInteractor* UViewportWorldInteraction::GetOtherInteractorIntertiaContri
 
 AActor* UViewportWorldInteraction::SpawnTransientSceneActor( TSubclassOf<AActor> ActorClass, const FString& ActorName, const bool bWithSceneComponent ) const
 {
-	UWorld* World = GetViewportWorld();
 	const bool bWasWorldPackageDirty = World->GetOutermost()->IsDirty();
 
 	// @todo vreditor: Needs respawn if world changes (map load, etc.)  Will that always restart the editor mode anyway?
@@ -2767,16 +2827,37 @@ AActor* UViewportWorldInteraction::SpawnTransientSceneActor( TSubclassOf<AActor>
 	// Don't dirty the level file after spawning a transient actor
 	if ( !bWasWorldPackageDirty )
 	{
-		GetViewportWorld()->GetOutermost()->SetDirtyFlag( false );
+		World->GetOutermost()->SetDirtyFlag( false );
 	}
 
 	return NewActor;
 }
 
 
+void UViewportWorldInteraction::DestroyActors()
+{
+	if(TransformGizmoActor != nullptr)
+	{
+		DestroyTransientActor( TransformGizmoActor );
+		TransformGizmoActor = nullptr;
+	}
+
+	if(SnapGridActor != nullptr)
+	{
+		DestroyTransientActor( SnapGridActor );
+		SnapGridActor = nullptr;
+		SnapGridMeshComponent = nullptr;
+	}
+
+	if(SnapGridMID != nullptr)
+	{
+		SnapGridMID->MarkPendingKill();
+		SnapGridMID = nullptr;
+	}
+}
+
 void UViewportWorldInteraction::DestroyTransientActor( AActor* Actor ) const
 {
-	UWorld* World = GetViewportWorld();
 	if (Actor != nullptr && World != nullptr)
 	{
 		const bool bWasWorldPackageDirty = World->GetOutermost()->IsDirty();
diff --git a/Engine/Source/Editor/ViewportInteraction/ViewportWorldInteractionManager.cpp b/Engine/Source/Editor/ViewportInteraction/ViewportWorldInteractionManager.cpp
new file mode 100644
--- /dev/null
+++ b/Engine/Source/Editor/ViewportInteraction/ViewportWorldInteractionManager.cpp
@@ -0,0 +1,59 @@
+// Copyright 1998-2015 Epic Games, Inc. All Rights Reserved.
+
+#include "ViewportInteractionModule.h"
+#include "ViewportWorldInteractionManager.h"
+#include "ViewportWorldInteraction.h"
+#include "ViewportInteractionInputProcessor.h"
+
+FViewportWorldInteractionManager::FViewportWorldInteractionManager() :
+	CurrentWorldInteraction( nullptr )
+{
+	InputProcessor = MakeShareable( new FViewportInteractionInputProcessor( this ) );
+	FSlateApplication::Get().SetInputPreProcessor( true, InputProcessor );
+}
+
+FViewportWorldInteractionManager::~FViewportWorldInteractionManager()
+{
+	CurrentWorldInteraction = nullptr;
+}
+
+void FViewportWorldInteractionManager::AddReferencedObjects( FReferenceCollector& Collector )
+{
+	Collector.AddReferencedObject( CurrentWorldInteraction );
+}
+
+void FViewportWorldInteractionManager::Tick( float DeltaTime )
+{
+	if( CurrentWorldInteraction && CurrentWorldInteraction->IsActive() )
+	{
+		// First give registered functions to tick before the world interaction is ticked
+		OnPreWorldInteractionTickEvent.Broadcast( DeltaTime );
+		CurrentWorldInteraction->Tick( DeltaTime );
+		OnPostWorldInteractionTickEvent.Broadcast( DeltaTime );
+ 	}
+}
+
+bool FViewportWorldInteractionManager::HandleInputKey( const FKey Key, const EInputEvent Event )
+{
+	bool bResult = false;
+	if(CurrentWorldInteraction && CurrentWorldInteraction->IsActive())
+	{
+		bResult = CurrentWorldInteraction->HandleInputKey( Key, Event );
+	}
+	return bResult;
+}
+
+bool FViewportWorldInteractionManager::HandleInputAxis( const int32 ControllerId, const FKey Key, const float Delta, const float DeltaTime )
+{
+	bool bResult = false;
+	if(CurrentWorldInteraction && CurrentWorldInteraction->IsActive())
+	{
+		bResult = CurrentWorldInteraction->HandleInputAxis( ControllerId, Key, Delta, DeltaTime );
+	}
+	return bResult;
+}
+
+void FViewportWorldInteractionManager::SetCurrentViewportWorldInteraction( UViewportWorldInteraction* WorldInteraction )
+{
+	CurrentWorldInteraction = WorldInteraction;
+}
diff --git a/Engine/Source/Editor/ViewportInteraction/ViewportWorldInteractionManager.h b/Engine/Source/Editor/ViewportInteraction/ViewportWorldInteractionManager.h
new file mode 100644
--- /dev/null
+++ b/Engine/Source/Editor/ViewportInteraction/ViewportWorldInteractionManager.h
@@ -0,0 +1,56 @@
+// Copyright 1998-2015 Epic Games, Inc. All Rights Reserved.
+
+#pragma once
+
+#include "IViewportWorldInteractionManager.h"
+
+class UViewportWorldInteraction;
+
+class FViewportWorldInteractionManager: public FGCObject,  public IViewportWorldInteractionManager
+{
+public:
+	FViewportWorldInteractionManager();
+	virtual ~FViewportWorldInteractionManager();
+
+	// FGCObject overrides
+	virtual void AddReferencedObjects( FReferenceCollector& Collector ) override;
+
+	// IViewportWorldInteractionManager overrides
+	virtual void Tick( float DeltaTime ) override;
+	virtual FOnWorldInteractionTick& OnPreWorldInteractionTick() override { return OnPreWorldInteractionTickEvent; };
+	virtual FOnWorldInteractionTick& OnPostWorldInteractionTick() override { return OnPostWorldInteractionTickEvent; };
+	virtual void SetCurrentViewportWorldInteraction( class UViewportWorldInteraction* WorldInteraction ) override;
+
+	//
+	// Input
+	//
+
+	/** Passes the input key to the current world interaction */
+	bool HandleInputKey( const FKey Key, const EInputEvent Event );
+
+	/** Passes the input axis to the current world interaction */
+	bool HandleInputAxis( const int32 ControllerId, const FKey Key, const float Delta, const float DeltaTime );
+
+
+private:
+
+	/** The current world interaction that is ticking */
+	class UViewportWorldInteraction* CurrentWorldInteraction;
+
+	//
+	// Events
+	//
+
+	/** Event to tick before the worldinteraction ticks */
+	FOnWorldInteractionTick OnPreWorldInteractionTickEvent;
+
+	/** Event to tick after the worldinteraction has ticked */
+	FOnWorldInteractionTick OnPostWorldInteractionTickEvent;
+
+	//
+	// Input
+	//
+
+	/** Slate Input Processor */
+	TSharedPtr<class FViewportInteractionInputProcessor> InputProcessor;
+};
\ No newline at end of file
diff --git a/Engine/Source/Runtime/Engine/Classes/Engine/Engine.h b/Engine/Source/Runtime/Engine/Classes/Engine/Engine.h
--- a/Engine/Source/Runtime/Engine/Classes/Engine/Engine.h
+++ b/Engine/Source/Runtime/Engine/Classes/Engine/Engine.h
@@ -2620,8 +2620,26 @@ public:
 	virtual void VerifyLoadMapWorldCleanup();
 
 	FWorldContext& CreateNewWorldContext(EWorldType::Type WorldType);
+
 	virtual void DestroyWorldContext(UWorld * InWorld);
 
+#if WITH_EDITOR
+	/** Triggered when a world context is destroyed. */
+	DECLARE_EVENT_OneParam(UEngine, FWorldContextDestroyedEvent, FWorldContext&);
+
+	/** Return the world context destroyed event. */
+	FWorldContextDestroyedEvent&	OnWorldContextDestroyed() { return WorldContextDestroyedEvent; }
+#endif // #if WITH_EDITOR
+
+private:
+
+#if WITH_EDITOR
+	/** Delegate broadcast when a world context is destroyed */
+	FWorldContextDestroyedEvent WorldContextDestroyedEvent;
+#endif // #if WITH_EDITOR
+
+public:
+
 	bool ShouldAbsorbAuthorityOnlyEvent();
 	bool ShouldAbsorbCosmeticOnlyEvent();
 
diff --git a/Engine/Source/Runtime/Engine/Private/UnrealEngine.cpp b/Engine/Source/Runtime/Engine/Private/UnrealEngine.cpp
--- a/Engine/Source/Runtime/Engine/Private/UnrealEngine.cpp
+++ b/Engine/Source/Runtime/Engine/Private/UnrealEngine.cpp
@@ -10600,6 +10600,7 @@ void UEngine::DestroyWorldContext(UWorld * InWorld)
 	{
 		if (WorldList[idx].World() == InWorld)
 		{
+			WorldContextDestroyedEvent.Broadcast(WorldList[idx]);
 			// Set the current world to NULL so that any external referencers are cleaned up before we remove
 			WorldList[idx].SetCurrentWorld(NULL);
 			WorldList.RemoveAt(idx);
