VR Editor: Fix for radial menu press not working from 3109859

diff --git a/Engine/Source/Editor/VREditor/VREditorUISystem.cpp b/Engine/Source/Editor/VREditor/VREditorUISystem.cpp
--- a/Engine/Source/Editor/VREditor/VREditorUISystem.cpp
+++ b/Engine/Source/Editor/VREditor/VREditorUISystem.cpp
@@ -192,8 +192,7 @@ void UVREditorUISystem::OnVRAction( FEditorViewportClient& ViewportClient, UView
 				const EViewportInteractionDraggingMode DraggingMode = Interactor->GetDraggingMode();
 				if ( Action.Event == IE_Pressed && 
 					 ( DraggingMode != EViewportInteractionDraggingMode::ActorsFreely && DraggingMode != EViewportInteractionDraggingMode::World && 
-					   ( DraggingMode != EViewportInteractionDraggingMode::AssistingDrag &&
-						 Interactor->GetOtherInteractor() != nullptr && Interactor->GetOtherInteractor()->GetDraggingMode() == EViewportInteractionDraggingMode::World ) ) )
+					   ( DraggingMode != EViewportInteractionDraggingMode::AssistingDrag && Interactor->GetOtherInteractor() != nullptr && Interactor->GetOtherInteractor()->GetDraggingMode() != EViewportInteractionDraggingMode::World ) ) )
 				{
 					UVREditorRadialMenu* RadialMenu = QuickRadialMenu->GetUserWidget<UVREditorRadialMenu>();
 					if ( IsShowingRadialMenu( VREditorInteractor ) )
