VR Editor: Added teleport without hitting objects

diff --git a/Engine/Source/Editor/VREditor/Teleporter/VREditorTeleporter.cpp b/Engine/Source/Editor/VREditor/Teleporter/VREditorTeleporter.cpp
--- a/Engine/Source/Editor/VREditor/Teleporter/VREditorTeleporter.cpp
+++ b/Engine/Source/Editor/VREditor/Teleporter/VREditorTeleporter.cpp
@@ -13,6 +13,7 @@ namespace VREd
 	static FAutoConsoleVariable TeleportLerpTime( TEXT( "VREd.TeleportLerpTime" ), 0.1f, TEXT( "The lerp time to teleport" ) );
 	static FAutoConsoleVariable TeleportOffset( TEXT( "VREd.TeleportOffset" ), 100.0f, TEXT( "The offset from the hitresult towards the controller" ) );
 	static FAutoConsoleVariable TeleportLaserPointerLength( TEXT( "VREd.TeleportLaserPointerLength" ), 500000.0f, TEXT( "Distance of the LaserPointer for teleporting" ) );
+	static FAutoConsoleVariable TeleportDistance( TEXT( "VREd.TeleportDistance" ), 500.0f, TEXT( "Default distance for teleporting when not hitting anything" ) );
 }
 
 UVREditorTeleporter::UVREditorTeleporter( const FObjectInitializer& Initializer ) :
@@ -49,11 +50,6 @@ void UVREditorTeleporter::Tick( const float DeltaTime )
 	if ( bIsTeleporting )
 	{
 		Teleport( DeltaTime );
-
-		//if ( bJustTeleported && !bIsTeleporting )
-		//{
-		//	bJustTeleported = false;
-		//}
 	}
 }
 
@@ -65,29 +61,42 @@ void UVREditorTeleporter::StartTeleport( UViewportInteractor* Interactor )
 		FVector LaserPointerStart, LaserPointerEnd;
 		if ( Interactor->GetLaserPointer( LaserPointerStart, LaserPointerEnd ) )
 		{
+			FVector EndLocation;
 			FHitResult HitResult = Interactor->GetHitResultFromLaserPointer( nullptr, true, nullptr, false, VREd::TeleportLaserPointerLength->GetFloat() );
-			if ( HitResult.bBlockingHit )
+			if( HitResult.bBlockingHit )
+			{
+				EndLocation = HitResult.Location;
+			}
+			else
+			{
+				FVector Direction = LaserPointerEnd - LaserPointerStart;
+				Direction.Normalize();
+				EndLocation = LaserPointerStart + ( Direction * ( VREd::TeleportDistance->GetFloat() * VRMode->GetWorldScaleFactor() ) );
+			}
+
+			const FVector HeadLocation = VRMode->GetHeadTransform().GetLocation();
+			const FVector RoomLocation = VRMode->GetRoomTransform().GetLocation();
+
+			FVector NewHandLocation = EndLocation;
+			if( HitResult.bBlockingHit )
 			{
-				const FVector HeadLocation = VRMode->GetHeadTransform().GetLocation();
-				const FVector RoomLocation = VRMode->GetRoomTransform().GetLocation();
-
-				const FVector Offset = ( LaserPointerStart - HitResult.Location ).GetSafeNormal() *
-					VREd::TeleportOffset->GetFloat() * VRMode->GetWorldInteraction().GetWorldScaleFactor(); 	// Offset from the hitresult, TeleportOffset is in centimeter
-				const FVector NewHandLocation = HitResult.Location + Offset;								// New world location offset from the hitresult
-				const FVector HandHeadOffset = HeadLocation - Interactor->GetTransform().GetLocation();		// The offset between the hand and the head
-				const FVector NewHeadLocation = NewHandLocation + HandHeadOffset;							// The location in world space where the head is going to be
-																											// The new roomspace location in world space
-				const FVector HeadRoomOffset = HeadLocation - RoomLocation;									// Where the head is going to be in world space
-				const FVector NewRoomSpaceLocation = NewHeadLocation - HeadRoomOffset;						// Result, where the roomspace is going to be in world space
-				
-																											// Set values to start teleporting
-				bIsTeleporting = true;
-				TeleportLerpAlpha = 0.0f;
-				TeleportStartLocation = RoomLocation;
-				TeleportGoalLocation = NewRoomSpaceLocation;
-
-				UGameplayStatics::PlaySound2D( VRMode->GetWorld(), TeleportSound );
+				const FVector Offset = ( LaserPointerStart - EndLocation ).GetSafeNormal() *
+					VREd::TeleportOffset->GetFloat() * VRMode->GetWorldInteraction().GetWorldScaleFactor(); // Offset from the hitresult, TeleportOffset is in centimeter
+				NewHandLocation = EndLocation + Offset;
 			}
+
+			const FVector HandHeadOffset = HeadLocation - Interactor->GetTransform().GetLocation();		// The offset between the hand and the head
+			const FVector NewHeadLocation = NewHandLocation + HandHeadOffset;							// The location in world space where the head is going to be
+																										// The new roomspace location in world space
+			const FVector HeadRoomOffset = HeadLocation - RoomLocation;									// Where the head is going to be in world space
+			const FVector NewRoomSpaceLocation = NewHeadLocation - HeadRoomOffset;						// Result, where the roomspace is going to be in world space
+																										// Set values to start teleporting
+			bIsTeleporting = true;
+			TeleportLerpAlpha = 0.0f;
+			TeleportStartLocation = RoomLocation;
+			TeleportGoalLocation = NewRoomSpaceLocation;
+
+			UGameplayStatics::PlaySound2D( VRMode->GetWorld(), TeleportSound );
 		}
 	}
 }
