Copying //UE4/Dev-VREditor to //UE4/Dev-Main (Source: //UE4/Dev-VREditor @ 3051763)

diff --git a/Engine/Source/Editor/VREditor/VREditorMode.h b/Engine/Source/Editor/VREditor/VREditorMode.h
--- a/Engine/Source/Editor/VREditor/VREditorMode.h
+++ b/Engine/Source/Editor/VREditor/VREditorMode.h
@@ -4,14 +4,19 @@
 
 #include "IVREditorMode.h"
 #include "IVREditorModule.h"
-#include "VirtualHand.h"
 #include "HeadMountedDisplayTypes.h"	// For EHMDDeviceType::Type
 #include "Kismet/HeadMountedDisplayFunctionLibrary.h"	// For EHMDTrackingOrigin::Type
-#include "VREditorInputProcessor.h"
 
 // Forward declare the GizmoHandleTypes that is defined in VREditorTransformGizmo
 enum class EGizmoHandleTypes : uint8;
 
+/** Reason for exiting the VR editor */
+enum EVREditorExitType
+{
+	Normal = 0,
+	PIE_VR = 1
+};
+
 /**
  * VR Editor Mode.  Extends editor viewports with functionality for VR controls and object manipulation
  */
@@ -50,9 +55,12 @@ public:
 	}
 
 	/** Call this to start exiting VR mode */
-	void StartExitingVRMode()
+	void StartExitingVRMode( const EVREditorExitType InExitType = EVREditorExitType::Normal );
+
+	/** Gets the reason for exiting the mode */
+	EVREditorExitType GetExitType() const
 	{
-		bWantsToExitMode = true;
+		return ExitType;
 	}
 
 	// FEdMode interface
@@ -65,23 +73,19 @@ public:
 	virtual void AddReferencedObjects( FReferenceCollector& Collector ) override;
 	virtual void Render( const FSceneView* SceneView, FViewport* Viewport, FPrimitiveDrawInterface* PDI ) override;	
 
-	bool HandleInputKey(FEditorViewportClient* ViewportClient, FViewport* Viewport, FKey Key, EInputEvent Event);
-	bool HandleInputAxis(FEditorViewportClient* ViewportClient, FViewport* Viewport, int32 ControllerId, FKey Key, float Delta, float DeltaTime);
-
 	// IVREditorMode interface
-	virtual bool GetLaserPointer( const int32 HandIndex, FVector& LaserPointerStart, FVector& LaserPointerEnd, const bool bEvenIfUIIsInFront = false, const float LaserLengthOverride = 0.0f) const override;
-	virtual bool IsInputCaptured( const FVRAction VRAction ) const;
-	virtual AActor* GetAvatarMeshActor() override
-	{
-		return AvatarMeshActor;
-	}
-	virtual FOnVRAction& OnVRAction() override
-	{
-		return OnVRActionEvent;
-	}
-	virtual FOnVRHoverUpdate& OnVRHoverUpdate() override
+	virtual AActor* GetAvatarMeshActor() override;
+	virtual UWorld* GetWorld() const override;
+	virtual FTransform GetRoomTransform() const override;
+	virtual void SetRoomTransform( const FTransform& NewRoomTransform ) override;
+	virtual FTransform GetRoomSpaceHeadTransform() const override;
+	virtual FTransform GetHeadTransform() const override;
+	virtual const class UViewportWorldInteraction& GetWorldInteraction() const override;
+	virtual class UViewportWorldInteraction& GetWorldInteraction() override;
+	virtual bool IsFullyInitialized() const override;
+	virtual FOnVRTickHandle& OnTickHandle() override
 	{
-		return OnVRHoverUpdateEvent;
+		return TickHandle;
 	}
 
 	/** Returns the Unreal controller ID for the motion controllers we're using */
@@ -90,6 +94,12 @@ public:
 		return MotionControllerID;
 	}
 
+	/** Returns whether or not the flashlight is visible */
+	bool IsFlashlightOn() const
+	{
+		return bIsFlashlightOn;
+	}
+
 	/** Returns the time since the VR Editor mode was last entered */
 	inline FTimespan GetTimeSinceModeEntered() const
 	{
@@ -105,56 +115,6 @@ public:
 		return 1.0f - ( ( Alpha * ( ( OvershootAmount + 1 ) * Alpha + OvershootAmount ) + 1 ) - 1.0f );
 	}
 
-	/** Gets the maximum length of a laser pointer */
-	float GetLaserPointerMaxLength() const;
-
-	/** Gets a virtual hand */
-	FVirtualHand& GetVirtualHand( const int32 VirtualHandIndex )
-	{
-		check( VirtualHandIndex >= 0 && VirtualHandIndex < VREditorConstants::NumVirtualHands );
-		return VirtualHands[ VirtualHandIndex ];
-	}
-
-	/** Gets a virtual hand (const) */
-	const FVirtualHand& GetVirtualHand( const int32 VirtualHandIndex ) const
-	{
-		check( VirtualHandIndex >= 0 && VirtualHandIndex < VREditorConstants::NumVirtualHands );
-		return VirtualHands[ VirtualHandIndex ];
-	}
-
-	/** Given an index of a hand, gets the index of the opposite hand */
-	int32 GetOtherHandIndex( const int32 HandIndex ) const
-	{
-		check( HandIndex >= 0 && HandIndex < VREditorConstants::NumVirtualHands );
-		const int32 OtherHandIndex = ( HandIndex == VREditorConstants::LeftHandIndex ) ? VREditorConstants::RightHandIndex : VREditorConstants::LeftHandIndex;
-		return OtherHandIndex;
-	}
-
-	/** Gets the other hand, given a hand index */
-	FVirtualHand& GetOtherHand( const int32 HandIndex )
-	{
-		return VirtualHands[ GetOtherHandIndex( HandIndex ) ];
-	}
-
-	/** Gets the other hand, given a hand index (const) */
-	const FVirtualHand& GetOtherHand( const int32 HandIndex ) const
-	{
-		return VirtualHands[ GetOtherHandIndex( HandIndex ) ];
-	}
-
-	/** Gets the world space transform of the calibrated VR room origin.  When using a seated VR device, this will feel like the
-	    camera's world transform (before any HMD positional or rotation adjustments are applied.) */
-	FTransform GetRoomTransform() const;
-
-	/** Sets a new transform for the room, in world space.  This is basically setting the editor's camera transform for the viewport */
-	void SetRoomTransform( const FTransform& NewRoomTransform );
-
-	/** Gets the transform of the user's HMD in room space */
-	FTransform GetRoomSpaceHeadTransform() const;
-
-	/** Gets the transform of the user's HMD in world space */
-	FTransform GetHeadTransform() const;
-
 	/** Spawns a transient actor that we can use in the editor world (templated for convenience) */
 	template<class T>
 	T* SpawnTransientSceneActor( const FString& ActorName, const bool bWithSceneComponent ) const
@@ -168,42 +128,14 @@ public:
 	/** Destroys a transient actor we created earlier and nulls out the pointer */
 	void DestroyTransientActor( AActor* Actor ) const;
 
-	/**
-	 * Traces along the laser pointer vector and returns what it first hits in the world
-	 *
-	 * @param HandIndex	Index of the hand that uses has the laser pointer to trace
-	 * @param OptionalListOfIgnoredActors Actors to exclude from hit testing
-	 * @param bIgnoreGizmos True if no gizmo results should be returned, otherwise they are preferred (x-ray)
-	 * @param bEvenIfUIIsInFront If true, ignores any UI that might be blocking the ray
-	 * @param LaserLengthOverride If zero the default laser length (VREdMode::GetLaserLength) is used
-	 *
-	 * @return What the laster pointer hit
-	 */
-	FHitResult GetHitResultFromLaserPointer( int32 HandIndex, TArray<AActor*>* OptionalListOfIgnoredActors = nullptr, const bool bIgnoreGizmos = false, const bool bEvenIfUIIsInFront = false, const float LaserLengthOverride = 0.0f );
-
-	/** Triggers a force feedback effect on either hand */
-	void PlayHapticEffect( const float LeftStrength, const float RightStrength );
-
-	/** Gets access to the world interaction system (const) */
-	const class FVREditorWorldInteraction& GetWorldInteraction() const
-	{
-		return *WorldInteraction;
-	}
-
-	/** Gets access to the world interaction system */
-	class FVREditorWorldInteraction& GetWorldInteraction()
-	{
-		return *WorldInteraction;
-	}
-
 	/** Gets access to the VR UI system (const) */
-	const class FVREditorUISystem& GetUISystem() const
+	const class UVREditorUISystem& GetUISystem() const
 	{
 		return *UISystem;
 	}
 
 	/** Gets access to the VR UI system */
-	class FVREditorUISystem& GetUISystem()
+	class UVREditorUISystem& GetUISystem()
 	{
 		return *UISystem;
 	}
@@ -217,10 +149,7 @@ public:
 	class SLevelViewport& GetLevelViewportPossessedForVR();
 
 	/** Gets the world scale factor, which can be multiplied by a scale vector to convert to room space */
-	float GetWorldScaleFactor() const
-	{
-		return GetWorld()->GetWorldSettings()->WorldToMeters / 100.0f;
-	}
+	float GetWorldScaleFactor() const;
 
 	/** Sets up our avatar mesh, if not already spawned */
 	void SpawnAvatarMeshActor();
@@ -228,46 +157,11 @@ public:
 	/** Gets our VR post process component, creating it if it doesn't exist yet */
 	UPostProcessComponent* GetPostProcessComponent();
 
-	/** Gets the snap grid mesh actor */
-	AActor* GetSnapGridActor()
-	{
-		SpawnAvatarMeshActor();
-		return SnapGridActor;
-	}
-
-	/** Gets the snap grid mesh MID */
-	class UMaterialInstanceDynamic* GetSnapGridMID()
-	{
-		SpawnAvatarMeshActor();
-		return SnapGridMID;
-	}
-
 	/** Called internally when the user changes maps, enters/exits PIE or SIE, or switched between PIE/SIE */
 	void CleanUpActorsBeforeMapChangeOrSimulate();
 
-	/** Invokes the editor's undo system to undo the last change */
-	void Undo();
-	
-	/** Redoes the last change */
-	void Redo();
-
-	/** Copies selected objects to the clipboard */
-	void Copy();
-
-	/** Pastes the clipboard contents into the scene */
-	void Paste();
-
-	/** Duplicates the selected objects */
-	void Duplicate();
-
-	/** Snaps the selected objects to the ground */
-	void SnapSelectedActorsToGround();
-
-	/** Sets the visuals of the LaserPointer */
-	void SetLaserVisuals( const int32 HandIndex, const FLinearColor& NewColor, const float CrawlFade = 0.0f, const float CrawlSpeed = 0.0f );
-
-	/** Given a world space velocity vector, applies inertial damping to it to slow it down */
-	void ApplyVelocityDamping( FVector& Velocity, const bool bVelocitySensitive );
+	/** Spawns a flashlight on the specified hand */
+	void ToggleFlashlight( class UVREditorInteractor* Interactor );
 
 	/** Will update the TransformGizmo Actor with the next Gizmo type  */
 	void CycleTransformGizmoHandleType();
@@ -275,85 +169,34 @@ public:
 	/** Gets the current Gizmo handle type */
 	EGizmoHandleTypes GetCurrentGizmoType() const;
 
-	/** Switch which trasform gizmo coordinate space we're using. */
-	void CycleTransformGizmoCoordinateSpace();
-
-	/** Returns which transform gizmo coordinate space we're using, world or local */
-	ECoordSystem GetTransformGizmoCoordinateSpace() const;
-		
 	/** @return Returns the type of HMD we're dealing with */
 	EHMDDeviceType::Type GetHMDDeviceType() const;
 
 	/** @return Checks to see if the specified hand is aiming roughly toward the specified capsule */
-	bool IsHandAimingTowardsCapsule( const int32 HandIndex, const FTransform& CapsuleTransform, const FVector CapsuleStart, const FVector CapsuleEnd, const float CapsuleRadius, const float MinDistanceToCapsule, const FVector CapsuleFrontDirection, const float MinDotForAimingAtCapsule ) const;
+	bool IsHandAimingTowardsCapsule( class UViewportInteractor* Interactor, const FTransform& CapsuleTransform, const FVector CapsuleStart, const FVector CapsuleEnd, const float CapsuleRadius, const float MinDistanceToCapsule, const FVector CapsuleFrontDirection, const float MinDotForAimingAtCapsule ) const;
 	
-	/**
-	 * Creates a hand transform and forward vector for a laser pointer for a given hand
-	 *
-	 * @param HandIndex			Index of the hand to use
-	 * @param OutHandTransform	The created hand transform
-	 * @param OutForwardVector	The forward vector of the hand
-	 *
-	 * @return	True if we have motion controller data for this hand and could return a valid result
-	 */
-	bool GetHandTransformAndForwardVector( int32 HandIndex, FTransform& OutHandTransform, FVector& OutForwardVector ) const;
-
-	/** Returns the maximum user scale */
-	float GetMaxScale();
-
-	/** Returns the minimum user scale */
-	float GetMinScale();
-
-	/** Sets GNewWorldToMetersScale */
-	void SetWorldToMetersScale( const float NewWorldToMetersScale );
+	/** Gets the hand interactor  */
+	class UVREditorInteractor* GetHandInteractor( const EControllerHand ControllerHand ) const;
+
+	/** Stops the haptic feedback for the left and right hand interactors for the oculus */
+	void StopOldHapticEffects();
 
 private:
 
 	/** Called when someone closes a standalone VR Editor window */
 	void OnVREditorWindowClosed( const TSharedRef<SWindow>& ClosedWindow );
 
-	/**
-	 * Translates a raw key into a vr editing action
-	 *
-	 * @param InKey FKey to translate
-	 * @return The action mapped to InKey
-	 */
-	int32 GetHandIndexFromKey( const FKey& InKey ) const;
-
-	/** Conditionally polls input from controllers, if we haven't done that already this frame */
-	void PollInputIfNeeded();
-
-	/** Grabs latest data from the controllers and updates our virtual hands.  You should call PollInputIfNeeded() instead. */
-	void PollInputFromMotionControllers();
-	
-	/** Stops playing any haptic effects that have been going for a while.  Called every frame. */
-	void StopOldHapticEffects();
-
-	/** Pops up some help text labels for the controller in the specified hand, or hides it, if requested */
-	void ShowHelpForHand( const int32 HandIndex, const bool bShowIt );
-
-	/** Called every frame to update the position of any floating help labels */
-	void UpdateHelpLabels();
-
-	/** Given a mesh and a key name, tries to find a socket on the mesh that matches a supported key */
-	class UStaticMeshSocket* FindMeshSocketForKey( UStaticMesh* StaticMesh, const FKey Key );
-
 	/** FEditorDelegates callbacks */
 	void OnMapChange( uint32 MapChangeFlags );
 	void OnBeginPIE( const bool bIsSimulatingInEditor );
 	void OnEndPIE( const bool bIsSimulatingInEditor );
 	void OnSwitchBetweenPIEAndSIE( const bool bIsSimulatingInEditor );
 
-	/** Changes the color of the buttons on the handmesh */
-	void ApplyButtonPressColors( FVRAction VRAction, EInputEvent Event );
-
-
 protected:
 
 	/** Static ID name for our editor mode */
 	static FEditorModeID VREditorModeID;
 
-
 	//
 	// Startup/Shutdown
 	//
@@ -421,6 +264,12 @@ protected:
 	/** Our avatar's head mesh */
 	class UStaticMeshComponent* HeadMeshComponent;
 
+	/** Spotlight for the flashlight */
+	class USpotLightComponent* FlashlightComponent;
+
+	/** If there is currently a flashlight in the scene */
+	bool bIsFlashlightOn;
+
 
 	//
 	// World movement grid & FX
@@ -441,21 +290,6 @@ protected:
 	/** Post process material for "greying out" the world while in world movement mode */
 	class UMaterialInstanceDynamic* WorldMovementPostProcessMaterial;
 
-
-	//
-	// Snap grid
-	//
-
-	/** Actor for the snap grid */
-	class AActor* SnapGridActor;
-
-	/** The plane mesh we use to draw a snap grid under selected objects */
-	class UStaticMeshComponent* SnapGridMeshComponent;
-
-	/** MID for the snap grid, so we can update snap values on the fly */
-	class UMaterialInstanceDynamic* SnapGridMID;
-
-
 	//
 	// World scaling progress bar
 	//
@@ -484,43 +318,45 @@ protected:
 	/** The Unreal controller ID for the motion controllers we're using */
 	int32 MotionControllerID;
 
-	/** The last frame that input was polled.  We keep track of this so that we can make sure to poll for latest input right before
-	    its first needed each frame */
-	uint32 LastFrameNumberInputWasPolled;
-
-	/** Mapping of raw keys to actions*/
-	TMap<FKey,FVRAction> KeyToActionMap;
-
-	/** List of virtual hands being controlled */
-	FVirtualHand VirtualHands[VREditorConstants::NumVirtualHands];
 
-	/** Event that's triggered when the user presses a button on their motion controller device.  Set bWasHandled to true if you reacted to the event. */
-	FOnVRAction OnVRActionEvent;
-
-	/** Event that fires every frame to update hover state based on what's under the cursor.  Set bWasHandled to true if you detect something to hover. */
-	FOnVRHoverUpdate OnVRHoverUpdateEvent;
+	//
+	// Subsystems registered
+	//
 
+	FOnVRTickHandle TickHandle;
 
 	//
-	// UI
+	// Subsystems
 	//
 
 	/** VR UI system */
-	TUniquePtr<class FVREditorUISystem> UISystem;
+	class UVREditorUISystem* UISystem;
 
+	/** Teleporter system */
+	class UVREditorTeleporter* TeleporterSystem;
 
 	//
 	// World interaction
 	//
 
 	/** World interaction manager */
-	TUniquePtr<class FVREditorWorldInteraction> WorldInteraction;
+	class UVREditorWorldInteraction* WorldInteraction;
 
 	/** The current Gizmo type that is used for the TransformGizmo Actor */
 	EGizmoHandleTypes CurrentGizmoType;
 
-	// Slate Input Processor
-	TSharedPtr<FVREditorInputProcessor> InputProcessor;
+	//
+	// Interactors
+	//
+
+	/** The mouse cursor interactor (currently only available when not in VR) */
+	class UMouseCursorInteractor* MouseCursorInteractor;
+
+	/** The right motion controller */
+	class UVREditorMotionControllerInteractor* LeftHandInteractor; //@todo vreditor: Hardcoded interactors
+	
+	/** The right motion controller */
+	class UVREditorMotionControllerInteractor* RightHandInteractor; 
 
 	//
 	// Colors
@@ -557,5 +393,12 @@ private:
 
 	/** If this is the first tick or before */
 	bool bFirstTick;
+
+	/** If the coordinate system was in world space before entering (local) scale mode */
+	bool bWasInWorldSpaceBeforeScaleMode;
+
+	/** The reason for exiting the mode, so the module can close the mode and take further actions on what should happen next */
+	EVREditorExitType ExitType;
+
 };
 
diff --git a/Engine/Source/Editor/VREditor/VREditorQuickMenu.cpp b/Engine/Source/Editor/VREditor/VREditorQuickMenu.cpp
--- a/Engine/Source/Editor/VREditor/VREditorQuickMenu.cpp
+++ b/Engine/Source/Editor/VREditor/VREditorQuickMenu.cpp
@@ -32,31 +32,33 @@ void UVREditorQuickMenu::OnDeleteButtonClicked()
 
 void UVREditorQuickMenu::OnUndoButtonClicked()
 {
-	GetOwner()->GetOwner().GetOwner().Undo();
+	GetOwner()->GetOwner().GetOwner().GetWorldInteraction().Undo();
 }
 
 
 void UVREditorQuickMenu::OnRedoButtonClicked()
 {
-	GetOwner()->GetOwner().GetOwner().Redo();
+	GetOwner()->GetOwner().GetOwner().GetWorldInteraction().Redo();
 }
 
 
 void UVREditorQuickMenu::OnCopyButtonClicked()
 {
-	GetOwner()->GetOwner().GetOwner().Copy();
+	GetOwner()->GetOwner().GetOwner().GetWorldInteraction().Copy();
 }
 
 
 void UVREditorQuickMenu::OnPasteButtonClicked()
 {
-	GetOwner()->GetOwner().GetOwner().Paste();
+	GetOwner()->GetOwner().GetOwner().GetWorldInteraction().Paste();
 }
 
 
-void UVREditorQuickMenu::OnTranslationSnapButtonClicked( const bool bIsChecked )
+bool UVREditorQuickMenu::OnTranslationSnapButtonClicked( const bool bIsChecked )
 {
 	GUnrealEd->Exec( GetWorld(), *FString::Printf( TEXT( "MODE GRID=%d" ), bIsChecked ? 1 : 0 ) );
+
+	return IsTranslationSnapEnabled();
 }
 
 
@@ -66,7 +68,7 @@ bool UVREditorQuickMenu::IsTranslationSnapEnabled() const
 }
 
 
-void UVREditorQuickMenu::OnTranslationSnapSizeButtonClicked()
+FText UVREditorQuickMenu::OnTranslationSnapSizeButtonClicked()
 {
 	const TArray<float>& GridSizes = GEditor->GetCurrentPositionGridArray();
 	const int32 CurrentGridSize = GetDefault<ULevelEditorViewportSettings>()->CurrentPosGridSize;
@@ -78,6 +80,8 @@ void UVREditorQuickMenu::OnTranslationSnapSizeButtonClicked()
 	}
 
 	GEditor->SetGridSize( NewGridSize );
+
+	return GetTranslationSnapSizeText();
 }
 
 
@@ -87,7 +91,7 @@ FText UVREditorQuickMenu::GetTranslationSnapSizeText() const
 }
 
 
-void UVREditorQuickMenu::OnRotationSnapSizeButtonClicked()
+FText UVREditorQuickMenu::OnRotationSnapSizeButtonClicked()
 {
 	const TArray<float>& GridSizes = GEditor->GetCurrentRotationGridArray();
 	const int32 CurrentGridSize = GetDefault<ULevelEditorViewportSettings>()->CurrentRotGridSize;
@@ -100,6 +104,8 @@ void UVREditorQuickMenu::OnRotationSnapSizeButtonClicked()
 	}
 
 	GEditor->SetRotGridSize( NewGridSize, CurrentGridMode );
+
+	return GetRotationSnapSizeText();
 }
 
 
@@ -109,7 +115,7 @@ FText UVREditorQuickMenu::GetRotationSnapSizeText() const
 }
 
 
-void UVREditorQuickMenu::OnScaleSnapSizeButtonClicked()
+FText UVREditorQuickMenu::OnScaleSnapSizeButtonClicked()
 {
 	const ULevelEditorViewportSettings* ViewportSettings = GetMutableDefault<ULevelEditorViewportSettings>();
 	const int32 NumGridSizes = ViewportSettings->ScalingGridSizes.Num();
@@ -123,6 +129,8 @@ void UVREditorQuickMenu::OnScaleSnapSizeButtonClicked()
 	}
 
 	GEditor->SetScaleGridSize( NewGridSize );
+
+	return GetScaleSnapSizeText();
 }
 
 
@@ -132,9 +140,11 @@ FText UVREditorQuickMenu::GetScaleSnapSizeText() const
 }
 
 
-void UVREditorQuickMenu::OnRotationSnapButtonClicked( const bool bIsChecked )
+bool UVREditorQuickMenu::OnRotationSnapButtonClicked( const bool bIsChecked )
 {
 	GUnrealEd->Exec( GetWorld(), *FString::Printf( TEXT( "MODE ROTGRID=%d" ), bIsChecked ? 1 : 0 ) );
+
+	return IsRotationSnapEnabled();
 }
 
 
@@ -144,9 +154,11 @@ bool UVREditorQuickMenu::IsRotationSnapEnabled() const
 }
 
 
-void UVREditorQuickMenu::OnScaleSnapButtonClicked( const bool bIsChecked )
+bool UVREditorQuickMenu::OnScaleSnapButtonClicked( const bool bIsChecked )
 {
 	GUnrealEd->Exec( GetWorld(), *FString::Printf( TEXT( "MODE SCALEGRID=%d" ), bIsChecked ? 1 : 0 ) );
+
+	return IsScaleSnapEnabled();
 }
 
 
@@ -156,29 +168,31 @@ bool UVREditorQuickMenu::IsScaleSnapEnabled() const
 }
 
 
-void UVREditorQuickMenu::OnGizmoCoordinateSystemButtonClicked()
+FText UVREditorQuickMenu::OnGizmoCoordinateSystemButtonClicked()
 {
-	GetOwner()->GetOwner().GetOwner().CycleTransformGizmoCoordinateSpace();
+	GetOwner()->GetOwner().GetOwner().GetWorldInteraction().CycleTransformGizmoCoordinateSpace(); 
+	return GetGizmoCoordinateSystemText();
 }
 
 
 FText UVREditorQuickMenu::GetGizmoCoordinateSystemText() const
 {
-	const ECoordSystem CurrentCoordSystem = GetOwner()->GetOwner().GetOwner().GetTransformGizmoCoordinateSpace();
+	const ECoordSystem CurrentCoordSystem = GetOwner()->GetOwner().GetOwner().GetWorldInteraction().GetTransformGizmoCoordinateSpace(); //@todo VREditor
 	return CurrentCoordSystem == COORD_World ? LOCTEXT( "WorldCoordinateSystem", "World Space" ) : LOCTEXT( "LocalCoordinateSystem", "Local Space" );
 }
 
 
-void UVREditorQuickMenu::OnGizmoModeButtonClicked()
+FText UVREditorQuickMenu::OnGizmoModeButtonClicked()
 {
-	const EGizmoHandleTypes CurrentGizmoMode = GetOwner()->GetOwner().GetOwner().GetCurrentGizmoType();
+	const EGizmoHandleTypes CurrentGizmoMode = GetOwner()->GetOwner().GetOwner().GetWorldInteraction().GetCurrentGizmoType();
 	GetOwner()->GetOwner().GetOwner().CycleTransformGizmoHandleType();
+	return GetGizmoModeText();
 }
 
 
 FText UVREditorQuickMenu::GetGizmoModeText() const
 {
-	const EGizmoHandleTypes CurrentGizmoType = GetOwner()->GetOwner().GetOwner().GetCurrentGizmoType();
+	const EGizmoHandleTypes CurrentGizmoType = GetOwner()->GetOwner().GetOwner().GetWorldInteraction().GetCurrentGizmoType();
 	FText GizmoTypeText;
 	switch( CurrentGizmoType )
 	{
@@ -207,7 +221,7 @@ FText UVREditorQuickMenu::GetGizmoModeText() const
 }
 
 
-void UVREditorQuickMenu::OnSimulateButtonClicked( const bool bIsChecked )
+bool UVREditorQuickMenu::OnSimulateButtonClicked( const bool bIsChecked )
 {
 	if( bIsChecked != GEditor->bIsSimulatingInEditor )
 	{
@@ -234,6 +248,8 @@ void UVREditorQuickMenu::OnSimulateButtonClicked( const bool bIsChecked )
 			GUnrealEd->RequestToggleBetweenPIEandSIE();
 		}
 	}
+
+	return IsSimulatingEnabled();
 }
 
 
@@ -243,55 +259,59 @@ bool UVREditorQuickMenu::IsSimulatingEnabled() const
 }
 
 
-void UVREditorQuickMenu::OnContentBrowserButtonClicked( const bool bIsChecked )
+bool UVREditorQuickMenu::OnContentBrowserButtonClicked( const bool bIsChecked )
 {
-	GetOwner()->GetOwner().TogglePanelVisibility( FVREditorUISystem::EEditorUIPanel::ContentBrowser );
+	GetOwner()->GetOwner().TogglePanelVisibility( UVREditorUISystem::EEditorUIPanel::ContentBrowser );
+	return IsContentBrowserVisible();
 }
 
 
 bool UVREditorQuickMenu::IsContentBrowserVisible() const
 {
-	return ( GetOwner()->GetOwner().IsShowingEditorUIPanel( FVREditorUISystem::EEditorUIPanel::ContentBrowser ) );
+	return ( GetOwner()->GetOwner().IsShowingEditorUIPanel( UVREditorUISystem::EEditorUIPanel::ContentBrowser ) );
 }
 
 
-void UVREditorQuickMenu::OnWorldOutlinerButtonClicked( const bool bIsChecked )
+bool UVREditorQuickMenu::OnWorldOutlinerButtonClicked( const bool bIsChecked )
 {
-	GetOwner()->GetOwner().TogglePanelVisibility( FVREditorUISystem::EEditorUIPanel::WorldOutliner );
+	GetOwner()->GetOwner().TogglePanelVisibility( UVREditorUISystem::EEditorUIPanel::WorldOutliner );
+	return IsWorldOutlinerVisible();
 }
 
 
 bool UVREditorQuickMenu::IsWorldOutlinerVisible() const
 {
-	return ( GetOwner()->GetOwner().IsShowingEditorUIPanel( FVREditorUISystem::EEditorUIPanel::WorldOutliner ) );
+	return ( GetOwner()->GetOwner().IsShowingEditorUIPanel( UVREditorUISystem::EEditorUIPanel::WorldOutliner ) );
 }
 
 
-void UVREditorQuickMenu::OnActorDetailsButtonClicked( const bool bIsChecked )
+bool UVREditorQuickMenu::OnActorDetailsButtonClicked( const bool bIsChecked )
 {
-	GetOwner()->GetOwner().TogglePanelVisibility( FVREditorUISystem::EEditorUIPanel::ActorDetails );
+	GetOwner()->GetOwner().TogglePanelVisibility( UVREditorUISystem::EEditorUIPanel::ActorDetails );
+	return IsActorDetailsVisible();
 }
 
 
 bool UVREditorQuickMenu::IsActorDetailsVisible() const
 {
-	return ( GetOwner()->GetOwner().IsShowingEditorUIPanel( FVREditorUISystem::EEditorUIPanel::ActorDetails ) );
+	return ( GetOwner()->GetOwner().IsShowingEditorUIPanel( UVREditorUISystem::EEditorUIPanel::ActorDetails ) );
 }
 
 
-void UVREditorQuickMenu::OnModesButtonClicked( const bool bIsChecked )
+bool UVREditorQuickMenu::OnModesButtonClicked( const bool bIsChecked )
 {
-	GetOwner()->GetOwner().TogglePanelVisibility( FVREditorUISystem::EEditorUIPanel::Modes );
+	GetOwner()->GetOwner().TogglePanelVisibility( UVREditorUISystem::EEditorUIPanel::Modes );
+	return IsModesVisible();
 }
 
 
 bool UVREditorQuickMenu::IsModesVisible() const
 {
-	return ( GetOwner()->GetOwner().IsShowingEditorUIPanel( FVREditorUISystem::EEditorUIPanel::Modes ) );
+	return ( GetOwner()->GetOwner().IsShowingEditorUIPanel( UVREditorUISystem::EEditorUIPanel::Modes ) );
 }
 
 
-void UVREditorQuickMenu::OnGameModeButtonClicked( const bool bIsChecked )
+bool UVREditorQuickMenu::OnGameModeButtonClicked( const bool bIsChecked )
 {
 	FEditorViewportClient* ViewportClient = GetOwner()->GetOwner().GetOwner().GetLevelViewportPossessedForVR().GetViewportClient().Get();
 	const bool bIsInGameView = ViewportClient->IsInGameView();
@@ -299,6 +319,8 @@ void UVREditorQuickMenu::OnGameModeButtonClicked( const bool bIsChecked )
 	{
 		ViewportClient->SetGameView( bIsChecked );
 	}
+
+	return IsGameModeEnabled();
 }
 
 
@@ -309,26 +331,79 @@ bool UVREditorQuickMenu::IsGameModeEnabled() const
 }
 
 
-void UVREditorQuickMenu::OnTutorialButtonClicked( const bool bIsChecked )
+bool UVREditorQuickMenu::OnTutorialButtonClicked( const bool bIsChecked )
 {
-	GetOwner()->GetOwner().TogglePanelVisibility( FVREditorUISystem::EEditorUIPanel::Tutorial );
+	GetOwner()->GetOwner().TogglePanelVisibility( UVREditorUISystem::EEditorUIPanel::Tutorial );
+	return IsTutorialVisible();
 }
 
 bool UVREditorQuickMenu::IsTutorialVisible() const
 {
-	return ( GetOwner()->GetOwner().IsShowingEditorUIPanel( FVREditorUISystem::EEditorUIPanel::Tutorial ) );
+	return ( GetOwner()->GetOwner().IsShowingEditorUIPanel( UVREditorUISystem::EEditorUIPanel::Tutorial ) );
 }
 
+bool UVREditorQuickMenu::OnLightButtonClicked( const bool bIsChecked ) 
+{
+	EControllerHand ControllerHand = GetOwner()->GetDockedTo() == AVREditorFloatingUI::EDockedTo::LeftArm ? EControllerHand::Right : EControllerHand::Left;
+	UVREditorInteractor* Interactor = GetOwner()->GetOwner().GetOwner().GetHandInteractor( ControllerHand );
+	GetOwner()->GetOwner().GetOwner().ToggleFlashlight( Interactor );
+	return IsLightVisible();
+}
 
-void UVREditorQuickMenu::OnAssetEditorButtonClicked( const bool bIsChecked )
+bool UVREditorQuickMenu::IsLightVisible() const
 {
-	GetOwner()->GetOwner().TogglePanelVisibility( FVREditorUISystem::EEditorUIPanel::AssetEditor );
-	
+	return (GetOwner()->GetOwner().GetOwner().IsFlashlightOn());
+}
+
+void UVREditorQuickMenu::OnScreenshotButtonClicked()
+{
+	// @todo vreditor: update after direct buffer grab changes
+	if ((GetOwner()->GetOwner().GetOwner().GetHMDDeviceType() == EHMDDeviceType::DT_OculusRift))
+	{	
+		TSharedRef<SWidget> WindowRef = GetOwner()->GetOwner().GetOwner().GetLevelViewportPossessedForVR().AsWidget();
+		TArray<FColor> OutImageData;
+		FIntVector OutImageSize;
+		FSlateApplication::Get().TakeScreenshot(WindowRef, OutImageData, OutImageSize);
+		float ScreenPercentAdjust = (2048.0f / OutImageSize.X) * 100.0f;
+		GEngine->DeferredCommands.Add(FString::Printf(TEXT("HMD SP %3.0f"), ScreenPercentAdjust));
+	}
+	else if(GetOwner()->GetOwner().GetOwner().GetHMDDeviceType() == EHMDDeviceType::DT_SteamVR)
+	{
+		GEngine->DeferredCommands.Add(FString::Printf(TEXT("r.screenpercentage 100")));
+	}
+
+	// Take a screenshot from the HMD, left eye + right eye undistorted
+	// @todo vreditor: verify an option to make hands visible
+	const bool bShowUI = false;
+	const bool bUseCustomFilename = true;
+	FScreenshotRequest::RequestScreenshot(FString(), bShowUI, bUseCustomFilename);
+}
+
+void UVREditorQuickMenu::OnPlayButtonClicked()
+{
+	GetOwner()->GetOwner().GetOwner().StartExitingVRMode( EVREditorExitType::PIE_VR );
+}
+
+bool UVREditorQuickMenu::OnAssetEditorButtonClicked( const bool bIsChecked )
+{
+	GetOwner()->GetOwner().TogglePanelVisibility( UVREditorUISystem::EEditorUIPanel::AssetEditor );
+	return IsAssetEditorVisible();
 }
 
 bool UVREditorQuickMenu::IsAssetEditorVisible() const
 {
-	return ( GetOwner()->GetOwner().IsShowingEditorUIPanel( FVREditorUISystem::EEditorUIPanel::AssetEditor ) );
+	return ( GetOwner()->GetOwner().IsShowingEditorUIPanel( UVREditorUISystem::EEditorUIPanel::AssetEditor ) );
+}
+
+bool UVREditorQuickMenu::OnWorldSettingsButtonClicked( const bool bIsChecked )
+{
+	GetOwner()->GetOwner().TogglePanelVisibility( UVREditorUISystem::EEditorUIPanel::WorldSettings);
+	return IsWorldSettingsVisible();
+}
+
+bool UVREditorQuickMenu::IsWorldSettingsVisible() const
+{
+	return ( GetOwner()->GetOwner().IsShowingEditorUIPanel( UVREditorUISystem::EEditorUIPanel::WorldSettings ) );
 }
 
 #undef LOCTEXT_NAMESPACE
\ No newline at end of file
diff --git a/Engine/Source/Editor/VREditor/VREditorRadialMenu.cpp b/Engine/Source/Editor/VREditor/VREditorRadialMenu.cpp
--- a/Engine/Source/Editor/VREditor/VREditorRadialMenu.cpp
+++ b/Engine/Source/Editor/VREditor/VREditorRadialMenu.cpp
@@ -8,7 +8,8 @@
 #include "VREditorMode.h"
 #include "VREditorWorldInteraction.h"
 #include "VREditorTransformGizmo.h"
-#include "VirtualHand.h"
+#include "VREditorInteractor.h"
+#include "VREditorMotionControllerInteractor.h"
 
 #include "Editor/LevelEditor/Public/LevelEditor.h"
 
@@ -28,55 +29,59 @@ PreviousItem( nullptr )
 }
 
 
-void UVREditorRadialMenu::Update( const FVirtualHand& Hand )
+void UVREditorRadialMenu::Update( const UVREditorInteractor* Interactor )
 {
-	TrackpadPosition = Hand.TrackpadPosition;
-
-	if (!GetOwner()->bHidden)
+	const UVREditorMotionControllerInteractor* MotionControllerInteractor = Cast< UVREditorMotionControllerInteractor >( Interactor );
+	if( MotionControllerInteractor )
 	{
-		// Check if position is from the center
-		if (!IsInMenuRadius())
-		{
-			const float AnglePerItem = 360 / MenuItems.Num();
-			float Angle = FRotator::NormalizeAxis( FMath::RadiansToDegrees( FMath::Atan2( TrackpadPosition.X, TrackpadPosition.Y ) ) );
-			TrackpadAngle = Angle;
+		TrackpadPosition = MotionControllerInteractor->GetTrackpadPosition();
 
-			Angle += AnglePerItem / 2.0f;
-			if (Angle < 0)
+		if ( !GetOwner()->bHidden )
+		{
+			// Check if position is from the center
+			if ( !IsInMenuRadius() )
 			{
-				Angle = Angle + 360;
-			}
+				const float AnglePerItem = 360 / MenuItems.Num();
+				float Angle = FRotator::NormalizeAxis( FMath::RadiansToDegrees( FMath::Atan2( TrackpadPosition.X, TrackpadPosition.Y ) ) );
+				TrackpadAngle = Angle;
 
-			const int32 Index = (Angle / AnglePerItem);
-			UVREditorRadialMenuItem* NewMenuItem = MenuItems[Index];
-			if (NewMenuItem)
-			{
-				CurrentItem = NewMenuItem;
-			}
+				Angle += AnglePerItem / 2.0f;
+				if ( Angle < 0 )
+				{
+					Angle = Angle + 360;
+				}
 
-			// Update the visuals
-			if (CurrentItem != PreviousItem)
-			{
-				if (CurrentItem)
+				const int32 Index = ( Angle / AnglePerItem );
+				UVREditorRadialMenuItem* NewMenuItem = MenuItems[ Index ];
+				if ( NewMenuItem )
 				{
-					CurrentItem->OnEnterHover();
+					CurrentItem = NewMenuItem;
 				}
 
-				if (PreviousItem)
+				// Update the visuals
+				if ( CurrentItem != PreviousItem )
 				{
-					PreviousItem->OnLeaveHover();
+					if ( CurrentItem )
+					{
+						CurrentItem->OnEnterHover();
+					}
+
+					if ( PreviousItem )
+					{
+						PreviousItem->OnLeaveHover();
+					}
 				}
-			}
 
-			PreviousItem = CurrentItem;
-		}
-		else
-		{
-			if (CurrentItem)
+				PreviousItem = CurrentItem;
+			}
+			else
 			{
-				CurrentItem->OnLeaveHover();
-				CurrentItem = nullptr;
-				PreviousItem = nullptr;
+				if ( CurrentItem )
+				{
+					CurrentItem->OnLeaveHover();
+					CurrentItem = nullptr;
+					PreviousItem = nullptr;
+				}
 			}
 		}
 	}
@@ -131,7 +136,7 @@ void UVREditorRadialMenu::ResetItem()
 	}
 }
 
-void UVREditorRadialMenu::SelectCurrentItem( const FVirtualHand& Hand )
+void UVREditorRadialMenu::SelectCurrentItem()
 {
 	if (IsInMenuRadius())
 	{
@@ -153,7 +158,7 @@ void UVREditorRadialMenu::UpdateGizmoTypeLabel()
 {
 	FString Result;
 
-	switch (GetOwner()->GetOwner().GetOwner().GetCurrentGizmoType())
+	switch (GetOwner()->GetOwner().GetOwner().GetWorldInteraction().GetCurrentGizmoType())
 	{
 	case EGizmoHandleTypes::All:
 		Result = LOCTEXT( "AllGizmoType", "Universal Gizmo" ).ToString();
@@ -182,7 +187,7 @@ void UVREditorRadialMenu::OnGizmoCycle()
 
 void UVREditorRadialMenu::OnDuplicateButtonClicked()
 {
-	GetOwner()->GetOwner().GetOwner().Duplicate();
+	GetOwner()->GetOwner().GetOwner().GetWorldInteraction().Duplicate();
 }
 
 
@@ -194,30 +199,34 @@ void UVREditorRadialMenu::OnDeleteButtonClicked()
 
 void UVREditorRadialMenu::OnUndoButtonClicked()
 {
-	GetOwner()->GetOwner().GetOwner().Undo();
+	GetOwner()->GetOwner().GetOwner().GetWorldInteraction().Undo();
 }
 
 
 void UVREditorRadialMenu::OnRedoButtonClicked()
 {
-	GetOwner()->GetOwner().GetOwner().Redo();
+	GetOwner()->GetOwner().GetOwner().GetWorldInteraction().Redo();
 }
 
 
 void UVREditorRadialMenu::OnCopyButtonClicked()
 {
-	GetOwner()->GetOwner().GetOwner().Copy();
+	GetOwner()->GetOwner().GetOwner().GetWorldInteraction().Copy();
 }
 
 
 void UVREditorRadialMenu::OnPasteButtonClicked()
 {
-	GetOwner()->GetOwner().GetOwner().Paste();
+	GetOwner()->GetOwner().GetOwner().GetWorldInteraction().Paste();
 }
 
 void UVREditorRadialMenu::OnSnapActorsToGroundClicked()
 {
-	GetOwner()->GetOwner().GetOwner().SnapSelectedActorsToGround();
+	UVREditorWorldInteraction* WorldInteraction = Cast<UVREditorWorldInteraction>( &( GetOwner()->GetOwner().GetOwner().GetWorldInteraction() ) );
+	if ( WorldInteraction )
+	{
+		WorldInteraction->SnapSelectedActorsToGround();
+	}
 }
 
 
diff --git a/Engine/Source/Editor/VREditor/VREditorRadialMenu.h b/Engine/Source/Editor/VREditor/VREditorRadialMenu.h
--- a/Engine/Source/Editor/VREditor/VREditorRadialMenu.h
+++ b/Engine/Source/Editor/VREditor/VREditorRadialMenu.h
@@ -21,7 +21,7 @@ public:
 	UVREditorRadialMenu( const FObjectInitializer& ObjectInitializer );
 
 	/** Update the state of the widget with the trackpad location */
-	void Update( const struct FVirtualHand& Hand );
+	void Update( const class UVREditorInteractor* Interactor );
 
 	/** To bind the buttons from UMG */
 	UFUNCTION(BlueprintCallable, Category = "RadialItems")
@@ -32,7 +32,7 @@ public:
 	void ResetItem();
 
 	/** Press the current item */
-	void SelectCurrentItem( const FVirtualHand& Hand );
+	void SelectCurrentItem();
 
 	/** If the users has his trackpad or analog stick inside the radius where buttons will not be selected */
 	bool IsInMenuRadius() const;
@@ -47,6 +47,7 @@ public:
 
 private:
 	/** All the buttons in the radial menu */
+	UPROPERTY()
 	TArray<UVREditorRadialMenuItem*> MenuItems;
 
 	/** The current widget currently selected */
diff --git a/Engine/Source/Editor/VREditor/VREditorTransformGizmo.cpp b/Engine/Source/Editor/VREditor/VREditorTransformGizmo.cpp
--- a/Engine/Source/Editor/VREditor/VREditorTransformGizmo.cpp
+++ b/Engine/Source/Editor/VREditor/VREditorTransformGizmo.cpp
@@ -6,8 +6,8 @@
 #include "Gizmo/VREditorTranslationGizmoHandle.h"
 #include "Gizmo/VREditorRotationGizmoHandle.h"
 #include "Gizmo/VREditorPlaneTranslationGizmoHandle.h"
-#include "Gizmo/VREditorStretchGizmoHandle.h"
-#include "Gizmo/VREditorUniformScaleGizmoHandle.h"
+#include "VIStretchGizmoHandle.h"
+#include "VIUniformScaleGizmoHandle.h"
 #include "VREditorMode.h"
 
 namespace VREd
@@ -16,7 +16,6 @@ namespace VREd
 	static FAutoConsoleVariable GizmoScale( TEXT( "VREd.GizmoScale" ), 0.80f, TEXT( "How big the gizmo handles should be" ) );
 	static FAutoConsoleVariable GizmoDistanceScaleFactor( TEXT( "VREd.GizmoDistanceScaleFactor" ), 0.002f, TEXT( "How much the gizmo handles should increase in size with distance from the camera, to make it easier to select" ) );
 	static FAutoConsoleVariable MinActorSizeForTransformGizmo( TEXT( "VREd.MinActorSizeForTransformGizmo" ), 50.0f, TEXT( "How big an object must be in scaled world units before we'll start to shrink the gizmo" ) );
-	static FAutoConsoleVariable AlwaysShowGizmoMeasurementText( TEXT( "VREd.AlwaysShowGizmoMeasurementText" ), 0, TEXT( "When enabled, gizmo measurements will always be visible.  Otherwise, only when hovering over a scale/stretch gizmo handle" ) );
 }
 
 
@@ -40,15 +39,16 @@ ATransformGizmo::ATransformGizmo()
 	RotationGizmoHandleGroup->SetupAttachment( SceneComponent );
 	AllHandleGroups.Add( RotationGizmoHandleGroup );
 
-	StretchGizmoHandleGroup = CreateDefaultSubobject<UVREditorStretchGizmoHandleGroup>( TEXT( "StretchHandles" ), true );
+	StretchGizmoHandleGroup = CreateDefaultSubobject<UStretchGizmoHandleGroup>( TEXT( "StretchHandles" ), true );
 	StretchGizmoHandleGroup->SetTranslucentGizmoMaterial( TranslucentGizmoMaterial );
 	StretchGizmoHandleGroup->SetGizmoMaterial( GizmoMaterial );
 	StretchGizmoHandleGroup->SetupAttachment( SceneComponent );
 	AllHandleGroups.Add( StretchGizmoHandleGroup );
 
-	UniformScaleGizmoHandleGroup = CreateDefaultSubobject<UVREditorUniformScaleGizmoHandleGroup>( TEXT( "UniformScaleHandles" ), true );
+	UniformScaleGizmoHandleGroup = CreateDefaultSubobject<UUniformScaleGizmoHandleGroup>( TEXT( "UniformScaleHandles" ), true );
 	UniformScaleGizmoHandleGroup->SetTranslucentGizmoMaterial( TranslucentGizmoMaterial );
 	UniformScaleGizmoHandleGroup->SetGizmoMaterial( GizmoMaterial );
+	UniformScaleGizmoHandleGroup->SetUsePivotPointAsLocation( false );
 	UniformScaleGizmoHandleGroup->SetupAttachment( SceneComponent );
 	AllHandleGroups.Add( UniformScaleGizmoHandleGroup );
 
@@ -96,8 +96,8 @@ ATransformGizmo::ATransformGizmo()
 
 			MeasurementText->SetWorldSize( 8.0f );	// @todo vreditor tweak: Size of text in world units
 
-													// Assign our custom text rendering material.  It should be Unlit and have Depth Test disabled
-													// so that it draws on top of overlapping geometry.
+			// Assign our custom text rendering material.  It should be Unlit and have Depth Test disabled
+			// so that it draws on top of overlapping geometry.
 			MeasurementText->SetTextMaterial( GizmoTextMaterial );
 
 			// Color the text based on the axis it represents
@@ -128,8 +128,12 @@ ATransformGizmo::ATransformGizmo()
 }
 
 
-void ATransformGizmo::UpdateGizmo( const EGizmoHandleTypes GizmoType, const ECoordSystem GizmoCoordinateSpace, const FTransform& LocalToWorld, const FBox& LocalBounds, const FVector ViewLocation, bool bAllHandlesVisible, UActorComponent* DraggingHandle, const TArray< UActorComponent* >& HoveringOverHandles, const float GizmoHoverScale, const float GizmoHoverAnimationDuration )
+void ATransformGizmo::UpdateGizmo( const EGizmoHandleTypes GizmoType, const ECoordSystem GizmoCoordinateSpace, const FTransform& LocalToWorld, const FBox& LocalBounds, const FVector ViewLocation, bool bAllHandlesVisible, 
+	UActorComponent* DraggingHandle, const TArray< UActorComponent* >& HoveringOverHandles, const float GizmoHoverScale, const float GizmoHoverAnimationDuration )
 {
+	Super::UpdateGizmo( GizmoType, GizmoCoordinateSpace, LocalToWorld, LocalBounds, ViewLocation, bAllHandlesVisible,
+		DraggingHandle, HoveringOverHandles, GizmoHoverScale, GizmoHoverAnimationDuration );
+
 	const float WorldScaleFactor = GetWorld()->GetWorldSettings()->WorldToMeters / 100.0f;
 
 	// Position the gizmo at the location of the first selected actor
@@ -159,12 +163,13 @@ void ATransformGizmo::UpdateGizmo( const EGizmoHandleTypes GizmoType, const ECoo
 
 	// Update all the handles
 	bool bIsHoveringOrDraggingScaleGizmo = false;
-	for (UVREditorGizmoHandleGroup* HandleGroup : AllHandleGroups)
+	for ( UGizmoHandleGroup* HandleGroup : AllHandleGroups )
 	{
-		if (HandleGroup != nullptr)
+		if ( HandleGroup != nullptr )
 		{
 			bool bIsHoveringOrDraggingThisHandleGroup = false;
-			HandleGroup->UpdateGizmoHandleGroup( LocalToWorld, LocalBounds, ViewLocation, bAllHandlesVisible, DraggingHandle, HoveringOverHandles, AnimationAlpha, GizmoScale, GizmoHoverScale, GizmoHoverAnimationDuration, /* Out */ bIsHoveringOrDraggingThisHandleGroup );
+			HandleGroup->UpdateGizmoHandleGroup( LocalToWorld, LocalBounds, ViewLocation, bAllHandlesVisible, DraggingHandle, 
+				HoveringOverHandles, AnimationAlpha, GizmoScale, GizmoHoverScale, GizmoHoverAnimationDuration, /* Out */ bIsHoveringOrDraggingThisHandleGroup );
 			
 			if( HandleGroup->GetHandleType() == EGizmoHandleTypes::Scale && bIsHoveringOrDraggingThisHandleGroup )
 			{
@@ -244,7 +249,7 @@ void ATransformGizmo::UpdateGizmo( const EGizmoHandleTypes GizmoType, const ECoo
 				Measurement.MeasurementText->SetRelativeLocation( GizmoSpaceEdgeCenter + GizmoSpaceUpVector * UpOffsetAmount );
 
 				const int32 TextFacingAxisIndex = AxisIndex == 0 ? 1 : ( AxisIndex == 1 ? 0 : 0 );
-				FVector TextFacingDirection = UVREditorGizmoHandleGroup::GetAxisVector( TextFacingAxisIndex, ETransformGizmoHandleDirection::Positive );
+				FVector TextFacingDirection = UGizmoHandleGroup::GetAxisVector( TextFacingAxisIndex, ETransformGizmoHandleDirection::Positive );
 
 				// Make sure to face the camera
 				const FVector WorldSpaceEdgeCenterToViewDirection = ( LocalToWorld.TransformPosition( GizmoSpaceEdgeCenter ) - ViewLocation ).GetSafeNormal();
@@ -263,10 +268,8 @@ void ATransformGizmo::UpdateGizmo( const EGizmoHandleTypes GizmoType, const ECoo
 
 	// Update visibility
 	{
-		UpdateHandleVisibility( GizmoType, GizmoCoordinateSpace, bAllHandlesVisible, DraggingHandle );
-
 		const bool bShouldShowMeasurements =
-			VREd::AlwaysShowGizmoMeasurementText->GetInt() != 0 ||
+			GetShowMeasurementText() ||
 			bIsHoveringOrDraggingScaleGizmo;
 
 		for( int32 AxisIndex = 0; AxisIndex < 3; ++AxisIndex )
diff --git a/Engine/Source/Editor/VREditor/VREditorTransformGizmo.h b/Engine/Source/Editor/VREditor/VREditorTransformGizmo.h
--- a/Engine/Source/Editor/VREditor/VREditorTransformGizmo.h
+++ b/Engine/Source/Editor/VREditor/VREditorTransformGizmo.h
@@ -2,23 +2,9 @@
 
 #pragma once
 
-#include "Gizmo/VREditorBaseTransformGizmo.h"
+#include "VIBaseTransformGizmo.h"
 #include "VREditorTransformGizmo.generated.h"
 
-/**
- * Displays measurements along the bounds of selected objects
- */
-USTRUCT()
-struct FTransformGizmoMeasurement
-{
-	GENERATED_BODY()
-
-	/** The text that displays the actual measurement and units */
-	UPROPERTY()
-	class UTextRenderComponent* MeasurementText;
-};
-
-
 /**
  * A transform gizmo that allows you to interact with selected objects by moving, scaling and rotating.
  */
@@ -35,7 +21,8 @@ public:
 	/** Called by the world interaction system after we've been spawned into the world, to allow
 	    us to create components and set everything up nicely for the selected objects that we'll be
 		used to manipulate */
-	void UpdateGizmo( const EGizmoHandleTypes GizmoType, const ECoordSystem GizmoCoordinateSpace, const FTransform& LocalToWorld, const FBox& LocalBounds, const FVector ViewLocation, bool bAllHandlesVisible, class UActorComponent* DraggingHandle, const TArray< UActorComponent* >& HoveringOverHandles, const float GizmoHoverScale, const float GizmoHoverAnimationDuration );
+	virtual void UpdateGizmo( const EGizmoHandleTypes GizmoType, const ECoordSystem GizmoCoordinateSpace, const FTransform& LocalToWorld, const FBox& LocalBounds, const FVector ViewLocation, bool bAllHandlesVisible, 
+		class UActorComponent* DraggingHandle, const TArray< UActorComponent* >& HoveringOverHandles, const float GizmoHoverScale, const float GizmoHoverAnimationDuration ) override;
 
 private:
 
@@ -43,18 +30,23 @@ private:
 	UPROPERTY()
 	FTransformGizmoMeasurement Measurements[ 3 ];
 
+	/** Bounding box translation handles */
 	UPROPERTY()
 	class UVREditorTranslationGizmoHandleGroup* TranslationGizmoHandleGroup;
 
+	/** Bounding box plane translation handles */
 	UPROPERTY()
 	class UVREditorPlaneTranslationGizmoHandleGroup* PlaneTranslationGizmoHandleGroup;
 
+	/** Bounding box rotation handles */
 	UPROPERTY()
 	class UVREditorRotationGizmoHandleGroup* RotationGizmoHandleGroup;
 
+	/** Bounding box stretch handles */
 	UPROPERTY()
-	class UVREditorStretchGizmoHandleGroup* StretchGizmoHandleGroup;
+	class UStretchGizmoHandleGroup* StretchGizmoHandleGroup;
 
+	/** Bounding box uniform scale handle */
 	UPROPERTY()
-	class UVREditorUniformScaleGizmoHandleGroup* UniformScaleGizmoHandleGroup;
+	class UUniformScaleGizmoHandleGroup* UniformScaleGizmoHandleGroup;
 };
\ No newline at end of file
