Copying //UE4/Dev-VREditor to //UE4/Dev-Main (Source: //UE4/Dev-VREditor @ 3051763)

diff --git a/Engine/Source/Editor/VREditor/VREditorUISystem.cpp b/Engine/Source/Editor/VREditor/VREditorUISystem.cpp
--- a/Engine/Source/Editor/VREditor/VREditorUISystem.cpp
+++ b/Engine/Source/Editor/VREditor/VREditorUISystem.cpp
@@ -8,8 +8,11 @@
 #include "VREditorQuickMenu.h"
 #include "VREditorRadialMenu.h"
 #include "VREditorRadialMenuItem.h"
-#include "VirtualHand.h"
 #include "IHeadMountedDisplay.h"
+#include "VREditorWorldInteraction.h"
+#include "ViewportInteractor.h"
+#include "VREditorInteractor.h"
+#include "VREditorMotionControllerInteractor.h"
 
 // UI
 #include "WidgetComponent.h"
@@ -25,10 +28,15 @@
 // Actor Details, Modes
 #include "LevelEditor.h"
 
+// World Settings
+#include "Editor/PropertyEditor/Public/PropertyEditorModule.h"
+#include "Editor/PropertyEditor/Public/IDetailsView.h"
+
 #include "SLevelViewport.h"
 #include "SScaleBox.h"
 #include "SDPIScaler.h"
 #include "SWidget.h"
+#include "SDockTab.h"
 
 namespace VREd
 {
@@ -75,31 +83,28 @@ namespace VREd
 }
 
 
-FVREditorUISystem::FVREditorUISystem( FVREditorMode& InitOwner )
-	: Owner( InitOwner ),
-	  FloatingUIs(),
-	  QuickMenuUI( nullptr ),
-	  QuickRadialMenu( nullptr ),
-	  RadialMenuHideDelayTime( 0.0f ),
-	  QuickMenuWidgetClass( nullptr ),
-	  QuickRadialWidgetClass( nullptr ),
-	  TutorialWidgetClass( nullptr ),
-	  DraggingUIHandIndex( INDEX_NONE ),
-	  DraggingUIOffsetTransform( FTransform::Identity ),
-	  DraggingUI( nullptr ),
-	  bRefocusViewport(false),
-	  StartDragUISound( nullptr ),
-	  StopDragUISound( nullptr ),
-	  HideUISound( nullptr ),
-	  ShowUISound( nullptr ),
-	  bDraggedDockFromHandPassedThreshold( false ),
-	  LastDraggingHoverLocation( FVector::ZeroVector ),
-	  bSetDefaultLayout( true )
+UVREditorUISystem::UVREditorUISystem( const FObjectInitializer& Initializer ) : 
+	Super( Initializer ),
+	VRMode( nullptr ),
+	FloatingUIs(),
+	QuickMenuUI( nullptr ),
+	QuickRadialMenu( nullptr ),
+	RadialMenuHideDelayTime( 0.0f ),
+	QuickMenuWidgetClass( nullptr ),
+	QuickRadialWidgetClass( nullptr ),
+	TutorialWidgetClass( nullptr ),
+	InteractorDraggingUI( nullptr ),
+	DraggingUIOffsetTransform( FTransform::Identity ),
+	DraggingUI( nullptr ),
+	bRefocusViewport(false),
+	StartDragUISound( nullptr ),
+	StopDragUISound( nullptr ),
+	HideUISound( nullptr ),
+	ShowUISound( nullptr ),
+	bDraggedDockFromHandPassedThreshold( false ),
+	LastDraggingHoverLocation( FVector::ZeroVector ),
+	bSetDefaultLayout( true )
 {
-	// Register to find out about VR events
-	Owner.OnVRAction().AddRaw( this, &FVREditorUISystem::OnVRAction );
-	Owner.OnVRHoverUpdate().AddRaw( this, &FVREditorUISystem::OnVRHoverUpdate );
-
 	EditorUIPanels.SetNumZeroed( (int32)EEditorUIPanel::TotalCount );
 	
 	// Set default layout transform in correct order of enum
@@ -107,12 +112,19 @@ FVREditorUISystem::FVREditorUISystem( FVREditorMode& InitOwner )
 	DefaultWindowTransforms.Add( VREd::DefaultWorldOutlinerTransform );
 	DefaultWindowTransforms.Add( VREd::DefaultActorDetailsTransform );
 	DefaultWindowTransforms.Add( VREd::DefaultModesTransform );
+}
+
+void UVREditorUISystem::Init()
+{
+	// Register to find out about VR events
+	GetOwner().GetWorldInteraction().OnViewportInteractionInputAction().AddUObject( this, &UVREditorUISystem::OnVRAction );
+	GetOwner().GetWorldInteraction().OnViewportInteractionHoverUpdate().AddUObject( this, &UVREditorUISystem::OnVRHoverUpdate );
 
 	QuickMenuWidgetClass = LoadClass<UVREditorQuickMenu>( nullptr, TEXT( "/Engine/VREditor/UI/VRQuickMenu.VRQuickMenu_C" ) );
 	check( QuickMenuWidgetClass != nullptr );
 
-	QuickRadialWidgetClass = LoadClass<UVREditorRadialMenu>(nullptr, TEXT("/Engine/VREditor/UI/VRRadialQuickMenu.VRRadialQuickMenu_C"));
-	check(QuickRadialWidgetClass != nullptr);
+	QuickRadialWidgetClass = LoadClass<UVREditorRadialMenu>( nullptr, TEXT( "/Engine/VREditor/UI/VRRadialQuickMenu.VRRadialQuickMenu_C" ) );
+	check( QuickRadialWidgetClass != nullptr );
 
 	TutorialWidgetClass = LoadClass<UVREditorBaseUserWidget>( nullptr, TEXT( "/Engine/VREditor/Tutorial/UI_VR_Tutorial_00.UI_VR_Tutorial_00_C" ) );
 	check( TutorialWidgetClass != nullptr );
@@ -130,16 +142,19 @@ FVREditorUISystem::FVREditorUISystem( FVREditorMode& InitOwner )
 	ShowUISound = LoadObject<USoundCue>( nullptr, TEXT( "/Engine/VREditor/Sounds/VR_open_Cue" ) );
 	check( ShowUISound != nullptr );
 
+	
 	// Create all of our UI panels
 	CreateUIs();
 }
 
 
-
-FVREditorUISystem::~FVREditorUISystem()
+void UVREditorUISystem::Shutdown()
 {
-	Owner.OnVRAction().RemoveAll( this );
-	Owner.OnVRHoverUpdate().RemoveAll( this );
+	if ( VRMode != nullptr && &VRMode->GetWorldInteraction() != nullptr )
+	{
+		GetOwner().GetWorldInteraction().OnViewportInteractionInputAction().RemoveAll( this );
+		GetOwner().GetWorldInteraction().OnViewportInteractionHoverUpdate().RemoveAll( this );
+	}
 
 	CleanUpActorsBeforeMapChangeOrSimulate();
 
@@ -152,96 +167,81 @@ FVREditorUISystem::~FVREditorUISystem()
 	StopDragUISound = nullptr;
 	HideUISound = nullptr;
 	ShowUISound = nullptr;
+	VRMode = nullptr;
 }
 
-
-void FVREditorUISystem::AddReferencedObjects( FReferenceCollector& Collector )
-{
-	for( AVREditorFloatingUI* FloatingUIPtr : FloatingUIs )
-	{
-		Collector.AddReferencedObject( FloatingUIPtr );
-	}
-	Collector.AddReferencedObject( QuickMenuWidgetClass );
-	Collector.AddReferencedObject( QuickRadialWidgetClass );
-	Collector.AddReferencedObject( TutorialWidgetClass );
-	Collector.AddReferencedObject( StartDragUISound );
-	Collector.AddReferencedObject( StopDragUISound );
-	Collector.AddReferencedObject( HideUISound );
-	Collector.AddReferencedObject( ShowUISound );
-}
-
-
-void FVREditorUISystem::OnVRAction( FEditorViewportClient& ViewportClient, const FVRAction VRAction, const EInputEvent Event, bool& bIsInputCaptured, bool& bWasHandled )
+void UVREditorUISystem::OnVRAction( FEditorViewportClient& ViewportClient, UViewportInteractor* Interactor,
+	const FViewportActionKeyInput& Action, bool& bOutIsInputCaptured, bool& bWasHandled )
 {
-	if( !bWasHandled )
+	UVREditorMotionControllerInteractor* VREditorInteractor = Cast<UVREditorMotionControllerInteractor>( Interactor );
+	if ( VREditorInteractor )
 	{
-
-		if( VRAction.ActionType == EVRActionType::ConfirmRadialSelection )
+		if ( !bWasHandled )
 		{
-			FVirtualHand& Hand = Owner.GetVirtualHand( VRAction.HandIndex );
-			if( Event == IE_Pressed )
+			if ( Action.ActionType == VRActionTypes::ConfirmRadialSelection )
 			{
-				UVREditorRadialMenu* RadialMenu = QuickRadialMenu->GetUserWidget<UVREditorRadialMenu>();
-				if( IsShowingRadialMenu( VRAction.HandIndex ) )
+				if ( Action.Event == IE_Pressed )
 				{
-					RadialMenu->Update( Owner.GetVirtualHand( VRAction.HandIndex ) );
+					UVREditorRadialMenu* RadialMenu = QuickRadialMenu->GetUserWidget<UVREditorRadialMenu>();
+					if ( IsShowingRadialMenu( VREditorInteractor ) )
+					{
+						RadialMenu->Update( VREditorInteractor );
+					}
+
+					RadialMenu->SelectCurrentItem();
 				}
-				
-				RadialMenu->SelectCurrentItem( Hand );
-			}
 
-			bWasHandled = true;
-		}
-		else if( VRAction.ActionType == EVRActionType::SelectAndMove_LightlyPressed )
-		{
-			FVirtualHand& Hand = Owner.GetVirtualHand( VRAction.HandIndex );
-			FVector LaserPointerStart, LaserPointerEnd;
-			if( Owner.GetLaserPointer( VRAction.HandIndex, LaserPointerStart, LaserPointerEnd ) )
+				bWasHandled = true;
+			}
+			else if ( Action.ActionType == ViewportWorldActionTypes::SelectAndMove_LightlyPressed )
 			{
-				FHitResult HitResult = Owner.GetHitResultFromLaserPointer( VRAction.HandIndex );
-				if( HitResult.Actor.IsValid()  )
+				FVector LaserPointerStart, LaserPointerEnd;
+				if ( VREditorInteractor->GetLaserPointer( LaserPointerStart, LaserPointerEnd ) )
 				{
-					// Only allow clicks to our own widget components
-					UWidgetComponent* WidgetComponent = Cast<UWidgetComponent>( HitResult.GetComponent() );
-					if( WidgetComponent != nullptr && IsWidgetAnEditorUIWidget( WidgetComponent ) )
+					FHitResult HitResult = VREditorInteractor->GetHitResultFromLaserPointer();
+					if ( HitResult.Actor.IsValid() )
 					{
-						// Always mark the event as handled so that the editor doesn't try to select the widget component
-						bWasHandled = true;
-
-						if( Event != IE_Repeat )
+						// Only allow clicks to our own widget components
+						UWidgetComponent* WidgetComponent = Cast<UWidgetComponent>( HitResult.GetComponent() );
+						if ( WidgetComponent != nullptr && IsWidgetAnEditorUIWidget( WidgetComponent ) )
 						{
-							// If the Modifier button is held down, treat this like a right click instead of a left click
-							const bool bIsRightClicking =
-								( Event == IE_Pressed && Hand.bIsModifierPressed ) ||
-								( Event == IE_Released && Hand.bIsRightClickingOnUI );
+							// Always mark the event as handled so that the editor doesn't try to select the widget component
+							bWasHandled = true;
 
+							if ( Action.Event != IE_Repeat )
 							{
+								// If the Modifier button is held down, treat this like a right click instead of a left click
+								const bool bIsRightClicking =
+									( Action.Event == IE_Pressed && VREditorInteractor->IsModifierPressed() ) ||
+									( Action.Event == IE_Released && VREditorInteractor->IsRightClickingOnUI() );
+
 								FVector2D LastLocalHitLocation = WidgetComponent->GetLastLocalHitLocation();
 
 								FVector2D LocalHitLocation;
 								WidgetComponent->GetLocalHitLocation( HitResult.ImpactPoint, LocalHitLocation );
 
 								// If we weren't already hovering over this widget, then we'll reset the last hit location
-								if( WidgetComponent != Hand.HoveringOverWidgetComponent )
+								if ( WidgetComponent != VREditorInteractor->GetHoveringOverWidgetComponent() )
 								{
 									LastLocalHitLocation = LocalHitLocation;
 
-									if ( UVREditorWidgetComponent* VRWidgetComponent = Cast<UVREditorWidgetComponent>(Hand.HoveringOverWidgetComponent) )
+									if ( UVREditorWidgetComponent* VRWidgetComponent = Cast<UVREditorWidgetComponent>( VREditorInteractor->GetHoveringOverWidgetComponent() ) )
 									{
-										VRWidgetComponent->SetIsHovering(false);
+										VRWidgetComponent->SetIsHovering( false );
 									}
 								}
 
 								FWidgetPath WidgetPathUnderFinger = FWidgetPath( WidgetComponent->GetHitWidgetPath( HitResult.ImpactPoint, /*bIgnoreEnabledStatus*/ false ) );
-								if( WidgetPathUnderFinger.IsValid() )
+								if ( WidgetPathUnderFinger.IsValid() )
 								{
 									TSet<FKey> PressedButtons;
-									if( Event == IE_Pressed )
+									if ( Action.Event == IE_Pressed )
 									{
 										PressedButtons.Add( bIsRightClicking ? EKeys::RightMouseButton : EKeys::LeftMouseButton );
 									}
+
 									FPointerEvent PointerEvent(
-										1 + VRAction.HandIndex,
+										1 + (uint8)VREditorInteractor->GetControllerSide(),
 										LocalHitLocation,
 										LastLocalHitLocation,
 										PressedButtons,
@@ -249,237 +249,172 @@ void FVREditorUISystem::OnVRAction( FEditorViewportClient& ViewportClient, const
 										0.0f,	// Wheel delta
 										FModifierKeysState() );
 
-									Hand.HoveringOverWidgetComponent = WidgetComponent;
+									VREditorInteractor->SetHoveringOverWidgetComponent( WidgetComponent );
 
-									if ( UVREditorWidgetComponent* VRWidgetComponent = Cast<UVREditorWidgetComponent>(WidgetComponent) )
+									if ( UVREditorWidgetComponent* VRWidgetComponent = Cast<UVREditorWidgetComponent>( VREditorInteractor->GetHoveringOverWidgetComponent() ) )
 									{
-										VRWidgetComponent->SetIsHovering(true);
+										VRWidgetComponent->SetIsHovering( true );
 									}
 
 									FReply Reply = FReply::Unhandled();
-									if( Event == IE_Pressed )
-									{
-										Reply = FSlateApplication::Get().RoutePointerDownEvent( WidgetPathUnderFinger, PointerEvent );
-										Hand.bIsClickingOnUI = true;
-										Hand.bIsRightClickingOnUI = bIsRightClicking;
-										bIsInputCaptured = true;
-
-										// Play a haptic effect on press
-										const float Strength = VREd::UIPressHapticFeedbackStrength->GetFloat();
-										Owner.PlayHapticEffect(
-											VRAction.HandIndex == VREditorConstants::LeftHandIndex ? Strength : 0.0f,
-											VRAction.HandIndex == VREditorConstants::RightHandIndex ? Strength : 0.0f );
-									}
-									else if( Event == IE_Released )
+									if ( Action.Event == IE_Pressed )
 									{
-										Reply = FSlateApplication::Get().RoutePointerUpEvent( WidgetPathUnderFinger, PointerEvent );
-
 										const double CurrentTime = FPlatformTime::Seconds();
-										if( CurrentTime - Hand.LastClickReleaseTime <= VREd::DoubleClickTime->GetFloat() )
+										if (CurrentTime - VREditorInteractor->GetLastClickReleaseTime() <= VREd::DoubleClickTime->GetFloat())
 										{
 											// Trigger a double click event!
-											Reply = FSlateApplication::Get().RoutePointerDoubleClickEvent( WidgetPathUnderFinger, PointerEvent );
+											Reply = FSlateApplication::Get().RoutePointerDoubleClickEvent(WidgetPathUnderFinger, PointerEvent);
 										}
+										else
+										{
+											Reply = FSlateApplication::Get().RoutePointerDownEvent(WidgetPathUnderFinger, PointerEvent);
+										}
+										
+										VREditorInteractor->SetIsClickingOnUI( true );
+										VREditorInteractor->SetIsRightClickingOnUI( bIsRightClicking );
+										VREditorInteractor->SetLastClickReleaseTime( CurrentTime );
+										bOutIsInputCaptured = true;
 
-										Hand.LastClickReleaseTime = CurrentTime;
+										// Play a haptic effect on press
+										VREditorInteractor->PlayHapticEffect( VREd::UIPressHapticFeedbackStrength->GetFloat() );
+									}
+									else if ( Action.Event == IE_Released )
+									{
+										Reply = FSlateApplication::Get().RoutePointerUpEvent( WidgetPathUnderFinger, PointerEvent );
 									}
 								}
 							}
 						}
 					}
 				}
-			}
 
-			if( Event == IE_Released )
-			{
-				bool bWasRightClicking = false;
-				if( Hand.bIsClickingOnUI )
+				if ( Action.Event == IE_Released )
 				{
-					if( Hand.bIsRightClickingOnUI )
+					bool bWasRightClicking = false;
+					if ( VREditorInteractor->IsClickingOnUI() )
 					{
-						bWasRightClicking = true;
+						if ( VREditorInteractor->IsRightClickingOnUI() )
+						{
+							bWasRightClicking = true;
+						}
+						VREditorInteractor->SetIsClickingOnUI( false );
+						VREditorInteractor->SetIsRightClickingOnUI( false );
+						bOutIsInputCaptured = false;
 					}
-					Hand.bIsClickingOnUI = false;
-					Hand.bIsRightClickingOnUI = false;
-					bIsInputCaptured = false;
-				}
 
-				if ( !bWasHandled )
-				{
-					TSet<FKey> PressedButtons;
-					FPointerEvent PointerEvent(
-						1 + VRAction.HandIndex,
-						FVector2D::ZeroVector,
-						FVector2D::ZeroVector,
-						PressedButtons,
-						bWasRightClicking ? EKeys::RightMouseButton : EKeys::LeftMouseButton,
-						0.0f,	// Wheel delta
-						FModifierKeysState() );
-
-					FWidgetPath EmptyWidgetPath;
-
-					Hand.bIsClickingOnUI = false;
-					Hand.bIsRightClickingOnUI = false;
-					FReply Reply = FSlateApplication::Get().RoutePointerUpEvent( EmptyWidgetPath, PointerEvent );
-				}
-			}
-		}
-	}
+					if ( !bWasHandled )
+					{
+						TSet<FKey> PressedButtons;
+						FPointerEvent PointerEvent(
+							1 + (uint8)VREditorInteractor->GetControllerSide(),
+							FVector2D::ZeroVector,
+							FVector2D::ZeroVector,
+							PressedButtons,
+							bWasRightClicking ? EKeys::RightMouseButton : EKeys::LeftMouseButton,
+							0.0f,	// Wheel delta
+							FModifierKeysState() );
 
-	// Stop dragging the dock if we are dragging a dock
-	if (VRAction.ActionType == EVRActionType::SelectAndMove && Event == IE_Released)
-	{
-		// Put the Dock back on the hand it came from or leave it where it is in the room
-		if (DraggingUI != nullptr && DraggingUIHandIndex == VRAction.HandIndex)
-		{
-			FVirtualHand& OtherHand = Owner.GetOtherHand( VRAction.HandIndex );
-			const int32 OtherHandIndex = Owner.GetOtherHandIndex( DraggingUIHandIndex );
+						FWidgetPath EmptyWidgetPath;
 
-			bool bOnHand = false;
-			const float Distance = FVector::Dist( LastDraggingHoverLocation, OtherHand.Transform.GetLocation() ) / Owner.GetWorldScaleFactor();
-			if (Distance > VREd::MinDockDragDistance->GetFloat())
-			{
-				DraggingUI->SetDockedTo( AVREditorFloatingUI::EDockedTo::Room );
-			}
-			else
-			{
-				bOnHand = true;
-				const FVector EditorUIRelativeOffset( DraggingUI->GetSize().Y * 0.5f + VREd::EditorUIOffsetFromHand->GetFloat(), 0.0f, 0.0f );
-				const FTransform MovingUIGoalTransform = DraggingUI->MakeUITransformLockedToHand( OtherHandIndex, false, EditorUIRelativeOffset, FRotator( 90.0f, 180.0f, 0.0f ) );
-				const AVREditorFloatingUI::EDockedTo NewDockedTo = OtherHandIndex == VREditorConstants::LeftHandIndex ? AVREditorFloatingUI::EDockedTo::LeftHand : AVREditorFloatingUI::EDockedTo::RightHand;
-				DraggingUI->MoveTo( MovingUIGoalTransform, VREd::DockSnapTime->GetFloat(), NewDockedTo );
+						VREditorInteractor->SetIsClickingOnUI( false );
+						VREditorInteractor->SetIsRightClickingOnUI( false );
+						FReply Reply = FSlateApplication::Get().RoutePointerUpEvent( EmptyWidgetPath, PointerEvent );
+					}
+				}
 			}
-
-			ShowEditorUIPanel( DraggingUI, OtherHandIndex, true, bOnHand );
-			StopDraggingDockUI();
 		}
 	}
 }
 
 
-void FVREditorUISystem::OnVRHoverUpdate( FEditorViewportClient& ViewportClient, const int32 HandIndex, FVector& HoverImpactPoint, bool& bIsHoveringOverUI, bool& bWasHandled )
+void UVREditorUISystem::OnVRHoverUpdate( FEditorViewportClient& ViewportClient, UViewportInteractor* Interactor, FVector& HoverImpactPoint, bool& bWasHandled )
 {
-	FVirtualHand& Hand = Owner.GetVirtualHand( HandIndex );
-
-	if( !bWasHandled && Hand.DraggingMode != EVREditorDraggingMode::DockableWindow )
+	UVREditorMotionControllerInteractor* VREditorInteractor = Cast<UVREditorMotionControllerInteractor>( Interactor );
+	if( VREditorInteractor != nullptr )
 	{
-		FVector LaserPointerStart, LaserPointerEnd;
-		if( Owner.GetLaserPointer( HandIndex, LaserPointerStart, LaserPointerEnd ) )
+		if( !bWasHandled && Interactor->GetDraggingMode() != EViewportInteractionDraggingMode::Interactable )
 		{
-			FHitResult HitResult = Owner.GetHitResultFromLaserPointer( HandIndex );
-			if( HitResult.Actor.IsValid() )
+			FVector LaserPointerStart, LaserPointerEnd;
+			if( Interactor->GetLaserPointer( LaserPointerStart, LaserPointerEnd ) )
 			{
-				// Only allow clicks to our own widget components
-				UWidgetComponent* WidgetComponent = Cast<UWidgetComponent>( HitResult.GetComponent() );
-				if( WidgetComponent != nullptr && IsWidgetAnEditorUIWidget( WidgetComponent ) )
+				FHitResult HitResult = Interactor->GetHitResultFromLaserPointer();
+				if( HitResult.Actor.IsValid() )
 				{
-					FVector2D LastLocalHitLocation = WidgetComponent->GetLastLocalHitLocation();
-
-					FVector2D LocalHitLocation;
-					WidgetComponent->GetLocalHitLocation( HitResult.ImpactPoint, LocalHitLocation );
-
-					// If we weren't already hovering over this widget, then we'll reset the last hit location
-					if( WidgetComponent != Hand.HoveringOverWidgetComponent )
-					{
-						LastLocalHitLocation = LocalHitLocation;
-
-						if ( UVREditorWidgetComponent* VRWidgetComponent = Cast<UVREditorWidgetComponent>(Hand.HoveringOverWidgetComponent) )
-						{
-							VRWidgetComponent->SetIsHovering(false);
-						}
-					}
-
-					// @todo vreditor UI: There is a CursorRadius optional arg that we migth want to make use of wherever we call GetHitWidgetPath()!
-					FWidgetPath WidgetPathUnderFinger = FWidgetPath( WidgetComponent->GetHitWidgetPath( HitResult.ImpactPoint, /*bIgnoreEnabledStatus*/ false ) );
-					if ( WidgetPathUnderFinger.IsValid() )
+					// Only allow clicks to our own widget components
+					UWidgetComponent* WidgetComponent = Cast<UWidgetComponent>( HitResult.GetComponent() );
+					if( WidgetComponent != nullptr && IsWidgetAnEditorUIWidget( WidgetComponent ) )
 					{
-						Hand.bIsHovering = true;
-						Hand.HoverLocation = HitResult.ImpactPoint;
-
-						TSet<FKey> PressedButtons;
-						FPointerEvent PointerEvent(
-							1 + HandIndex,
-							LocalHitLocation,
-							LastLocalHitLocation,
-							LocalHitLocation - LastLocalHitLocation,
-							PressedButtons,
-							FModifierKeysState());
+						FVector2D LastLocalHitLocation = WidgetComponent->GetLastLocalHitLocation();
 
-						FSlateApplication::Get().RoutePointerMoveEvent(WidgetPathUnderFinger, PointerEvent, false);
+						FVector2D LocalHitLocation;
+						WidgetComponent->GetLocalHitLocation( HitResult.ImpactPoint, LocalHitLocation );
 
-						bWasHandled = true;
-						HoverImpactPoint = HitResult.ImpactPoint;
-						Hand.HoveringOverWidgetComponent = WidgetComponent;
-						bIsHoveringOverUI = true;
-
-						if ( UVREditorWidgetComponent* VRWidgetComponent = Cast<UVREditorWidgetComponent>(WidgetComponent) )
-						{
-							VRWidgetComponent->SetIsHovering(true);
-						}
-
-						// Route the mouse scrolling
-						if ( Hand.bIsTrackpadPositionValid[1] )
+						// If we weren't already hovering over this widget, then we'll reset the last hit location
+						if( WidgetComponent != VREditorInteractor->GetHoveringOverWidgetComponent() )
 						{
-							const bool bIsAbsolute = ( Owner.GetHMDDeviceType() == EHMDDeviceType::DT_SteamVR );
+							LastLocalHitLocation = LocalHitLocation;
 
-							float ScrollDelta = 0.0f;
-							if( Hand.bIsTouchingTrackpad || !bIsAbsolute )
+							if ( UVREditorWidgetComponent* VRWidgetComponent = Cast<UVREditorWidgetComponent>( VREditorInteractor->GetHoveringOverWidgetComponent() ) )
 							{
-								if( bIsAbsolute )
-								{
-									const float ScrollSpeed = VREd::UIAbsoluteScrollSpeed->GetFloat();
-									ScrollDelta = ( Hand.TrackpadPosition.Y - Hand.LastTrackpadPosition.Y ) * ScrollSpeed;
-								}
-								else
-								{
-									const float ScrollSpeed = VREd::UIRelativeScrollSpeed->GetFloat();
-									ScrollDelta = Hand.TrackpadPosition.Y * ScrollSpeed;
-								}
+								VRWidgetComponent->SetIsHovering(false);
 							}
+						}
 
-							// If using a trackpad (Vive), invert scroll direction so that it feels more like scrolling on a mobile device
-							if( Owner.GetHMDDeviceType() == EHMDDeviceType::DT_SteamVR )
+						// @todo vreditor UI: There is a CursorRadius optional arg that we migth want to make use of wherever we call GetHitWidgetPath()!
+						FWidgetPath WidgetPathUnderFinger = FWidgetPath( WidgetComponent->GetHitWidgetPath( HitResult.ImpactPoint, /*bIgnoreEnabledStatus*/ false ) );
+						if ( WidgetPathUnderFinger.IsValid() )
+						{
+							TSet<FKey> PressedButtons;
+							FPointerEvent PointerEvent(
+								1 + (uint8)VREditorInteractor->GetControllerSide(),
+								LocalHitLocation,
+								LastLocalHitLocation,
+								LocalHitLocation - LastLocalHitLocation,
+								PressedButtons,
+								FModifierKeysState());
+
+							FSlateApplication::Get().RoutePointerMoveEvent(WidgetPathUnderFinger, PointerEvent, false);
+
+							bWasHandled = true;
+							HoverImpactPoint = HitResult.ImpactPoint;
+							VREditorInteractor->SetHoveringOverWidgetComponent( WidgetComponent );
+							VREditorInteractor->SetIsHoveringOverUI( true );
+
+							if ( UVREditorWidgetComponent* VRWidgetComponent = Cast<UVREditorWidgetComponent>(WidgetComponent) )
 							{
-								ScrollDelta *= -1.0f;
+								VRWidgetComponent->SetIsHovering(true);
 							}
 
-							if( !FMath::IsNearlyZero( ScrollDelta ) )
+							// Route the mouse scrolling
+							if ( VREditorInteractor->IsTrackpadPositionValid( 1 ) )
 							{
-								FPointerEvent MouseWheelEvent(
-									1 + HandIndex,
-									LocalHitLocation,
-									LastLocalHitLocation,
-									PressedButtons,
-									EKeys::MouseWheelAxis,
-									ScrollDelta,
-									FModifierKeysState() );
-
-								FSlateApplication::Get().RouteMouseWheelOrGestureEvent(WidgetPathUnderFinger, MouseWheelEvent, nullptr);
+								const bool bIsAbsolute = ( GetOwner().GetHMDDeviceType() == EHMDDeviceType::DT_SteamVR );
 
-								Hand.UIScrollVelocity = 0.0f;
-
-								// Don't apply inertia unless the user dragged a decent amount this frame
-								if( bIsAbsolute && FMath::Abs( ScrollDelta ) >= VREd::MinUIScrollDeltaForInertia->GetFloat() )
+								float ScrollDelta = 0.0f;
+								if( VREditorInteractor->IsTouchingTrackpad() || !bIsAbsolute )
 								{
-									// Don't apply inertia if our data is sort of old
-									const FTimespan CurrentTime = FTimespan::FromSeconds( FPlatformTime::Seconds() );
-									if( CurrentTime - Hand.LastTrackpadPositionUpdateTime < FTimespan::FromSeconds( 1.0f / 30.0f ) )
+									if( bIsAbsolute )
+									{
+										const float ScrollSpeed = VREd::UIAbsoluteScrollSpeed->GetFloat();
+										ScrollDelta = ( VREditorInteractor->GetTrackpadPosition().Y - VREditorInteractor->GetLastTrackpadPosition().Y ) * ScrollSpeed;
+									}
+									else
 									{
-										//GWarn->Logf( TEXT( "INPUT: UIScrollVelocity=%0.2f" ), ScrollDelta );
-										Hand.UIScrollVelocity = ScrollDelta;
+										const float ScrollSpeed = VREd::UIRelativeScrollSpeed->GetFloat();
+										ScrollDelta = VREditorInteractor->GetTrackpadPosition().Y * ScrollSpeed;
 									}
 								}
-							}
-						}
-						else
-						{
-							if( !FMath::IsNearlyZero( Hand.UIScrollVelocity ) )
-							{
-								// Apply UI scrolling inertia
-								const float ScrollDelta = Hand.UIScrollVelocity;
+
+								// If using a trackpad (Vive), invert scroll direction so that it feels more like scrolling on a mobile device
+								if( GetOwner().GetHMDDeviceType() == EHMDDeviceType::DT_SteamVR )
+								{
+									ScrollDelta *= -1.0f;
+								}
+
+								if( !FMath::IsNearlyZero( ScrollDelta ) )
 								{
 									FPointerEvent MouseWheelEvent(
-										1 + HandIndex,
+										1 + (uint8)VREditorInteractor->GetControllerSide(),
 										LocalHitLocation,
 										LastLocalHitLocation,
 										PressedButtons,
@@ -487,53 +422,85 @@ void FVREditorUISystem::OnVRHoverUpdate( FEditorViewportClient& ViewportClient,
 										ScrollDelta,
 										FModifierKeysState() );
 
-									FSlateApplication::Get().RouteMouseWheelOrGestureEvent( WidgetPathUnderFinger, MouseWheelEvent, nullptr );
-								}
+									FSlateApplication::Get().RouteMouseWheelOrGestureEvent(WidgetPathUnderFinger, MouseWheelEvent, nullptr);
 
-								// Apply damping
-								FVector ScrollVelocityVector( Hand.UIScrollVelocity, 0.0f, 0.0f );
-								const bool bVelocitySensitive = false;
-								Owner.ApplyVelocityDamping( ScrollVelocityVector, bVelocitySensitive );
-								Hand.UIScrollVelocity = ScrollVelocityVector.X;
+									VREditorInteractor->SetUIScrollVelocity( 0.0f );
 
-								//GWarn->Logf( TEXT( "INERTIA: UIScrollVelocity==%0.2f  (DAMPING: UIScrollVelocity==%0.2f)" ), ScrollDelta, Hand.UIScrollVelocity );
+									// Don't apply inertia unless the user dragged a decent amount this frame
+									if( bIsAbsolute && FMath::Abs( ScrollDelta ) >= VREd::MinUIScrollDeltaForInertia->GetFloat() )
+									{
+										// Don't apply inertia if our data is sort of old
+										const FTimespan CurrentTime = FTimespan::FromSeconds( FPlatformTime::Seconds() );
+										if( CurrentTime - VREditorInteractor->GetLastTrackpadPositionUpdateTime() < FTimespan::FromSeconds( 1.0f / 30.0f ) )
+										{
+											//GWarn->Logf( TEXT( "INPUT: UIScrollVelocity=%0.2f" ), ScrollDelta );
+											VREditorInteractor->SetUIScrollVelocity( ScrollDelta );
+										}
+									}
+								}
 							}
 							else
 							{
-								Hand.UIScrollVelocity = 0.0f;
+								if( !FMath::IsNearlyZero( VREditorInteractor->GetUIScrollVelocity() ) )
+								{
+									// Apply UI scrolling inertia
+									const float ScrollDelta = VREditorInteractor->GetUIScrollVelocity();
+									{
+										FPointerEvent MouseWheelEvent(
+											1 + (uint8)VREditorInteractor->GetControllerSide(),
+											LocalHitLocation,
+											LastLocalHitLocation,
+											PressedButtons,
+											EKeys::MouseWheelAxis,
+											ScrollDelta,
+											FModifierKeysState() );
+
+										FSlateApplication::Get().RouteMouseWheelOrGestureEvent( WidgetPathUnderFinger, MouseWheelEvent, nullptr );
+									}
+
+									// Apply damping
+									FVector ScrollVelocityVector( VREditorInteractor->GetUIScrollVelocity(), 0.0f, 0.0f );
+									const bool bVelocitySensitive = false;
+									GetOwner().GetWorldInteraction().ApplyVelocityDamping( ScrollVelocityVector, bVelocitySensitive );
+									VREditorInteractor->SetUIScrollVelocity( ScrollVelocityVector.X );
+								}
+								else
+								{
+									VREditorInteractor->SetUIScrollVelocity( 0.0f );
+								}
 							}
 						}
 					}
 				}
 			}
 		}
-	}
 
-	// If nothing was hovered, make sure we tell Slate about that
-	if( !bWasHandled && Hand.HoveringOverWidgetComponent != nullptr )
-	{
-		if ( UVREditorWidgetComponent* VRWidgetComponent = Cast<UVREditorWidgetComponent>(Hand.HoveringOverWidgetComponent) )
+		// If nothing was hovered, make sure we tell Slate about that
+		if( !bWasHandled && VREditorInteractor->GetHoveringOverWidgetComponent() != nullptr )
 		{
-			VRWidgetComponent->SetIsHovering(false);
-		}
+			if ( UVREditorWidgetComponent* VRWidgetComponent = Cast<UVREditorWidgetComponent>( VREditorInteractor->GetHoveringOverWidgetComponent() ) )
+			{
+				VRWidgetComponent->SetIsHovering( false );
+			}
 
-		const FVector2D LastLocalHitLocation = Hand.HoveringOverWidgetComponent->GetLastLocalHitLocation();
-		Hand.HoveringOverWidgetComponent = nullptr;
-
-		TSet<FKey> PressedButtons;
-		FPointerEvent PointerEvent(
-			1 + HandIndex,
-			LastLocalHitLocation,
-			LastLocalHitLocation,
-			FVector2D::ZeroVector,
-			PressedButtons,
-			FModifierKeysState() );
-		FSlateApplication::Get().RoutePointerMoveEvent( FWidgetPath(), PointerEvent, false );
+			const FVector2D LastLocalHitLocation = VREditorInteractor->GetHoveringOverWidgetComponent()->GetLastLocalHitLocation();
+			VREditorInteractor->SetHoveringOverWidgetComponent( nullptr );
+
+			TSet<FKey> PressedButtons;
+			FPointerEvent PointerEvent(
+				1 + (uint8)VREditorInteractor->GetControllerSide(),
+				LastLocalHitLocation,
+				LastLocalHitLocation,
+				FVector2D::ZeroVector,
+				PressedButtons,
+				FModifierKeysState() );
+
+			FSlateApplication::Get().RoutePointerMoveEvent( FWidgetPath(), PointerEvent, false );
+		}
 	}
 }
 
-
-void FVREditorUISystem::Tick( FEditorViewportClient* ViewportClient, const float DeltaTime )
+void UVREditorUISystem::Tick( FEditorViewportClient* ViewportClient, const float DeltaTime )
 {
 	if ( bRefocusViewport )
 	{
@@ -545,108 +512,73 @@ void FVREditorUISystem::Tick( FEditorViewportClient* ViewportClient, const float
 	// aimed at when the user does this.
 	if( QuickMenuUI != nullptr )
 	{
-		int32 HandIndexWithQuickMenu = INDEX_NONE;
+		UVREditorInteractor* HandInteractorWithQuickMenu = nullptr;
 		if( QuickMenuUI->IsUIVisible() )
 		{
-			HandIndexWithQuickMenu = QuickMenuUI->GetDockedTo() == AVREditorFloatingUI::EDockedTo::LeftArm ? VREditorConstants::LeftHandIndex : VREditorConstants::RightHandIndex;
+			HandInteractorWithQuickMenu = GetOwner().GetHandInteractor( QuickMenuUI->GetDockedTo() == AVREditorFloatingUI::EDockedTo::LeftArm ? EControllerHand::Left : EControllerHand::Right );
 		}
 
-		const float WorldScaleFactor = Owner.GetWorldScaleFactor();
+		const float WorldScaleFactor = GetOwner().GetWorldScaleFactor();
 
-		int32 HandIndexThatNeedsQuickMenu = INDEX_NONE;
-		for( int32 HandIndex = 0; HandIndex < VREditorConstants::NumVirtualHands; ++HandIndex )
+		UViewportInteractor* HandInteractorThatNeedsQuickMenu = nullptr;
+		for( UViewportInteractor* Interactor : GetOwner().GetWorldInteraction().GetInteractors() )
 		{
 			bool bShowQuickMenuOnArm = false;
 
-			const FVirtualHand& Hand = Owner.GetVirtualHand( HandIndex );
-			const FVirtualHand& OtherHand = Owner.GetOtherHand( HandIndex );
-			const int32 OtherHandIndex = Owner.GetOtherHandIndex( HandIndex );		
-
-			// @todo vreditor tweak: Weird to hard code this here.  Probably should be an accessor on the hand itself, and based on the actual device type
-			const FTransform UICapsuleTransform = OtherHand.Transform;
-			const FVector UICapsuleStart = FVector( -9.0f, 0.0f, 0.0f ) * WorldScaleFactor;
-			const FVector UICapsuleEnd = FVector( -18.0f, 0.0f, 0.0f ) * WorldScaleFactor;
-			const float UICapsuleLocalRadius = 6.0f * WorldScaleFactor;
-			const float MinDistanceToUICapsule = 8.0f * WorldScaleFactor;	// @todo vreditor tweak
-			const FVector UIForwardVector = FVector::UpVector;
-			const float MinDotForAimingAtOtherHand = 0.25f;	// @todo vreditor tweak
-
-			if( Owner.IsHandAimingTowardsCapsule( HandIndex, UICapsuleTransform, UICapsuleStart, UICapsuleEnd, UICapsuleLocalRadius, MinDistanceToUICapsule, UIForwardVector, MinDotForAimingAtOtherHand ) )
-			{
-				bShowQuickMenuOnArm = true;
-			}
-
-			if( bShowQuickMenuOnArm )
+			UViewportInteractor* OtherInteractor = Interactor->GetOtherInteractor();
+			if( OtherInteractor != nullptr )
 			{
-				HandIndexThatNeedsQuickMenu = OtherHandIndex;
-			}
-		}
+				// @todo vreditor tweak: Weird to hard code this here. Probably should be an accessor on the hand itself, and based on the actual device type
+				const FTransform UICapsuleTransform = OtherInteractor->GetTransform();
+				const FVector UICapsuleStart = FVector( -9.0f, 0.0f, 0.0f ) * WorldScaleFactor;
+				const FVector UICapsuleEnd = FVector( -18.0f, 0.0f, 0.0f ) * WorldScaleFactor;
+				const float UICapsuleLocalRadius = 6.0f * WorldScaleFactor;
+				const float MinDistanceToUICapsule = 8.0f * WorldScaleFactor; // @todo vreditor tweak
+				const FVector UIForwardVector = FVector::UpVector;
+				const float MinDotForAimingAtOtherHand = 0.25f;	// @todo vreditor tweak
+
+				if( GetOwner().IsHandAimingTowardsCapsule( Interactor, UICapsuleTransform, UICapsuleStart, UICapsuleEnd, UICapsuleLocalRadius, MinDistanceToUICapsule, UIForwardVector, MinDotForAimingAtOtherHand ) )
+				{
+					bShowQuickMenuOnArm = true;
+				}
 
-		if( QuickMenuUI->IsUIVisible() )
-		{
-			// If we don't need a quick menu, or if a different hand needs to spawn it, then kill the existing menu
-			if( HandIndexThatNeedsQuickMenu == INDEX_NONE || HandIndexThatNeedsQuickMenu != HandIndexWithQuickMenu  )
-			{
-				// Despawn
-				Owner.GetVirtualHand( HandIndexWithQuickMenu ).bHasUIOnForearm = false;
-				QuickMenuUI->ShowUI( false );
+				if( bShowQuickMenuOnArm )
+				{
+					HandInteractorThatNeedsQuickMenu = OtherInteractor;
+				}
 			}
 		}
-
-		if( HandIndexThatNeedsQuickMenu != INDEX_NONE && !QuickMenuUI->IsUIVisible() )
+		
+		// If we don't need a quick menu, or if a different hand needs to spawn it, then kill the existing menu
+		if( QuickMenuUI->IsUIVisible() && HandInteractorWithQuickMenu != nullptr && 
+			( HandInteractorThatNeedsQuickMenu == nullptr || HandInteractorThatNeedsQuickMenu != HandInteractorWithQuickMenu )  )
 		{
-			const AVREditorFloatingUI::EDockedTo DockTo = ( HandIndexThatNeedsQuickMenu == VREditorConstants::LeftHandIndex ) ? AVREditorFloatingUI::EDockedTo::LeftArm : AVREditorFloatingUI::EDockedTo::RightArm;
-			QuickMenuUI->SetDockedTo( DockTo );
-			QuickMenuUI->ShowUI( true );
-			Owner.GetVirtualHand( HandIndexThatNeedsQuickMenu ).bHasUIOnForearm = true;
+			// Despawn
+			HandInteractorWithQuickMenu->SetHasUIOnForearm( false );
+			QuickMenuUI->ShowUI( false );
 		}
-	}
 
-	// If the user is moving the analog stick, try to spawn the radial menu for that hand
-	if( Owner.GetHMDDeviceType() == EHMDDeviceType::DT_OculusRift )
-	{
-		for( int32 HandIndex = 0; HandIndex < VREditorConstants::NumVirtualHands; ++HandIndex )
+		if( HandInteractorThatNeedsQuickMenu != nullptr && !QuickMenuUI->IsUIVisible() )
 		{
-			const FVirtualHand& Hand = Owner.GetVirtualHand( HandIndex );
-			const FVirtualHand& OtherHand = Owner.GetOtherHand( HandIndex );
-			bool bShouldShowRadialMenu = false;
-
-			const float MinJoystickOffsetBeforeRadialMenu = 0.15f;	// @todo vreditor twea
-
-			const bool bOtherHandAlreadyHasRadialMenu =
-				IsShowingRadialMenu( Owner.GetOtherHandIndex( HandIndex ) );
-
-			if( !Hand.bHasUIInFront &&
-				Hand.bIsTrackpadPositionValid[ 0 ] &&
-				Hand.bIsTrackpadPositionValid[ 1 ] &&
-				Hand.TrackpadPosition.Size() > MinJoystickOffsetBeforeRadialMenu &&
-				GetDraggingDockUIHandIndex() != HandIndex &&
-				!bOtherHandAlreadyHasRadialMenu )
-			{
-				bShouldShowRadialMenu = true;
-			}
-
-			if( bShouldShowRadialMenu )
+			UVREditorMotionControllerInteractor* MotionControllerHandInteractorThatNeedsQuickMenu = Cast<UVREditorMotionControllerInteractor>( HandInteractorThatNeedsQuickMenu );
+			if ( MotionControllerHandInteractorThatNeedsQuickMenu )
 			{
-				TryToSpawnRadialMenu( HandIndex );
-			}
-			else
-			{
-				// Close it
-				HideRadialMenu( HandIndex );
+				const AVREditorFloatingUI::EDockedTo DockTo = ( MotionControllerHandInteractorThatNeedsQuickMenu->GetControllerSide() == EControllerHand::Left ) ? AVREditorFloatingUI::EDockedTo::LeftArm : AVREditorFloatingUI::EDockedTo::RightArm;
+				QuickMenuUI->SetDockedTo( DockTo );
+				QuickMenuUI->ShowUI( true );
+				MotionControllerHandInteractorThatNeedsQuickMenu->SetHasUIOnForearm( true );
 			}
 		}
 	}
-	else
+
+	// Close the radial menu if it was not updated for a while
+	RadialMenuHideDelayTime += DeltaTime;
+	if( RadialMenuHideDelayTime >= VREd::RadialMenuFadeDelay->GetFloat() )
 	{
-		// Close the radial menu if it was not updated for a while
-		RadialMenuHideDelayTime += DeltaTime;
-		if( RadialMenuHideDelayTime >= VREd::RadialMenuFadeDelay->GetFloat() )
-		{
-			HideRadialMenu( QuickRadialMenu->GetDockedTo() == AVREditorFloatingUI::EDockedTo::LeftHand ? VREditorConstants::LeftHandIndex : VREditorConstants::RightHandIndex );
-		}
+		UVREditorInteractor* Interactor = GetOwner().GetHandInteractor( QuickRadialMenu->GetDockedTo() == AVREditorFloatingUI::EDockedTo::LeftHand ? EControllerHand::Left : EControllerHand::Right );
+		HideRadialMenu( Interactor );
 	}
-
+	
 	// Tick all of our floating UIs
 	for( AVREditorFloatingUI* FloatingUIPtr : FloatingUIs )
 	{
@@ -657,12 +589,12 @@ void FVREditorUISystem::Tick( FEditorViewportClient* ViewportClient, const float
 	}
 }
 
-void FVREditorUISystem::Render( const FSceneView* SceneView, FViewport* Viewport, FPrimitiveDrawInterface* PDI )
+void UVREditorUISystem::Render( const FSceneView* SceneView, FViewport* Viewport, FPrimitiveDrawInterface* PDI )
 {
 	// ...
 }
 
-void FVREditorUISystem::CreateUIs()
+void UVREditorUISystem::CreateUIs()
 {
 	const FIntPoint DefaultResolution( VREd::DefaultEditorUIResolutionX->GetInt(), VREd::DefaultEditorUIResolutionY->GetInt() );
 
@@ -828,7 +760,9 @@ void FVREditorUISystem::CreateUIs()
 
 		// Create the tutorial dockable window
 		{
-			const bool bShowAtStart = !bIsVREditorDemo;
+			const bool bShowAtStart = 
+				!bIsVREditorDemo && 
+				GetOwner().IsActuallyUsingVR();	// Don't show tutorial at startup when in "forced VR" mode, because you can't interact with UI using the mouse currently
 
 			AVREditorFloatingUI* TutorialUI = GetOwner().SpawnTransientSceneActor< AVREditorDockableWindow >( TEXT( "TutorialUI" ), false );
 			TutorialUI->SetUMGWidget( *this, TutorialWidgetClass, FIntPoint( VREd::TutorialUIResolutionX->GetFloat(), VREd::TutorialUIResolutionY->GetFloat() ), VREd::TutorialUISize->GetFloat(), bShowAtStart ? AVREditorFloatingUI::EDockedTo::Room : AVREditorFloatingUI::EDockedTo::Nothing );
@@ -849,33 +783,64 @@ void FVREditorUISystem::CreateUIs()
 			const FIntPoint Resolution(VREd::AssetEditorUIResolutionX->GetInt(), VREd::AssetEditorUIResolutionY->GetInt());
 
 			const bool bWithSceneComponent = false;
-			AVREditorDockableWindow* TabManagerUI = GetOwner().SpawnTransientSceneActor< AVREditorDockableWindow >(TEXT("AssetEditor"), bWithSceneComponent);
-			TabManagerUI->SetSlateWidget(*this, SNullWidget::NullWidget, Resolution, VREd::EditorUISize->GetFloat(), AVREditorFloatingUI::EDockedTo::Nothing);
+			AVREditorDockableWindow* TabManagerUI = GetOwner().SpawnTransientSceneActor< AVREditorDockableWindow >( TEXT( "AssetEditor" ), bWithSceneComponent );
+			TabManagerUI->SetSlateWidget( *this, SNullWidget::NullWidget, Resolution, VREd::EditorUISize->GetFloat(), AVREditorFloatingUI::EDockedTo::Nothing );
 			TabManagerUI->ShowUI( false );
 
 			// @todo vreditor: Could use "Hovering" instead for better performance with many UIs, but needs to be manually refreshed in some cases
-			TabManagerUI->GetWidgetComponent()->SetDrawingPolicy(EVREditorWidgetDrawingPolicy::Always);
+			TabManagerUI->GetWidgetComponent()->SetDrawingPolicy( EVREditorWidgetDrawingPolicy::Always );
 
-			FloatingUIs.Add(TabManagerUI);
+			FloatingUIs.Add( TabManagerUI );
 
 			EditorUIPanels[ (int32)EEditorUIPanel::AssetEditor ] = TabManagerUI;
 
 			TSharedPtr<SWindow> TabManagerWindow = TabManagerUI->GetWidgetComponent()->GetSlateWindow();
 			TSharedRef<SWindow> TabManagerWindowRef = TabManagerWindow.ToSharedRef();
-			ProxyTabManager = MakeShareable(new FProxyTabmanager(TabManagerWindowRef));
+			ProxyTabManager = MakeShareable( new FProxyTabmanager( TabManagerWindowRef ) );
 
-			ProxyTabManager->OnTabOpened.Add(FOnTabEvent::FDelegate::CreateRaw(this, &FVREditorUISystem::OnProxyTabLaunched));
-			ProxyTabManager->OnAttentionDrawnToTab.Add(FOnTabEvent::FDelegate::CreateRaw(this, &FVREditorUISystem::OnAttentionDrawnToTab));
+			ProxyTabManager->OnTabOpened.AddUObject( this, &UVREditorUISystem::OnProxyTabLaunched );
+			ProxyTabManager->OnAttentionDrawnToTab.AddUObject( this, &UVREditorUISystem::OnAttentionDrawnToTab );
 
 			// We're going to start stealing tabs from the global tab manager inserting them into the world instead.
-			FGlobalTabmanager::Get()->SetProxyTabManager(ProxyTabManager);
+			FGlobalTabmanager::Get()->SetProxyTabManager( ProxyTabManager );
 
-			FAssetEditorManager::Get().OnAssetEditorOpened().AddRaw(this, &FVREditorUISystem::OnAssetEditorOpened);
+			FAssetEditorManager::Get().OnAssetEditorOpened().AddUObject( this, &UVREditorUISystem::OnAssetEditorOpened );
+		}
+		// create the world settings menu
+		{
+			FPropertyEditorModule& PropertyEditorModule = FModuleManager::Get().LoadModuleChecked< FPropertyEditorModule >( "PropertyEditor" );
+
+			FDetailsViewArgs DetailsViewArgs( false, false, true, FDetailsViewArgs::HideNameArea, false, GUnrealEd );
+			DetailsViewArgs.bShowActorLabel = false;
+
+			TSharedRef<IDetailsView> WorldSettings = PropertyEditorModule.CreateDetailView( DetailsViewArgs );
+
+			UWorld* World = GetOwner().GetWorld();
+			if ( World != NULL )
+			{
+				WorldSettings->SetObject( World->GetWorldSettings() );
+			}
+
+			TSharedRef<SWidget> WidgetToDraw = 
+				SNew( SDPIScaler )
+				.DPIScale( VREd::EditorUIScale->GetFloat() )
+				[
+					WorldSettings
+				]
+			;
+
+			const bool bWithSceneComponent = false;
+			AVREditorFloatingUI* WorldSettingsUI = GetOwner().SpawnTransientSceneActor< AVREditorDockableWindow >( TEXT( "WorldSettingsUI" ), bWithSceneComponent );
+			WorldSettingsUI->SetSlateWidget( *this, WidgetToDraw, DefaultResolution, VREd::EditorUISize->GetFloat(), AVREditorFloatingUI::EDockedTo::Nothing );
+			WorldSettingsUI->ShowUI( false );
+			FloatingUIs.Add( WorldSettingsUI );
+
+			EditorUIPanels[ (int32)EEditorUIPanel::WorldSettings ] = WorldSettingsUI;
 		}
 	}
 }
 
-void FVREditorUISystem::CleanUpActorsBeforeMapChangeOrSimulate()
+void UVREditorUISystem::CleanUpActorsBeforeMapChangeOrSimulate()
 {
 	for( AVREditorFloatingUI* FloatingUIPtr : FloatingUIs )
 	{
@@ -898,7 +863,7 @@ void FVREditorUISystem::CleanUpActorsBeforeMapChangeOrSimulate()
 	FAssetEditorManager::Get().OnAssetEditorOpened().RemoveAll(this);
 }
 
-void FVREditorUISystem::OnAssetEditorOpened(UObject* Asset)
+void UVREditorUISystem::OnAssetEditorOpened(UObject* Asset)
 {
 	// We need to disable drag drop on the tabs spawned in VR mode.
 	TArray<IAssetEditorInstance*> Editors = FAssetEditorManager::Get().FindEditorsForAsset(Asset);
@@ -908,7 +873,7 @@ void FVREditorUISystem::OnAssetEditorOpened(UObject* Asset)
 	}
 }
 
-bool FVREditorUISystem::IsWidgetAnEditorUIWidget( const UActorComponent* WidgetComponent ) const
+bool UVREditorUISystem::IsWidgetAnEditorUIWidget( const UActorComponent* WidgetComponent ) const
 {
 	if( WidgetComponent != nullptr && WidgetComponent->IsA( UWidgetComponent::StaticClass() ) )
 	{
@@ -928,7 +893,7 @@ bool FVREditorUISystem::IsWidgetAnEditorUIWidget( const UActorComponent* WidgetC
 }
 
 
-bool FVREditorUISystem::IsShowingEditorUIPanel( const EEditorUIPanel EditorUIPanel ) const
+bool UVREditorUISystem::IsShowingEditorUIPanel( const EEditorUIPanel EditorUIPanel ) const
 {
 	AVREditorFloatingUI* Panel = EditorUIPanels[ (int32)EditorUIPanel ];
 	if( Panel != nullptr )
@@ -940,7 +905,7 @@ bool FVREditorUISystem::IsShowingEditorUIPanel( const EEditorUIPanel EditorUIPan
 }
 
 
-void FVREditorUISystem::ShowEditorUIPanel( const UWidgetComponent* WidgetComponent, const int32 HandIndex, const bool bShouldShow, const bool OnHand, const bool bRefreshQuickMenu, const bool bPlaySound )
+void UVREditorUISystem::ShowEditorUIPanel( const UWidgetComponent* WidgetComponent, UVREditorInteractor* Interactor, const bool bShouldShow, const bool OnHand, const bool bRefreshQuickMenu, const bool bPlaySound )
 {
 	AVREditorFloatingUI* Panel = nullptr;
 	for( AVREditorFloatingUI* CurrentPanel : EditorUIPanels )
@@ -952,17 +917,17 @@ void FVREditorUISystem::ShowEditorUIPanel( const UWidgetComponent* WidgetCompone
 		}
 	}
 
-	ShowEditorUIPanel( Panel, HandIndex, bShouldShow, OnHand, bRefreshQuickMenu, bPlaySound );
+	ShowEditorUIPanel( Panel, Interactor, bShouldShow, OnHand, bRefreshQuickMenu, bPlaySound );
 }
 
 
-void FVREditorUISystem::ShowEditorUIPanel( const EEditorUIPanel EditorUIPanel, const int32 HandIndex, const bool bShouldShow, const bool OnHand, const bool bRefreshQuickMenu, const bool bPlaySound )
+void UVREditorUISystem::ShowEditorUIPanel( const EEditorUIPanel EditorUIPanel, UVREditorInteractor* Interactor, const bool bShouldShow, const bool OnHand, const bool bRefreshQuickMenu, const bool bPlaySound )
 {
 	AVREditorFloatingUI* Panel = EditorUIPanels[ (int32)EditorUIPanel ];
-	ShowEditorUIPanel( Panel, HandIndex, bShouldShow, OnHand, bRefreshQuickMenu, bPlaySound );
+	ShowEditorUIPanel( Panel, Interactor, bShouldShow, OnHand, bRefreshQuickMenu, bPlaySound );
 }
 
-void FVREditorUISystem::ShowEditorUIPanel( AVREditorFloatingUI* Panel, const int32 HandIndex, const bool bShouldShow, const bool OnHand, const bool bRefreshQuickMenu, const bool bPlaySound )
+void UVREditorUISystem::ShowEditorUIPanel( AVREditorFloatingUI* Panel, UVREditorInteractor* Interactor, const bool bShouldShow, const bool OnHand, const bool bRefreshQuickMenu, const bool bPlaySound )
 {
 	if( Panel != nullptr )
 	{
@@ -975,42 +940,47 @@ void FVREditorUISystem::ShowEditorUIPanel( AVREditorFloatingUI* Panel, const int
 				AVREditorFloatingUI* OtherPanel = EditorUIPanels[PanelIndex];
 				if( OtherPanel != nullptr && OtherPanel->IsUIVisible() && OtherPanel->GetDockedTo() != AVREditorFloatingUI::EDockedTo::Room )
 				{
-					const uint32 OtherPanelDockedToHandIndex = OtherPanel->GetDockedTo() == AVREditorFloatingUI::EDockedTo::LeftHand ? VREditorConstants::LeftHandIndex : VREditorConstants::RightHandIndex;
-					if( OtherPanelDockedToHandIndex == HandIndex )
+					UVREditorInteractor* OtherPanelDockToVREditorInteractor = GetOwner().GetHandInteractor( OtherPanel->GetDockedTo() == AVREditorFloatingUI::EDockedTo::LeftHand ? EControllerHand::Left : EControllerHand::Right );
+					if( OtherPanelDockToVREditorInteractor == Interactor )
 					{
 						OtherPanel->ShowUI( false );
 						OtherPanel->SetDockedTo( AVREditorFloatingUI::EDockedTo::Nothing );
-						Owner.GetVirtualHand( OtherPanelDockedToHandIndex ).bHasUIInFront = false;
+						OtherPanelDockToVREditorInteractor->SetHasUIInFront( false );
 					}
 				}
 			}
 			
-			const AVREditorFloatingUI::EDockedTo NewDockedTo = HandIndex == VREditorConstants::LeftHandIndex ? AVREditorFloatingUI::EDockedTo::LeftHand : AVREditorFloatingUI::EDockedTo::RightHand;
-			Panel->SetDockedTo( NewDockedTo );
+			UVREditorMotionControllerInteractor* MotionControllerInteractor = Cast<UVREditorMotionControllerInteractor>( Interactor );
+			if ( MotionControllerInteractor )
+			{
+				const AVREditorFloatingUI::EDockedTo NewDockedTo = MotionControllerInteractor->GetControllerSide() == EControllerHand::Left ? AVREditorFloatingUI::EDockedTo::LeftHand : AVREditorFloatingUI::EDockedTo::RightHand;
+				Panel->SetDockedTo( NewDockedTo );
+			}
 			
-			if (bShouldShow)
+			if ( bShouldShow )
 			{
 				Panel->SetScale( Panel->GetInitialScale() );
 			}
 
 			const FVector EditorUIRelativeOffset( Panel->GetSize().Y * 0.5f + VREd::EditorUIOffsetFromHand->GetFloat(), 0.0f, 0.0f );
 			Panel->SetRelativeOffset( EditorUIRelativeOffset );
-			Panel->SetLocalRotation( FRotator( 90.0f, 180.0f, 0.0f ) ); // Todo needs initial rotation 
+			Panel->SetLocalRotation( FRotator( 90.0f, 180.0f, 0.0f ) ); // @todo vreditor: needs initial rotation
 		}
 
 		Panel->ShowUI( bShouldShow );
 
-		if (Panel->GetDockedTo() == AVREditorFloatingUI::EDockedTo::LeftHand || Panel->GetDockedTo() == AVREditorFloatingUI::EDockedTo::RightHand)
+		if ( Panel->GetDockedTo() == AVREditorFloatingUI::EDockedTo::LeftHand || Panel->GetDockedTo() == AVREditorFloatingUI::EDockedTo::RightHand )
 		{
-			Owner.GetVirtualHand( HandIndex ).bHasUIInFront = bShouldShow;
+			UVREditorInteractor* PanelDockedToVREditorInteractor = GetOwner().GetHandInteractor( Panel->GetDockedTo() == AVREditorFloatingUI::EDockedTo::LeftHand ? EControllerHand::Left : EControllerHand::Right );
+			PanelDockedToVREditorInteractor->SetHasUIInFront( bShouldShow );
 		}
 
 		if( bPlaySound )
 		{
-			UGameplayStatics::PlaySound2D( Owner.GetWorld(), bShouldShow ? ShowUISound : HideUISound );
+			UGameplayStatics::PlaySound2D( GetOwner().GetWorld(), bShouldShow ? ShowUISound : HideUISound );
 		}
 
-		if (bRefreshQuickMenu && QuickMenuUI)
+		if ( bRefreshQuickMenu && QuickMenuUI )
 		{
 			QuickMenuUI->GetUserWidget<UVREditorQuickMenu>()->RefreshUI();
 		}
@@ -1018,133 +988,134 @@ void FVREditorUISystem::ShowEditorUIPanel( AVREditorFloatingUI* Panel, const int
 }
 
 
-bool FVREditorUISystem::IsShowingRadialMenu( const int32 HandIndex ) const
+bool UVREditorUISystem::IsShowingRadialMenu( UVREditorInteractor* Interactor ) const
 {
-	int32 DockedToHandIndex = INDEX_NONE;
-	DockedToHandIndex = QuickRadialMenu->GetDockedTo() == AVREditorFloatingUI::EDockedTo::LeftHand ? VREditorConstants::LeftHandIndex : VREditorConstants::RightHandIndex;
-	return DockedToHandIndex == HandIndex && !QuickRadialMenu->bHidden;
+	const EControllerHand DockedToHand = QuickRadialMenu->GetDockedTo() == AVREditorFloatingUI::EDockedTo::LeftHand ? EControllerHand::Left : EControllerHand::Right;
+	const UVREditorInteractor* DockedToInteractor = GetOwner().GetHandInteractor( DockedToHand );
+	return DockedToInteractor == Interactor && !QuickRadialMenu->bHidden;
 }
 
 
-void FVREditorUISystem::UpdateRadialMenu( const int32 HandIndex )
+void UVREditorUISystem::UpdateRadialMenu( UVREditorInteractor* Interactor )
 {
-	if(QuickRadialMenu->bHidden)
+	if( QuickRadialMenu->bHidden )
 	{
 		QuickRadialMenu->ShowUI( true, false );
 	}
 
-	if(!QuickRadialMenu->bHidden)
+	if( !QuickRadialMenu->bHidden )
 	{
 		RadialMenuHideDelayTime = 0.0f;
-		QuickRadialMenu->GetUserWidget<UVREditorRadialMenu>()->Update( Owner.GetVirtualHand( HandIndex ) );
+		QuickRadialMenu->GetUserWidget<UVREditorRadialMenu>()->Update( Interactor );
 	}
 }
 
 
-void FVREditorUISystem::TryToSpawnRadialMenu( const int32 HandIndex )
+void UVREditorUISystem::TryToSpawnRadialMenu( UVREditorInteractor* Interactor )
 {
-	FVirtualHand& Hand = Owner.GetVirtualHand( HandIndex );
-
-	int32 DockedToHandIndex = INDEX_NONE;
-	if( !QuickRadialMenu->bHidden )
+	UVREditorMotionControllerInteractor* MotionControllerInteractor = Cast<UVREditorMotionControllerInteractor>( Interactor );
+	if( MotionControllerInteractor )
 	{
-		DockedToHandIndex = QuickRadialMenu->GetDockedTo() == AVREditorFloatingUI::EDockedTo::LeftHand ? VREditorConstants::LeftHandIndex : VREditorConstants::RightHandIndex;
-	}
+		UVREditorInteractor* DockedToInteractor = nullptr;
+		if ( !QuickRadialMenu->bHidden )
+		{
+			const EControllerHand ControllerHand = QuickRadialMenu->GetDockedTo() == AVREditorFloatingUI::EDockedTo::LeftHand ? EControllerHand::Left : EControllerHand::Right;
+			DockedToInteractor = GetOwner().GetHandInteractor( ControllerHand );
+		}
 
-	bool bNeedsSpawn =
-		( QuickRadialMenu->bHidden || DockedToHandIndex != HandIndex ) &&
-		Hand.DraggingMode != EVREditorDraggingMode::ActorsAtLaserImpact &&	// Don't show radial menu if the hand is busy dragging something around
-		Hand.DraggingMode != EVREditorDraggingMode::ActorsFreely &&
-		Hand.DraggingMode != EVREditorDraggingMode::ActorsWithGizmo &&
-		Hand.DraggingMode != EVREditorDraggingMode::AssistingDrag &&
-		DraggingUIHandIndex != HandIndex &&
-		!Hand.bIsHoveringOverUI;	// Don't show radial menu when aiming at a UI  (too much clutter)
-
-	UVREditorRadialMenu* RadialMenu = QuickRadialMenu->GetUserWidget<UVREditorRadialMenu>();
-	// We need to update the trackpad position in the radialmenu before checking if it can be used
-	RadialMenu->Update( Hand );
-	if( RadialMenu && RadialMenu->IsInMenuRadius() )
-	{
-		bNeedsSpawn = false;
-	}
+		EViewportInteractionDraggingMode DraggingMode = Interactor->GetDraggingMode();
 
-	if( bNeedsSpawn )
-	{
-		DockedToHandIndex = HandIndex;
+		bool bNeedsSpawn =
+			( QuickRadialMenu->bHidden || DockedToInteractor != Interactor ) &&
+			DraggingMode != EViewportInteractionDraggingMode::ActorsAtLaserImpact &&	// Don't show radial menu if the hand is busy dragging something around
+			DraggingMode != EViewportInteractionDraggingMode::ActorsFreely &&
+			DraggingMode != EViewportInteractionDraggingMode::ActorsWithGizmo &&
+			DraggingMode != EViewportInteractionDraggingMode::AssistingDrag &&
+			( InteractorDraggingUI == nullptr || InteractorDraggingUI != Interactor ) &&
+			!Interactor->IsHoveringOverUI();	// Don't show radial menu when aiming at a UI  (too much clutter)
+
+		UVREditorRadialMenu* RadialMenu = QuickRadialMenu->GetUserWidget<UVREditorRadialMenu>();
+		// We need to update the trackpad position in the radialmenu before checking if it can be used
+		RadialMenu->Update( Interactor );
+		if ( RadialMenu && RadialMenu->IsInMenuRadius() )
+		{
+			bNeedsSpawn = false;
+		}
 
-		const AVREditorFloatingUI::EDockedTo DockedTo = DockedToHandIndex == VREditorConstants::LeftHandIndex ? AVREditorFloatingUI::EDockedTo::LeftHand : AVREditorFloatingUI::EDockedTo::RightHand;
-		QuickRadialMenu->SetDockedTo( DockedTo );
-		QuickRadialMenu->ShowUI( true );
+		if ( bNeedsSpawn )
+		{
+			const AVREditorFloatingUI::EDockedTo DockedTo = MotionControllerInteractor->GetControllerSide() == EControllerHand::Left ? AVREditorFloatingUI::EDockedTo::LeftHand : AVREditorFloatingUI::EDockedTo::RightHand;
+			QuickRadialMenu->SetDockedTo( DockedTo );
+			QuickRadialMenu->ShowUI( true );
+		}
 	}
 }
 
 
-void FVREditorUISystem::HideRadialMenu( const int32 HandIndex )
+void UVREditorUISystem::HideRadialMenu( UVREditorInteractor* Interactor )
 {
-	if( IsShowingRadialMenu( HandIndex ) )
+	// Only hide the radial menu if the passed interactor is actually the interactor with the radial menu
+	if( IsShowingRadialMenu( Interactor ) )
 	{
-		UVREditorRadialMenu* RadialMenu = QuickRadialMenu->GetUserWidget<UVREditorRadialMenu>();
 		QuickRadialMenu->ShowUI( false, true, VREd::RadialMenuFadeDelay->GetFloat() );
 	}
 }
 
 
-FTransform FVREditorUISystem::MakeDockableUITransformOnLaser( AVREditorDockableWindow* InitDraggingDockUI, const int32 HandIndex, const float DockSelectDistance ) const
+FTransform UVREditorUISystem::MakeDockableUITransformOnLaser( AVREditorDockableWindow* InitDraggingDockUI, UVREditorInteractor* Interactor, const float DockSelectDistance ) const
 {
-	const FVirtualHand& Hand = Owner.GetVirtualHand( HandIndex );
-
 	FTransform HandTransform;
 	FVector HandForward;
-	Owner.GetHandTransformAndForwardVector( HandIndex, HandTransform, HandForward );
+	Interactor->GetTransformAndForwardVector( HandTransform, HandForward );
+	const FTransform InteractorTransform = Interactor->GetTransform();
+
+	const FVector NewLocation = InteractorTransform.GetLocation() + ( HandForward.GetSafeNormal() * DockSelectDistance );
+	FRotator NewRotation = ( InteractorTransform.GetLocation() - NewLocation ).ToOrientationRotator();
+	NewRotation.Roll = -InteractorTransform.GetRotation().Rotator().Roll;
 
-	const FVector NewLocation = Hand.Transform.GetLocation() + ( HandForward.GetSafeNormal() * DockSelectDistance );
-	
-	FRotator NewRotation = ( Hand.Transform.GetLocation() - NewLocation ).ToOrientationRotator();
-	NewRotation.Roll = -Hand.Transform.GetRotation().Rotator().Roll;
-	
 	const FTransform LaserImpactToWorld( NewRotation, NewLocation );
 	return LaserImpactToWorld;
 }
 
 
-FTransform FVREditorUISystem::MakeDockableUITransform( AVREditorDockableWindow* InitDraggingDockUI, const int32 HandIndex, const float DockSelectDistance )
+FTransform UVREditorUISystem::MakeDockableUITransform( AVREditorDockableWindow* InitDraggingDockUI, UVREditorInteractor* Interactor, const float DockSelectDistance )
 {
-	const FTransform UIOnLaserToWorld = MakeDockableUITransformOnLaser( DraggingUI, HandIndex, DockSelectDistance );
+	const FTransform UIOnLaserToWorld = MakeDockableUITransformOnLaser( DraggingUI, Interactor, DockSelectDistance );
 	const FTransform UIToUIOnLaser = DraggingUIOffsetTransform;
 	
 	const FTransform UIToWorld = UIToUIOnLaser * UIOnLaserToWorld;
 
-	if( DraggingUI != nullptr &&  DraggingUI->GetDockedTo() == AVREditorFloatingUI::EDockedTo::Dragging )
+	if( DraggingUI != nullptr &&  DraggingUI->GetDockedTo() == AVREditorFloatingUI::EDockedTo::Dragging && InteractorDraggingUI && Interactor == InteractorDraggingUI )
 	{
-		
-		FVirtualHand& Hand = Owner.GetVirtualHand( DraggingUIHandIndex );
-		FVirtualHand& OtherHand = Owner.GetOtherHand( DraggingUIHandIndex );
-		
-		const float WorldScaleFactor = Owner.GetWorldScaleFactor();
-		const FVector DraggingUILocation = Hand.HoverLocation;
-		const FVector OtherHandLocation =  OtherHand.Transform.GetLocation();
-		LastDraggingHoverLocation = UIOnLaserToWorld.GetLocation();
+		UVREditorMotionControllerInteractor* OtherInteractor = Cast<UVREditorMotionControllerInteractor>( InteractorDraggingUI->GetOtherInteractor() );
+		if ( OtherInteractor )
+		{
+			const float WorldScaleFactor = GetOwner().GetWorldScaleFactor();
+			const FVector DraggingUILocation = InteractorDraggingUI->GetHoverLocation();
+			const FVector OtherHandLocation = OtherInteractor->GetTransform().GetLocation();
+			LastDraggingHoverLocation = UIOnLaserToWorld.GetLocation();
 
-		const float Distance = FVector::Dist( UIOnLaserToWorld.GetLocation(), OtherHandLocation ) / WorldScaleFactor;
+			const float Distance = FVector::Dist( UIOnLaserToWorld.GetLocation(), OtherHandLocation ) / WorldScaleFactor;
 
-		// if dragged passed the threshold since starting dragging
-		if ( Distance > VREd::MinDockDragDistance->GetFloat() && !bDraggedDockFromHandPassedThreshold && ( DraggingUI->GetPreviouslyDockedTo() == AVREditorFloatingUI::EDockedTo::LeftHand || DraggingUI->GetPreviouslyDockedTo() == AVREditorFloatingUI::EDockedTo::RightHand ) )
-		{
-			bDraggedDockFromHandPassedThreshold = true;
-		}
+			// if dragged passed the threshold since starting dragging
+			if ( Distance > VREd::MinDockDragDistance->GetFloat() && !bDraggedDockFromHandPassedThreshold &&
+				( DraggingUI->GetPreviouslyDockedTo() == AVREditorFloatingUI::EDockedTo::LeftHand || DraggingUI->GetPreviouslyDockedTo() == AVREditorFloatingUI::EDockedTo::RightHand ) )
+			{
+				bDraggedDockFromHandPassedThreshold = true;
+			}
 
-		// Snapping to a hand when in range
-		if ( Distance <= VREd::MinDockDragDistance->GetFloat() && ( bDraggedDockFromHandPassedThreshold || DraggingUI->GetPreviouslyDockedTo() == AVREditorFloatingUI::EDockedTo::Room ) )
-		{
-			const int32 OtherHandIndex = Owner.GetOtherHandIndex( HandIndex );
-			const FVector EditorUIRelativeOffset( DraggingUI->GetSize().Y * 0.5f + VREd::EditorUIOffsetFromHand->GetFloat(), 0.0f, 0.0f );
-			const FTransform MovingUIGoalTransform = DraggingUI->MakeUITransformLockedToHand( OtherHandIndex, false, EditorUIRelativeOffset, FRotator( 90.0f, 180.0f, 0.0f ) );
-			const AVREditorFloatingUI::EDockedTo NewDockedTo = OtherHandIndex == VREditorConstants::LeftHandIndex ? AVREditorFloatingUI::EDockedTo::LeftHand : AVREditorFloatingUI::EDockedTo::RightHand;
-			DraggingUI->MoveTo( MovingUIGoalTransform, VREd::DockDragSpeed->GetFloat(), NewDockedTo );
-		}
-		else
-		{
-			DraggingUI->MoveTo( UIToWorld, VREd::DockDragSpeed->GetFloat(), AVREditorFloatingUI::EDockedTo::Room );
+			// Snapping to a hand when in range
+			if ( Distance <= VREd::MinDockDragDistance->GetFloat() && ( bDraggedDockFromHandPassedThreshold || DraggingUI->GetPreviouslyDockedTo() == AVREditorFloatingUI::EDockedTo::Room ) )
+			{
+				const FVector EditorUIRelativeOffset( DraggingUI->GetSize().Y * 0.5f + VREd::EditorUIOffsetFromHand->GetFloat(), 0.0f, 0.0f );
+				const FTransform MovingUIGoalTransform = DraggingUI->MakeUITransformLockedToHand( OtherInteractor, false, EditorUIRelativeOffset, FRotator( 90.0f, 180.0f, 0.0f ) );
+				const AVREditorFloatingUI::EDockedTo NewDockedTo = OtherInteractor->GetControllerSide() == EControllerHand::Left ? AVREditorFloatingUI::EDockedTo::LeftHand : AVREditorFloatingUI::EDockedTo::RightHand;
+				DraggingUI->MoveTo( MovingUIGoalTransform, VREd::DockDragSpeed->GetFloat(), NewDockedTo );
+			}
+			else
+			{
+				DraggingUI->MoveTo( UIToWorld, VREd::DockDragSpeed->GetFloat(), AVREditorFloatingUI::EDockedTo::Room );
+			}
 		}
 	}
 
@@ -1152,22 +1123,32 @@ FTransform FVREditorUISystem::MakeDockableUITransform( AVREditorDockableWindow*
 }
 
 
-AVREditorFloatingUI* FVREditorUISystem::StartDraggingDockUI( AVREditorDockableWindow* InitDraggingDockUI, const int32 HandIndex, const float DockSelectDistance )
+void UVREditorUISystem::StartDraggingDockUI( AVREditorDockableWindow* InitDraggingDockUI, UVREditorInteractor* Interactor, const float DockSelectDistance )
 {
 	AVREditorFloatingUI::EDockedTo DockTo = InitDraggingDockUI->GetDockedTo();
 	if( DockTo == AVREditorFloatingUI::EDockedTo::LeftHand || DockTo == AVREditorFloatingUI::EDockedTo::RightHand )
 	{
 		bDraggedDockFromHandPassedThreshold = false;
-		Owner.GetOtherHand( HandIndex ).bHasUIInFront = false;
 	}
 
-	DraggingUIHandIndex = HandIndex;
+	// Set the laser of the interactor the dockable window was pulled from
+	UViewportInteractor* OtherInteractor = Interactor->GetOtherInteractor();
+	if ( OtherInteractor )
+	{
+		UVREditorInteractor* OtherVREditorInteractor = Cast<UVREditorInteractor>( OtherInteractor );
+		if ( OtherVREditorInteractor )
+		{
+			OtherVREditorInteractor->SetHasUIInFront( false );
+		}
+	}
+
+	InteractorDraggingUI = Interactor;
 
 	FTransform UIToWorld = InitDraggingDockUI->GetActorTransform();
 	UIToWorld.SetScale3D( FVector( 1.0f ) );
 	const FTransform WorldToUI = UIToWorld.Inverse();
 
-	const FTransform UIOnLaserToWorld = MakeDockableUITransformOnLaser( InitDraggingDockUI, HandIndex, DockSelectDistance );
+	const FTransform UIOnLaserToWorld = MakeDockableUITransformOnLaser( InitDraggingDockUI, Interactor, DockSelectDistance );
 	const FTransform UIOnLaserToUI = UIOnLaserToWorld * WorldToUI;
 	const FTransform UIToUIOnLaser = UIOnLaserToUI.Inverse();
 	DraggingUIOffsetTransform = UIToUIOnLaser;
@@ -1175,38 +1156,63 @@ AVREditorFloatingUI* FVREditorUISystem::StartDraggingDockUI( AVREditorDockableWi
 	DraggingUI = InitDraggingDockUI;
 	DraggingUI->SetDockedTo( AVREditorFloatingUI::EDockedTo::Dragging );
 
-	UGameplayStatics::PlaySound2D( Owner.GetWorld(), StartDragUISound );
-
-	return DraggingUI;
+	UGameplayStatics::PlaySound2D( GetOwner().GetWorld(), StartDragUISound );
 }
 
-AVREditorDockableWindow* FVREditorUISystem::GetDraggingDockUI() const
+AVREditorDockableWindow* UVREditorUISystem::GetDraggingDockUI() const
 {
 	return DraggingUI;
 }
 
-int32 FVREditorUISystem::GetDraggingDockUIHandIndex() const
+void UVREditorUISystem::StopDraggingDockUI( UVREditorInteractor* VREditorInteractor )
 {
-	return DraggingUIHandIndex;
-}
+	if ( IsInteractorDraggingDockUI( VREditorInteractor ) )
+	{
+		// Put the Dock back on the hand it came from or leave it where it is in the room
+		UViewportInteractor* OtherInteractor = VREditorInteractor->GetOtherInteractor();
+		if ( OtherInteractor )
+		{
+			UVREditorMotionControllerInteractor* MotionControllerOtherInteractor = Cast<UVREditorMotionControllerInteractor>( OtherInteractor );
+			if ( MotionControllerOtherInteractor )
+			{
+				bool bOnHand = false;
+				const float Distance = FVector::Dist( LastDraggingHoverLocation, MotionControllerOtherInteractor->GetTransform().GetLocation() ) / GetOwner().GetWorldScaleFactor();
+				if ( Distance > VREd::MinDockDragDistance->GetFloat() )
+				{
+					DraggingUI->SetDockedTo( AVREditorFloatingUI::EDockedTo::Room );
+				}
+				else
+				{
+					bOnHand = true;
+					const FVector EditorUIRelativeOffset( DraggingUI->GetSize().Y * 0.5f + VREd::EditorUIOffsetFromHand->GetFloat(), 0.0f, 0.0f );
+					const FTransform MovingUIGoalTransform = DraggingUI->MakeUITransformLockedToHand( MotionControllerOtherInteractor, false, EditorUIRelativeOffset, FRotator( 90.0f, 180.0f, 0.0f ) );
+					const AVREditorFloatingUI::EDockedTo NewDockedTo = MotionControllerOtherInteractor->GetControllerSide() == EControllerHand::Left ? AVREditorFloatingUI::EDockedTo::LeftHand : AVREditorFloatingUI::EDockedTo::RightHand;
+					DraggingUI->MoveTo( MovingUIGoalTransform, VREd::DockSnapTime->GetFloat(), NewDockedTo );
+				}
 
-void FVREditorUISystem::StopDraggingDockUI()
-{
-	DraggingUI = nullptr;
-	FVirtualHand& Hand = Owner.GetVirtualHand( DraggingUIHandIndex );
-	Hand.DraggingMode = EVREditorDraggingMode::Nothing;	
-	DraggingUIHandIndex = INDEX_NONE;
+				ShowEditorUIPanel( DraggingUI, MotionControllerOtherInteractor, true, bOnHand );
 
-	UGameplayStatics::PlaySound2D( Owner.GetWorld(), StopDragUISound );
+				DraggingUI = nullptr;
+				InteractorDraggingUI->SetDraggingMode( EViewportInteractionDraggingMode::Nothing );
+				InteractorDraggingUI = nullptr;
 
+				UGameplayStatics::PlaySound2D( GetOwner().GetWorld(), StopDragUISound );
+			}
+		}
+	}
 }
 
-bool FVREditorUISystem::IsDraggingDockUI()
+bool UVREditorUISystem::IsDraggingDockUI()
 {
 	return DraggingUI != nullptr && DraggingUI->GetDockedTo() == AVREditorFloatingUI::EDockedTo::Dragging;
 }
 
-void FVREditorUISystem::TogglePanelsVisibility()
+bool UVREditorUISystem::IsInteractorDraggingDockUI( const UVREditorInteractor* Interactor ) const
+{
+	return InteractorDraggingUI && InteractorDraggingUI == Interactor;
+}
+
+void UVREditorUISystem::TogglePanelsVisibility()
 {
 	bool bAnyPanelsVisible = false;
 	AVREditorFloatingUI* PanelOnHand = nullptr;
@@ -1230,7 +1236,7 @@ void FVREditorUISystem::TogglePanelsVisibility()
 	// Hide if there is any UI visible
 	const bool bShowUI = !bAnyPanelsVisible;
 	
-	UGameplayStatics::PlaySound2D( Owner.GetWorld(), bShowUI ? ShowUISound : HideUISound );
+	UGameplayStatics::PlaySound2D( GetOwner().GetWorld(), bShowUI ? ShowUISound : HideUISound );
 
 	if ( bSetDefaultLayout)
 	{
@@ -1239,20 +1245,20 @@ void FVREditorUISystem::TogglePanelsVisibility()
 	}
 	else
 	{
-		for (AVREditorFloatingUI* Panel : EditorUIPanels)
+		for ( AVREditorFloatingUI* Panel : EditorUIPanels )
 		{
-			if (Panel != nullptr && Panel->IsUIVisible() != bShowUI)
+			if ( Panel != nullptr && Panel->IsUIVisible() != bShowUI )
 			{
-				if (bShowUI)
+				if ( bShowUI )
 				{
 					bool SetNewVisibility = true;
 					const AVREditorFloatingUI::EDockedTo DockedTo = Panel->GetDockedTo();
-					if (DockedTo == AVREditorFloatingUI::EDockedTo::LeftHand || DockedTo == AVREditorFloatingUI::EDockedTo::RightHand)
+					if ( DockedTo == AVREditorFloatingUI::EDockedTo::LeftHand || DockedTo == AVREditorFloatingUI::EDockedTo::RightHand )
 					{
-						if (!PanelOnHand)
+						if ( !PanelOnHand )
 						{
-							const uint32 HandIndex = Panel->GetDockedTo() == AVREditorFloatingUI::EDockedTo::LeftHand ? VREditorConstants::LeftHandIndex : VREditorConstants::RightHandIndex;
-							Owner.GetVirtualHand( HandIndex ).bHasUIInFront = true;
+							UVREditorInteractor* VREditorInteractor = GetOwner().GetHandInteractor( Panel->GetDockedTo() == AVREditorFloatingUI::EDockedTo::LeftHand ? EControllerHand::Left : EControllerHand::Right );
+							VREditorInteractor->SetHasUIInFront( true );
 							PanelOnHand = Panel;
 						}
 						else
@@ -1260,12 +1266,12 @@ void FVREditorUISystem::TogglePanelsVisibility()
 							SetNewVisibility = false;
 						}
 					}
-					else if (DockedTo == AVREditorFloatingUI::EDockedTo::Nothing)
+					else if ( DockedTo == AVREditorFloatingUI::EDockedTo::Nothing )
 					{
 						SetNewVisibility = false;
 					}
 
-					if (SetNewVisibility)
+					if ( SetNewVisibility )
 					{
 						Panel->ShowUI( true );
 					}
@@ -1280,8 +1286,8 @@ void FVREditorUISystem::TogglePanelsVisibility()
 
 	if( !bShowUI && PanelOnHand )
 	{
-		const uint32 HandIndex = PanelOnHand->GetDockedTo() == AVREditorFloatingUI::EDockedTo::LeftHand ? VREditorConstants::LeftHandIndex : VREditorConstants::RightHandIndex;
-		Owner.GetVirtualHand( HandIndex ).bHasUIInFront = false;
+		UVREditorInteractor* VREditorInteractor = GetOwner().GetHandInteractor( PanelOnHand->GetDockedTo() == AVREditorFloatingUI::EDockedTo::LeftHand ? EControllerHand::Left : EControllerHand::Right );
+		VREditorInteractor->SetHasUIInFront( false );
 	}
 
 	if (QuickMenuUI)
@@ -1290,27 +1296,32 @@ void FVREditorUISystem::TogglePanelsVisibility()
 	}
 }
 
-float FVREditorUISystem::GetMaxDockWindowSize() const
+float UVREditorUISystem::GetMaxDockWindowSize() const
 {
 	return VREd::MaxDockWindowSize->GetFloat();
 }
 
-float FVREditorUISystem::GetMinDockWindowSize() const
+float UVREditorUISystem::GetMinDockWindowSize() const
 {
 	return VREd::MinDockWindowSize->GetFloat();
 }
 
-void FVREditorUISystem::OnProxyTabLaunched(TSharedPtr<SDockTab> NewTab)
+void UVREditorUISystem::OnProxyTabLaunched(TSharedPtr<SDockTab> NewTab)
 {
 	ShowAssetEditor();
 }
 
-void FVREditorUISystem::OnAttentionDrawnToTab(TSharedPtr<SDockTab> NewTab)
+void UVREditorUISystem::OnAttentionDrawnToTab(TSharedPtr<SDockTab> NewTab)
 {
-	ShowAssetEditor();
+	// @todo vreditor: When clicking on Modes icons in the Modes panel, the LevelEditor tab is invoked which causes an empty asset editor to pop-up
+	static FName LevelEditorTabID( "LevelEditor" );
+	if( NewTab->GetLayoutIdentifier().TabType != LevelEditorTabID )
+	{
+		ShowAssetEditor();
+	}
 }
 
-void FVREditorUISystem::ShowAssetEditor()
+void UVREditorUISystem::ShowAssetEditor()
 {
 	bRefocusViewport = true;
 
@@ -1319,26 +1330,28 @@ void FVREditorUISystem::ShowAssetEditor()
 	if ( !IsShowingEditorUIPanel(EEditorUIPanel::AssetEditor) )
 	{
 		// Always spawn on a hand.  But which hand?  Well, we'll choose the hand that isn't actively clicking on something using a laser.
-		const int32 HandIndex = Owner.GetVirtualHand(VREditorConstants::LeftHandIndex).bIsClickingOnUI ? VREditorConstants::RightHandIndex : VREditorConstants::LeftHandIndex;	// Hand that did not clicked with a laser
+		UVREditorInteractor* VREditorInteractor = GetOwner().GetHandInteractor( EControllerHand::Left ); // Hand that did not clicked with a laser
+		if ( VREditorInteractor->IsClickingOnUI() )
+		{
+			VREditorInteractor = GetOwner().GetHandInteractor( EControllerHand::Right );
+		}
+		
 		const bool bShouldShow = true;
 		const bool bShowOnHand = true;
-		ShowEditorUIPanel(EEditorUIPanel::AssetEditor, HandIndex, bShouldShow, bShowOnHand);
+		ShowEditorUIPanel( EEditorUIPanel::AssetEditor, VREditorInteractor, bShouldShow, bShowOnHand );
 
 		// Play haptic effect so user knows to look at their hand that now has UI on it!
-		const float Strength = VREd::UIAssetEditorSummonedOnHandHapticFeedbackStrength->GetFloat();
-		Owner.PlayHapticEffect(
-			HandIndex == VREditorConstants::LeftHandIndex ? Strength : 0.0f,
-			HandIndex == VREditorConstants::RightHandIndex ? Strength : 0.0f);
+		VREditorInteractor->PlayHapticEffect( VREd::UIAssetEditorSummonedOnHandHapticFeedbackStrength->GetFloat() );
 	}
 }
 
-void FVREditorUISystem::TogglePanelVisibility( const EEditorUIPanel EditorUIPanel )
+void UVREditorUISystem::TogglePanelVisibility( const EEditorUIPanel EditorUIPanel )
 {
 	AVREditorFloatingUI* Panel = EditorUIPanels[(int32)EditorUIPanel];
 	if (Panel != nullptr)
 	{
 		const bool bIsShowing = Panel->IsUIVisible();
-		const int32 HandIndexWithQuickMenu = QuickMenuUI->GetDockedTo() == AVREditorFloatingUI::EDockedTo::LeftArm ? VREditorConstants::LeftHandIndex : VREditorConstants::RightHandIndex;
+		UVREditorInteractor* InteractorWithQuickMenu = GetOwner().GetHandInteractor( QuickMenuUI->GetDockedTo() == AVREditorFloatingUI::EDockedTo::LeftArm ? EControllerHand::Left : EControllerHand::Right );
 		if( Panel->GetDockedTo() == AVREditorFloatingUI::EDockedTo::Room )
 		{
 			Panel->ShowUI( false );
@@ -1346,31 +1359,32 @@ void FVREditorUISystem::TogglePanelVisibility( const EEditorUIPanel EditorUIPane
 		}
 		else
 		{
-			ShowEditorUIPanel( Panel, HandIndexWithQuickMenu, !bIsShowing, true );
+			ShowEditorUIPanel( Panel, InteractorWithQuickMenu, !bIsShowing, true );
 		}
 	}
 }
 
-void FVREditorUISystem::SetDefaultWindowLayout()
+void UVREditorUISystem::SetDefaultWindowLayout()
 {
-	for (int32 PanelIndex = 0; PanelIndex < DefaultWindowTransforms.Num(); ++PanelIndex)
+	for ( int32 PanelIndex = 0; PanelIndex < DefaultWindowTransforms.Num(); ++PanelIndex )
 	{
 		AVREditorFloatingUI* Panel = EditorUIPanels[PanelIndex];
-		if (Panel)
+		if ( Panel )
 		{
 			const AVREditorFloatingUI::EDockedTo DockedTo = Panel->GetDockedTo();
 			if (DockedTo == AVREditorFloatingUI::EDockedTo::LeftHand || DockedTo == AVREditorFloatingUI::EDockedTo::RightHand)
 			{
-				const int32 HandIndex = DockedTo == AVREditorFloatingUI::EDockedTo::LeftHand ? VREditorConstants::LeftHandIndex : VREditorConstants::RightHandIndex;
-				Owner.GetVirtualHand( HandIndex ).bHasUIInFront = false;
+				const EControllerHand ControllerHand = DockedTo == AVREditorFloatingUI::EDockedTo::LeftHand ? EControllerHand::Left : EControllerHand::Right;
+				UVREditorInteractor* Interactor = GetOwner().GetHandInteractor( ControllerHand );
+				Interactor->SetHasUIInFront( false );
 			}
 
 			Panel->SetDockedTo( AVREditorFloatingUI::EDockedTo::Room );
 			Panel->ShowUI( true );
 
 			// Make sure the UIs are centered around the direction your head is looking (yaw only!)
-			const FVector RoomSpaceHeadLocation = Owner.GetRoomSpaceHeadTransform().GetLocation() / Owner.GetWorldScaleFactor();
-			FRotator RoomSpaceHeadYawRotator = Owner.GetRoomSpaceHeadTransform().GetRotation().Rotator();
+			const FVector RoomSpaceHeadLocation = GetOwner().GetRoomSpaceHeadTransform().GetLocation() / GetOwner().GetWorldScaleFactor();
+			FRotator RoomSpaceHeadYawRotator = GetOwner().GetRoomSpaceHeadTransform().GetRotation().Rotator();
 			RoomSpaceHeadYawRotator.Pitch = 0.0f;
 			RoomSpaceHeadYawRotator.Roll = 0.0f;
 
diff --git a/Engine/Source/Editor/VREditor/VREditorUISystem.h b/Engine/Source/Editor/VREditor/VREditorUISystem.h
--- a/Engine/Source/Editor/VREditor/VREditorUISystem.h
+++ b/Engine/Source/Editor/VREditor/VREditorUISystem.h
@@ -2,11 +2,19 @@
 
 #pragma once
 
+#include "VREditorUISystem.generated.h"
+
+class UVREditorInteractor;
+class UViewportInteractor;
+
 /**
  * VR Editor user interface manager
  */
-class FVREditorUISystem : public FGCObject
+UCLASS()
+class UVREditorUISystem : public UObject
 {
+	GENERATED_UCLASS_BODY()
+
 public:
 
 	enum class EEditorUIPanel
@@ -17,20 +25,17 @@ public:
 		Modes,
 		Tutorial,
 		AssetEditor,
-
+		WorldSettings,
 		// ...
 
 		TotalCount,
 	};
 
-	/** Default constructor for FVREditorUISystem, which sets up safe defaults */
-	FVREditorUISystem( class FVREditorMode& InitOwner );
-
-	/** Clean up before destruction */
-	virtual ~FVREditorUISystem();
+	/** Shuts down the UISystem whenever the mode is exited */
+	void Shutdown();
 
-	// FGCObject overrides
-	virtual void AddReferencedObjects( FReferenceCollector& Collector ) override;
+	/** Initializes default values for the UISystem and creates the UI panels */
+	void Init();
 
 	/** Called by VREditorMode to update us every frame */
 	void Tick( class FEditorViewportClient* ViewportClient, const float DeltaTime );
@@ -38,16 +43,22 @@ public:
 	/** Called by VREditorMode to draw debug graphics every frame */
 	void Render( const class FSceneView* SceneView, class FViewport* Viewport, class FPrimitiveDrawInterface* PDI );
 
+	/** Sets the owner of this system */
+	void SetOwner( class FVREditorMode* NewOwner )
+	{
+		VRMode = NewOwner;
+	}
+
 	/** Gets the owner of this system */
 	class FVREditorMode& GetOwner()
 	{
-		return Owner;
+		return *VRMode;
 	}
 
 	/** Gets the owner of this system (const) */
 	const class FVREditorMode& GetOwner() const
 	{
-		return Owner;
+		return *VRMode;
 	}
 
 	/** Called by our owner right before a map is loaded or switching between Simulate and normal Editor, so we can clean
@@ -61,43 +72,42 @@ public:
 	bool IsShowingEditorUIPanel( const EEditorUIPanel EditorUIPanel ) const;
 
 	/** Sets whether the specified editor UI panel should be visible.  Any other UI floating off this hand will be dismissed when showing it. */
-	void ShowEditorUIPanel( const class UWidgetComponent* WidgetComponent, const int32 HandIndex, const bool bShouldShow, const bool OnHand = false, const bool bRefreshQuickMenu = true, const bool bPlaySound = true );
-	void ShowEditorUIPanel( const EEditorUIPanel EditorUIPanel, const int32 HandIndex, const bool bShouldShow, const bool OnHand = false, const bool bRefreshQuickMenu = true, const bool bPlaySound = true );
-	void ShowEditorUIPanel( class AVREditorFloatingUI* Panel, const int32 HandIndex, const bool bShouldShow, const bool OnHand = false, const bool bRefreshQuickMenu = true, const bool bPlaySound = true );
+	void ShowEditorUIPanel( const class UWidgetComponent* WidgetComponent, UVREditorInteractor* Interactor, const bool bShouldShow, const bool OnHand = false, const bool bRefreshQuickMenu = true, const bool bPlaySound = true );
+	void ShowEditorUIPanel( const EEditorUIPanel EditorUIPanel, UVREditorInteractor* Interactor, const bool bShouldShow, const bool OnHand = false, const bool bRefreshQuickMenu = true, const bool bPlaySound = true );
+	void ShowEditorUIPanel( class AVREditorFloatingUI* Panel, UVREditorInteractor* Interactor, const bool bShouldShow, const bool OnHand = false, const bool bRefreshQuickMenu = true, const bool bPlaySound = true );
 
 	/** Returns true if the radial menu is visible on this hand */
-	bool IsShowingRadialMenu( const int32 HandIndex ) const;
+	bool IsShowingRadialMenu( UVREditorInteractor* Interactor ) const;
 
 	/** Tries to spawn the radial menu (if the specified hand isn't doing anything else) */
-	void TryToSpawnRadialMenu( const int32 HandIndex );
+	void TryToSpawnRadialMenu( UVREditorInteractor* Interactor );
 
 	/** Updates the radial menu with the current hand trackpad location */
-	void UpdateRadialMenu( const int32 HandIndex );
+	void UpdateRadialMenu( UVREditorInteractor* Interactor );
 
 	/** Hides the radial menu if the specified hand is showing it */
-	void HideRadialMenu( const int32 HandIndex );
+	void HideRadialMenu( UVREditorInteractor* Interactor );
 
 	/** Start dragging a dock window on the hand */
-	class AVREditorFloatingUI* StartDraggingDockUI( class AVREditorDockableWindow* InitDraggingDockUI, const int32 HandIndex, const float DockSelectDistance );
+	void StartDraggingDockUI( class AVREditorDockableWindow* InitDraggingDockUI, UVREditorInteractor* Interactor, const float DockSelectDistance );
 
 	/** Makes up a transform for a dockable UI when dragging it with a laser at the specified distance from the laser origin */
-	FTransform MakeDockableUITransformOnLaser( AVREditorDockableWindow* InitDraggingDockUI, const int32 HandIndex, const float DockSelectDistance ) const;
+	FTransform MakeDockableUITransformOnLaser( AVREditorDockableWindow* InitDraggingDockUI, UVREditorInteractor* Interactor, const float DockSelectDistance ) const;
 
 	/** Makes up a transform for a dockable UI when dragging it that includes the original offset from the laser's impact point */
-	FTransform MakeDockableUITransform( AVREditorDockableWindow* InitDraggingDockUI, const int32 HandIndex, const float DockSelectDistance );
+	FTransform MakeDockableUITransform( AVREditorDockableWindow* InitDraggingDockUI, UVREditorInteractor* Interactor, const float DockSelectDistance );
 	
 	/** Returns the current dragged dock window, nullptr if none */
 	class AVREditorDockableWindow* GetDraggingDockUI() const;
 
-	/** Gets the hand index of the hand that is dragging a dock UI */
-	int32 GetDraggingDockUIHandIndex() const;
-
 	/** Resets all values to stop dragging the current dock window */
-	void StopDraggingDockUI();
+	void StopDraggingDockUI( UVREditorInteractor* VREditorInteractor );
 
 	/** Are we currently dragging a dock window */
 	bool IsDraggingDockUI();	
 
+	bool IsInteractorDraggingDockUI( const UVREditorInteractor* Interactor ) const;
+
 	/** Hides and unhides all the editor UI panels */
 	void TogglePanelsVisibility();
 
@@ -113,10 +123,11 @@ public:
 protected:
 
 	/** Called when the user presses a button on their motion controller device */
-	void OnVRAction( FEditorViewportClient& ViewportClient, const FVRAction VRAction, const EInputEvent Event, bool& bIsInputCaptured, bool& bOutWasHandled );
+	void OnVRAction( class FEditorViewportClient& ViewportClient, UViewportInteractor* Interactor,
+		const struct FViewportActionKeyInput& Action, bool& bOutIsInputCaptured, bool& bWasHandled );
 
 	/** Called every frame to update hover state */
-	void OnVRHoverUpdate( FEditorViewportClient& ViewportClient, const int32 HandIndex, FVector& HoverImpactPoint, bool& bIsHoveringOverUI, bool& bOutWasHandled );
+	void OnVRHoverUpdate( class FEditorViewportClient& ViewportClient, UViewportInteractor* Interactor, FVector& HoverImpactPoint, bool& bWasHandled );
 
 	/** Testing Slate UI */
 	void CreateUIs();
@@ -142,18 +153,22 @@ protected:
 protected:
 
 	/** Owning object */
-	class FVREditorMode& Owner;
+	class FVREditorMode* VRMode;
 
 	/** All of the floating UIs.  These may or may not be visible (spawned) */
+	UPROPERTY()
 	TArray< class AVREditorFloatingUI* > FloatingUIs;
 
 	/** Our Quick Menu UI */
+	UPROPERTY()
 	AVREditorFloatingUI* QuickMenuUI;
 
 	/** Editor UI panels */
+	UPROPERTY()
 	TArray<AVREditorFloatingUI*> EditorUIPanels;
 
 	/** The Radial Menu UI */
+	UPROPERTY()
 	AVREditorFloatingUI* QuickRadialMenu;
 
 	/** The time since the radial menu was updated */
@@ -172,13 +187,14 @@ protected:
 	// Dragging UI
 	//
 
-	/** Hand that is dragging the UI */
-	int32 DraggingUIHandIndex;
+	/** Interactor that is dragging the UI */
+	class UVREditorInteractor* InteractorDraggingUI;
 
 	/** Offset transform from room-relative transform to the object where we picked it up at */
 	FTransform DraggingUIOffsetTransform;
 
 	/** The current UI that is being dragged */
+	UPROPERTY()
 	class AVREditorDockableWindow* DraggingUI;
 
 	//
@@ -197,15 +213,19 @@ protected:
 	//
 
 	/** Start dragging UI sound */
+	UPROPERTY()
 	class USoundCue* StartDragUISound;
 
 	/** Stop dragging UI sound */
+	UPROPERTY()
 	class USoundCue* StopDragUISound;
 
 	/** Hide UI sound */
+	UPROPERTY()
 	class USoundCue* HideUISound;
 
 	/** Show UI sound */
+	UPROPERTY()
 	class USoundCue* ShowUISound;
 
 	/** If the current dragged dock passed a certain distance if dragged from a hand */
