VR Editor: Fix pointer trace goes through Docking anchor if UI panel is behind #jira UE-29391

diff --git a/Engine/Source/Editor/VREditor/VREditorDockableWindow.cpp b/Engine/Source/Editor/VREditor/VREditorDockableWindow.cpp
--- a/Engine/Source/Editor/VREditor/VREditorDockableWindow.cpp
+++ b/Engine/Source/Editor/VREditor/VREditorDockableWindow.cpp
@@ -64,10 +64,6 @@ AVREditorDockableWindow::AVREditorDockableWindow()
 		SelectionBarMeshComponent->SetupAttachment( RootComponent );
 
 		SelectionBarMeshComponent->SetCollisionEnabled( ECollisionEnabled::QueryOnly );
-		SelectionBarMeshComponent->SetCollisionResponseToAllChannels( ECollisionResponse::ECR_Ignore );
-		SelectionBarMeshComponent->SetCollisionResponseToChannel( COLLISION_GIZMO, ECollisionResponse::ECR_Block );
-		SelectionBarMeshComponent->SetCollisionObjectType( COLLISION_GIZMO );
-
 		SelectionBarMeshComponent->bGenerateOverlapEvents = false;
 		SelectionBarMeshComponent->SetCanEverAffectNavigation( false );
 		SelectionBarMeshComponent->bCastDynamicShadow = false;
diff --git a/Engine/Source/Editor/VREditor/VREditorUISystem.cpp b/Engine/Source/Editor/VREditor/VREditorUISystem.cpp
--- a/Engine/Source/Editor/VREditor/VREditorUISystem.cpp
+++ b/Engine/Source/Editor/VREditor/VREditorUISystem.cpp
@@ -187,7 +187,9 @@ void UVREditorUISystem::OnVRAction( FEditorViewportClient& ViewportClient, UView
 	const FViewportActionKeyInput& Action, bool& bOutIsInputCaptured, bool& bWasHandled )
 {
 	UVREditorMotionControllerInteractor* VREditorInteractor = Cast<UVREditorMotionControllerInteractor>( Interactor );
-	if ( VREditorInteractor )
+	if( VREditorInteractor && 
+		!( VREditorInteractor->GetDraggingMode() == EViewportInteractionDraggingMode::ActorsWithGizmo ||
+		  VREditorInteractor->GetDraggingMode() == EViewportInteractionDraggingMode::Interactable ) )
 	{
 		if ( !bWasHandled )
 		{
