VR Editor: Fix Oculus VR Headset A,B,X, and Y buttons not lighting properly to correct buttons #jira UE-30950

diff --git a/Engine/Source/Editor/VREditor/Public/VREditorMode.h b/Engine/Source/Editor/VREditor/Public/VREditorMode.h
--- a/Engine/Source/Editor/VREditor/Public/VREditorMode.h
+++ b/Engine/Source/Editor/VREditor/Public/VREditorMode.h
@@ -25,6 +25,7 @@ namespace VRActionTypes
 {
 	static const FName Touch( "Touch" );
 	static const FName Modifier( "Modifier" );
+	static const FName Modifier2( "Modifier2" ); //Only used by Oculus
 	static const FName ConfirmRadialSelection( "ConfirmRadialSelection" );
 	static const FName TrackpadPositionX( "TrackpadPositionX" );
 	static const FName TrackpadPositionY( "TrackpadPositionY" );
diff --git a/Engine/Source/Editor/VREditor/VREditorMotionControllerInteractor.cpp b/Engine/Source/Editor/VREditor/VREditorMotionControllerInteractor.cpp
--- a/Engine/Source/Editor/VREditor/VREditorMotionControllerInteractor.cpp
+++ b/Engine/Source/Editor/VREditor/VREditorMotionControllerInteractor.cpp
@@ -114,6 +114,7 @@ void UVREditorMotionControllerInteractor::Init( class UVREditorMode* InVRMode )
 	Super::Init( InVRMode );
 	bHaveMotionController = true;
 
+	const EHMDDeviceType::Type HMDDeviceType = GetVRMode().GetHMDDeviceType();
 	// Setup keys
 	if ( ControllerHandSide == EControllerHand::Left )
 	{
@@ -128,13 +129,14 @@ void UVREditorMotionControllerInteractor::Init( class UVREditorMode* InVRMode )
 		AddKeyAction( EKeys::MotionController_Left_Thumbstick_Y, FViewportActionKeyInput( UVREditorMotionControllerInteractor::TrackpadPositionY ) );
 		AddKeyAction( EKeys::MotionController_Left_Thumbstick, FViewportActionKeyInput( VRActionTypes::ConfirmRadialSelection ) );
 
-		if ( GetVRMode().GetHMDDeviceType() == EHMDDeviceType::DT_SteamVR )
+		if ( HMDDeviceType == EHMDDeviceType::DT_SteamVR )
 		{
 			AddKeyAction( EKeys::MotionController_Left_Shoulder, FViewportActionKeyInput( VRActionTypes::Modifier ) );
 		}
-		else
+		else if( HMDDeviceType == EHMDDeviceType::DT_OculusRift )
 		{
 			AddKeyAction( EKeys::MotionController_Left_FaceButton1, FViewportActionKeyInput( VRActionTypes::Modifier ) );
+			AddKeyAction( EKeys::MotionController_Left_FaceButton2, FViewportActionKeyInput( VRActionTypes::Modifier2 ) );
 		}
 	}
 	else if ( ControllerHandSide == EControllerHand::Right )
@@ -148,13 +150,14 @@ void UVREditorMotionControllerInteractor::Init( class UVREditorMode* InVRMode )
 		AddKeyAction( EKeys::MotionController_Right_Thumbstick_Y, FViewportActionKeyInput( UVREditorMotionControllerInteractor::TrackpadPositionY ) );
 		AddKeyAction( EKeys::MotionController_Right_Thumbstick, FViewportActionKeyInput( VRActionTypes::ConfirmRadialSelection ) );
 
-		if ( GetVRMode().GetHMDDeviceType() == EHMDDeviceType::DT_SteamVR )
+		if ( HMDDeviceType == EHMDDeviceType::DT_SteamVR )
 		{
 			AddKeyAction( EKeys::MotionController_Right_Shoulder, FViewportActionKeyInput( VRActionTypes::Modifier ) );
 		}
-		else
+		else if ( HMDDeviceType == EHMDDeviceType::DT_OculusRift )
 		{
 			AddKeyAction( EKeys::MotionController_Right_FaceButton1, FViewportActionKeyInput( VRActionTypes::Modifier ) );
+			AddKeyAction( EKeys::MotionController_Right_FaceButton2, FViewportActionKeyInput( VRActionTypes::Modifier2 ) );
 		}
 	}
 }
@@ -956,11 +959,17 @@ void UVREditorMotionControllerInteractor::ApplyButtonPressColors( const FViewpor
 	}
 
 	// Modifier
-	if ( ActionType == VRActionTypes::Modifier )
+	if( ActionType == VRActionTypes::Modifier )
 	{
 		static FName StaticModifierParameter( "B4" );
 		SetMotionControllerButtonPressedVisuals( Event, StaticModifierParameter, PressStrength );
 	}
+
+	if( GetVRMode().GetHMDDeviceType() == EHMDDeviceType::DT_OculusRift && ActionType == VRActionTypes::Modifier2 )
+	{
+		static FName StaticModifierParameter( "B5" );
+		SetMotionControllerButtonPressedVisuals( Event, StaticModifierParameter, PressStrength );
+	}
 }
 
 void UVREditorMotionControllerInteractor::SetMotionControllerButtonPressedVisuals( const EInputEvent Event, const FName& ParameterName, const float PressStrength )
