VR Editor: Fix color picker can have inconsistent / odd spawn transforms #jira UE-33770

diff --git a/Engine/Source/Editor/VREditor/VREditorUISystem.cpp b/Engine/Source/Editor/VREditor/VREditorUISystem.cpp
--- a/Engine/Source/Editor/VREditor/VREditorUISystem.cpp
+++ b/Engine/Source/Editor/VREditor/VREditorUISystem.cpp
@@ -66,6 +66,7 @@ namespace VREd
 	static FAutoConsoleVariable MinDockWindowSize( TEXT( "VREd.MinDockWindowSize" ), 40, TEXT( "Minimum size for dockable windows" ) );
 	static FAutoConsoleVariable DockSnapTime( TEXT( "VREd.DockSnapTime" ), 0.1f, TEXT( "Lerp time for snapping dock to hand and back" ) );
 	static FAutoConsoleVariable DockDragSpeed( TEXT( "VREd.DockDragSpeed" ), 0.05f, TEXT( "Dock dragging time delay" ) );
+	static FAutoConsoleVariable ColorPickerDockSpawnOffset( TEXT( "VREd.ColorPickerDockSpawnOffset" ),  3.f, TEXT( "Offset of where the color picker spawns" ) );
 
 	// Tutorial UI commands
 	static FAutoConsoleVariable TutorialUIResolutionX( TEXT( "VREd.TutorialUI.Resolution.X" ), 1920, TEXT( "The X resolution for the tutorial UI panel" ) );
@@ -100,6 +101,7 @@ UVREditorUISystem::UVREditorUISystem( const FObjectInitializer& Initializer ) :
 	InteractorDraggingUI( nullptr ),
 	DraggingUIOffsetTransform( FTransform::Identity ),
 	DraggingUI( nullptr ),
+	ColorPickerUI( nullptr ),
 	bRefocusViewport(false),
 	StartDragUISound( nullptr ),
 	StopDragUISound( nullptr ),
@@ -181,6 +183,8 @@ void UVREditorUISystem::Shutdown()
 	HideUISound = nullptr;
 	ShowUISound = nullptr;
 	VRMode = nullptr;
+	DraggingUI = nullptr;
+	ColorPickerUI = nullptr;
 }
 
 void UVREditorUISystem::OnVRAction( FEditorViewportClient& ViewportClient, UViewportInteractor* Interactor,
@@ -1037,42 +1041,44 @@ void UVREditorUISystem::UpdateRadialMenu( UVREditorInteractor* Interactor )
 
 void UVREditorUISystem::TryToSpawnRadialMenu( UVREditorInteractor* Interactor )
 {
-	UVREditorMotionControllerInteractor* MotionControllerInteractor = Cast<UVREditorMotionControllerInteractor>( Interactor );
-	if( MotionControllerInteractor )
+	if( Interactor != nullptr )
 	{
-		UVREditorInteractor* DockedToInteractor = nullptr;
-		if ( !QuickRadialMenu->bHidden )
+		UVREditorMotionControllerInteractor* MotionControllerInteractor = Cast<UVREditorMotionControllerInteractor>( Interactor );
+		if( MotionControllerInteractor != nullptr )
 		{
-			const EControllerHand ControllerHand = QuickRadialMenu->GetDockedTo() == AVREditorFloatingUI::EDockedTo::LeftHand ? EControllerHand::Left : EControllerHand::Right;
-			DockedToInteractor = GetOwner().GetHandInteractor( ControllerHand );
-		}
+			UVREditorInteractor* DockedToInteractor = nullptr;
+			if( !QuickRadialMenu->bHidden )
+			{
+				const EControllerHand ControllerHand = QuickRadialMenu->GetDockedTo() == AVREditorFloatingUI::EDockedTo::LeftHand ? EControllerHand::Left : EControllerHand::Right;
+				DockedToInteractor = GetOwner().GetHandInteractor( ControllerHand );
+			}
 
-		EViewportInteractionDraggingMode DraggingMode = Interactor->GetDraggingMode();
-
-		CA_SUPPRESS(6011); // warning C6011: Dereferencing NULL pointer 'Interactor'.
-		bool bNeedsSpawn =
-			( QuickRadialMenu->bHidden || DockedToInteractor != Interactor ) &&
-			DraggingMode != EViewportInteractionDraggingMode::ActorsAtLaserImpact &&	// Don't show radial menu if the hand is busy dragging something around
-			DraggingMode != EViewportInteractionDraggingMode::ActorsFreely &&
-			DraggingMode != EViewportInteractionDraggingMode::ActorsWithGizmo &&
-			DraggingMode != EViewportInteractionDraggingMode::World &&
-			DraggingMode != EViewportInteractionDraggingMode::AssistingDrag &&
-			( InteractorDraggingUI == nullptr || InteractorDraggingUI != Interactor ) &&
-			!Interactor->IsHoveringOverUI();	// Don't show radial menu when aiming at a UI  (too much clutter)
-
-		UVREditorRadialMenu* RadialMenu = QuickRadialMenu->GetUserWidget<UVREditorRadialMenu>();
-		// We need to update the trackpad position in the radialmenu before checking if it can be used
-		RadialMenu->Update( Interactor );
-		if ( RadialMenu && RadialMenu->IsInMenuRadius() )
-		{
-			bNeedsSpawn = false;
-		}
+			EViewportInteractionDraggingMode DraggingMode = Interactor->GetDraggingMode();
+
+			bool bNeedsSpawn =
+				( QuickRadialMenu->bHidden || DockedToInteractor != Interactor ) &&
+				DraggingMode != EViewportInteractionDraggingMode::ActorsAtLaserImpact &&	// Don't show radial menu if the hand is busy dragging something around
+				DraggingMode != EViewportInteractionDraggingMode::ActorsFreely &&
+				DraggingMode != EViewportInteractionDraggingMode::ActorsWithGizmo &&
+				DraggingMode != EViewportInteractionDraggingMode::World &&
+				DraggingMode != EViewportInteractionDraggingMode::AssistingDrag &&
+				( InteractorDraggingUI == nullptr || InteractorDraggingUI != Interactor ) &&
+				!Interactor->IsHoveringOverUI();	// Don't show radial menu when aiming at a UI  (too much clutter)
+
+			UVREditorRadialMenu* RadialMenu = QuickRadialMenu->GetUserWidget<UVREditorRadialMenu>();
+			// We need to update the trackpad position in the radialmenu before checking if it can be used
+			RadialMenu->Update( Interactor );
+			if( RadialMenu && RadialMenu->IsInMenuRadius() )
+			{
+				bNeedsSpawn = false;
+			}
 
-		if ( bNeedsSpawn )
-		{
-			const AVREditorFloatingUI::EDockedTo DockedTo = MotionControllerInteractor->GetControllerSide() == EControllerHand::Left ? AVREditorFloatingUI::EDockedTo::LeftHand : AVREditorFloatingUI::EDockedTo::RightHand;
-			QuickRadialMenu->SetDockedTo( DockedTo );
-			QuickRadialMenu->ShowUI( true );
+			if( bNeedsSpawn )
+			{
+				const AVREditorFloatingUI::EDockedTo DockedTo = MotionControllerInteractor->GetControllerSide() == EControllerHand::Left ? AVREditorFloatingUI::EDockedTo::LeftHand : AVREditorFloatingUI::EDockedTo::RightHand;
+				QuickRadialMenu->SetDockedTo( DockedTo );
+				QuickRadialMenu->ShowUI( true );
+			}
 		}
 	}
 }
@@ -1426,55 +1432,82 @@ void UVREditorUISystem::SetDefaultWindowLayout()
 }
 void UVREditorUISystem::CreateVRColorPicker(const TSharedRef<SColorPicker>& ColorPicker)
 {
-	DestroyVRColorPicker();
-
 	// Check that the Color Picker Panel isn't currently showing. Also handles the EEditorUIPanel::ColorPicker being null.
-	if (!IsShowingEditorUIPanel(EEditorUIPanel::ColorPicker))
+	const FIntPoint DefaultResolution(VREd::DefaultEditorUIResolutionX->GetInt(), VREd::DefaultEditorUIResolutionY->GetInt());
+
+	TSharedRef<SWidget> WidgetToDraw =
+		SNew(SDPIScaler)
+		.DPIScale(VREd::EditorUIScale->GetFloat())
+		[
+			ColorPicker
+		]
+	;
+
+	// Only spawn the Color picker dockable window if there is none yet
+	bool bJustSpawned = false;
+	if( ColorPickerUI == nullptr )
 	{
-		const FIntPoint DefaultResolution(VREd::DefaultEditorUIResolutionX->GetInt(), VREd::DefaultEditorUIResolutionY->GetInt());
-
-		TSharedRef<SWidget> WidgetToDraw =
-			SNew(SDPIScaler)
-			.DPIScale(VREd::EditorUIScale->GetFloat())
-			[
-				ColorPicker
-			]
-		;
-
 		const bool bWithSceneComponent = false;
-		AVREditorFloatingUI* ColorPickerUI = GetOwner().SpawnTransientSceneActor< AVREditorDockableWindow >(TEXT("ColorPickerUI"), bWithSceneComponent);
-		ColorPickerUI->ShowUI(true);
-		FloatingUIs.Add(ColorPickerUI);
+		ColorPickerUI = GetOwner().SpawnTransientSceneActor< AVREditorDockableWindow >( TEXT( "ColorPickerUI" ), bWithSceneComponent );
+		FloatingUIs.Add( ColorPickerUI );
+		EditorUIPanels[ (int32)EEditorUIPanel::ColorPicker ] = ColorPickerUI;
+		bJustSpawned = true;
+	}
 
-		EditorUIPanels[(int32)EEditorUIPanel::ColorPicker] = ColorPickerUI;
+	ColorPickerUI->SetSlateWidget( *this, WidgetToDraw, DefaultResolution, VREd::EditorUISize->GetFloat(), AVREditorFloatingUI::EDockedTo::Room );
 
-		// Always spawn based on the location of the menu you are hovering over. To get this information, find the hand hovering over the UI.
-		UVREditorInteractor* VREditorInteractor = GetOwner().GetHandInteractor(EControllerHand::Left);
-		if (!VREditorInteractor->IsHoveringOverUI())
-		{
-			VREditorInteractor = GetOwner().GetHandInteractor(EControllerHand::Right);
+	// Always spawn based on the location of the menu you are hovering over. To get this information, find the hand hovering over the UI.
+	UVREditorInteractor* VREditorInteractor = GetOwner().GetHandInteractor(EControllerHand::Left);
+	if (!VREditorInteractor->IsHoveringOverUI())
+	{
+		VREditorInteractor = GetOwner().GetHandInteractor(EControllerHand::Right);
+	}
 
-		}
+	// Spawn the color picker close to the pressed UI it was opened using UI. Otherwise spawn it in front of the user.
+	UWidgetComponent* HoveredWidgetComponent = VREditorInteractor->GetHoveringOverWidgetComponent();
+	if( HoveredWidgetComponent != nullptr )
+	{
+		FVector LaserPointerStart, LaserPointerEnd;
+		if( VREditorInteractor->GetLaserPointer( LaserPointerStart, LaserPointerEnd ) )
+		{
+			//Calculate new location using the direction of the laser with an offset towards the user
+			FVector LaserPointerDirection = LaserPointerEnd - LaserPointerStart;
+			LaserPointerDirection.Normalize();
+			const FVector HoverOverUILocation = VREditorInteractor->GetHoverLocation();
+			const FVector Offset = LaserPointerDirection * VREd::ColorPickerDockSpawnOffset->GetFloat() * VRMode->GetWorldScaleFactor();
 
-		ColorPickerUI->SetSlateWidget(*this, WidgetToDraw, DefaultResolution, VREd::EditorUISize->GetFloat(), AVREditorFloatingUI::EDockedTo::Room);
+			FTransform NewTransform = FTransform( HoveredWidgetComponent->GetComponentQuat(), HoverOverUILocation - Offset );
+			
+			// Move the color picker if it was already spawned, or instantly set the transform if it was just spawned. We don't want to move the UI from location zero
+			if( bJustSpawned )
+			{
+				ColorPickerUI->SetDockedTo( AVREditorBaseActor::EDockedTo::Room );
+				ColorPickerUI->SetTransform( NewTransform );
+			}
+			else
+			{
+				ColorPickerUI->MoveTo( NewTransform, 0.2f, AVREditorBaseActor::EDockedTo::Room );
+			}
 
+			ColorPickerUI->SetActorTransform( NewTransform );
+			ColorPickerUI->UpdateRelativeRoomTransform();
+		}
+	}
+	else
+	{
 		// Make sure the UIs are centered around the direction your head is looking (yaw only!)
-		const FVector RoomSpaceHeadLocation = GetOwner().GetRoomSpaceHeadTransform().GetLocation() / GetOwner().GetWorldScaleFactor();
-		FRotator RoomSpaceHeadYawRotator = GetOwner().GetRoomSpaceHeadTransform().GetRotation().Rotator();
+		const FVector RoomSpaceHeadLocation = VRMode->GetRoomSpaceHeadTransform().GetLocation() / VRMode->GetWorldScaleFactor();
+		FRotator RoomSpaceHeadYawRotator = VRMode->GetRoomSpaceHeadTransform().GetRotation().Rotator();
 		RoomSpaceHeadYawRotator.Pitch = 0.0f;
 		RoomSpaceHeadYawRotator.Roll = 0.0f;
-
 	
 		FTransform NewTransform = VREd::DefaultColorPickerTransform;
 		NewTransform *= FTransform(RoomSpaceHeadYawRotator.Quaternion(), FVector::ZeroVector);
 		ColorPickerUI->SetLocalRotation(NewTransform.GetRotation().Rotator());
 		ColorPickerUI->SetRelativeOffset(RoomSpaceHeadLocation + NewTransform.GetTranslation());
-
-		const bool bShouldShow = true;
-		const bool bShowOnHand = false;
-		ShowEditorUIPanel(EEditorUIPanel::ColorPicker, VREditorInteractor, bShouldShow, bShowOnHand);
 	}
 
+	ColorPickerUI->ShowUI( true );
 }
 
 void UVREditorUISystem::DestroyVRColorPicker()
diff --git a/Engine/Source/Editor/VREditor/VREditorUISystem.h b/Engine/Source/Editor/VREditor/VREditorUISystem.h
--- a/Engine/Source/Editor/VREditor/VREditorUISystem.h
+++ b/Engine/Source/Editor/VREditor/VREditorUISystem.h
@@ -208,6 +208,10 @@ protected:
 	UPROPERTY()
 	class AVREditorDockableWindow* DraggingUI;
 
+	/** The color picker dockable window */
+	UPROPERTY()
+	class AVREditorDockableWindow* ColorPickerUI;
+
 	//
 	// Tab Manager UI
 	//
