VR Editor: Moved UI related classes to the UI subfolder and VREditorTransformGizmo to Gizmo subfolder. Added subfolders to PrivateIncludePaths in Build.cs.

diff --git a/Engine/Source/Editor/VREditor/Gizmo/VREditorTransformGizmo.cpp b/Engine/Source/Editor/VREditor/Gizmo/VREditorTransformGizmo.cpp
new file mode 100644
--- /dev/null
+++ b/Engine/Source/Editor/VREditor/Gizmo/VREditorTransformGizmo.cpp
@@ -0,0 +1,287 @@
+// Copyright 1998-2015 Epic Games, Inc. All Rights Reserved.
+
+#include "VREditorModule.h"
+#include "VREditorTransformGizmo.h"
+#include "UnitConversion.h"
+#include "Gizmo/VREditorTranslationGizmoHandle.h"
+#include "Gizmo/VREditorRotationGizmoHandle.h"
+#include "Gizmo/VREditorPlaneTranslationGizmoHandle.h"
+#include "VIStretchGizmoHandle.h"
+#include "VIUniformScaleGizmoHandle.h"
+#include "VREditorMode.h"
+
+namespace VREd
+{
+	// @todo vreditor tweak: Tweak out console variables
+	static FAutoConsoleVariable GizmoScale( TEXT( "VREd.GizmoScale" ), 0.80f, TEXT( "How big the gizmo handles should be" ) );
+	static FAutoConsoleVariable GizmoDistanceScaleFactor( TEXT( "VREd.GizmoDistanceScaleFactor" ), 0.002f, TEXT( "How much the gizmo handles should increase in size with distance from the camera, to make it easier to select" ) );
+	static FAutoConsoleVariable MinActorSizeForTransformGizmo( TEXT( "VREd.MinActorSizeForTransformGizmo" ), 50.0f, TEXT( "How big an object must be in scaled world units before we'll start to shrink the gizmo" ) );
+}
+
+
+ATransformGizmo::ATransformGizmo()
+{
+	if (UNLIKELY(IsRunningDedicatedServer()))
+	{
+		return;
+	}
+
+	TranslationGizmoHandleGroup = CreateDefaultSubobject<UVREditorTranslationGizmoHandleGroup>( TEXT( "TranslationHandles" ), true );
+	TranslationGizmoHandleGroup->SetTranslucentGizmoMaterial( TranslucentGizmoMaterial );
+	TranslationGizmoHandleGroup->SetGizmoMaterial( GizmoMaterial );
+	TranslationGizmoHandleGroup->SetupAttachment( SceneComponent );
+	AllHandleGroups.Add( TranslationGizmoHandleGroup );
+
+	PlaneTranslationGizmoHandleGroup = CreateDefaultSubobject<UVREditorPlaneTranslationGizmoHandleGroup>( TEXT( "PlaneTranslationHandles" ), true );
+	PlaneTranslationGizmoHandleGroup->SetTranslucentGizmoMaterial( TranslucentGizmoMaterial );
+	PlaneTranslationGizmoHandleGroup->SetGizmoMaterial( GizmoMaterial );
+	PlaneTranslationGizmoHandleGroup->SetupAttachment( SceneComponent );
+	AllHandleGroups.Add( PlaneTranslationGizmoHandleGroup );
+
+	RotationGizmoHandleGroup = CreateDefaultSubobject<UVREditorRotationGizmoHandleGroup>( TEXT( "RotationHandles" ), true );
+	RotationGizmoHandleGroup->SetTranslucentGizmoMaterial( TranslucentGizmoMaterial );
+	RotationGizmoHandleGroup->SetGizmoMaterial( GizmoMaterial );
+	RotationGizmoHandleGroup->SetupAttachment( SceneComponent );
+	AllHandleGroups.Add( RotationGizmoHandleGroup );
+
+	StretchGizmoHandleGroup = CreateDefaultSubobject<UStretchGizmoHandleGroup>( TEXT( "StretchHandles" ), true );
+	StretchGizmoHandleGroup->SetTranslucentGizmoMaterial( TranslucentGizmoMaterial );
+	StretchGizmoHandleGroup->SetGizmoMaterial( GizmoMaterial );
+	StretchGizmoHandleGroup->SetupAttachment( SceneComponent );
+	AllHandleGroups.Add( StretchGizmoHandleGroup );
+
+	UniformScaleGizmoHandleGroup = CreateDefaultSubobject<UUniformScaleGizmoHandleGroup>( TEXT( "UniformScaleHandles" ), true );
+	UniformScaleGizmoHandleGroup->SetTranslucentGizmoMaterial( TranslucentGizmoMaterial );
+	UniformScaleGizmoHandleGroup->SetGizmoMaterial( GizmoMaterial );
+	UniformScaleGizmoHandleGroup->SetUsePivotPointAsLocation( false );
+	UniformScaleGizmoHandleGroup->SetupAttachment( SceneComponent );
+	AllHandleGroups.Add( UniformScaleGizmoHandleGroup );
+
+	const bool bAllowGizmoLighting = false;	// @todo vreditor: Not sure if we want this for gizmos or not yet.  Needs feedback.  Also they're translucent right now.
+
+	// Setup measurements
+	{
+		UMaterialInterface* GizmoTextMaterial = nullptr;
+		{
+			static ConstructorHelpers::FObjectFinder<UMaterial> ObjectFinder( TEXT( "/Engine/VREditor/Fonts/VRTextMaterial" ) );
+			GizmoTextMaterial = ObjectFinder.Object;
+			check( GizmoTextMaterial != nullptr );
+		}
+
+		UFont* GizmoTextFont = nullptr;
+		{
+			static ConstructorHelpers::FObjectFinder<UFont> ObjectFinder( TEXT( "/Engine/VREditor/Fonts/VRText_RobotoLarge" ) );
+			GizmoTextFont = ObjectFinder.Object;
+			check( GizmoTextFont != nullptr );
+		}
+
+		for( int32 AxisIndex = 0; AxisIndex < 3; ++AxisIndex )
+		{
+			const FString AxisName = AxisIndex == 0 ? TEXT( "X" ) : ( AxisIndex == 1 ? TEXT( "Y" ) : TEXT( "Z" ) );
+
+			FString ComponentName = AxisName + TEXT( "Measurement" );
+
+			UTextRenderComponent* MeasurementText = CreateDefaultSubobject<UTextRenderComponent>( *ComponentName );
+			check( MeasurementText != nullptr );
+
+			MeasurementText->SetMobility( EComponentMobility::Movable );
+			MeasurementText->SetupAttachment( SceneComponent );
+
+			MeasurementText->SetCollisionProfileName( UCollisionProfile::NoCollision_ProfileName );
+
+			MeasurementText->bGenerateOverlapEvents = false;
+			MeasurementText->SetCanEverAffectNavigation( false );
+			MeasurementText->bCastDynamicShadow = bAllowGizmoLighting;
+			MeasurementText->bCastStaticShadow = false;
+			MeasurementText->bAffectDistanceFieldLighting = bAllowGizmoLighting;
+			MeasurementText->bAffectDynamicIndirectLighting = bAllowGizmoLighting;
+
+			// Use a custom font.  The text will be visible up close.
+			MeasurementText->SetFont( GizmoTextFont );
+
+			MeasurementText->SetWorldSize( 8.0f );	// @todo vreditor tweak: Size of text in world units
+
+			// Assign our custom text rendering material.  It should be Unlit and have Depth Test disabled
+			// so that it draws on top of overlapping geometry.
+			MeasurementText->SetTextMaterial( GizmoTextMaterial );
+
+			// Color the text based on the axis it represents
+			switch( AxisIndex )
+			{
+			case 0:
+				MeasurementText->SetTextRenderColor( FLinearColor( 0.4f, 0.075f, 0.075f ).ToFColor( false ) );
+				break;
+			case 1:
+				MeasurementText->SetTextRenderColor( FLinearColor( 0.075f, 0.4f, 0.075f ).ToFColor( false ) );
+				break;
+			case 2:
+				MeasurementText->SetTextRenderColor( FLinearColor( 0.075f, 0.075f, 0.4f ).ToFColor( false ) );
+				break;
+			}
+
+			// Center the text horizontally
+			MeasurementText->SetHorizontalAlignment( EHTA_Center );
+
+			FTransformGizmoMeasurement& Measurement = Measurements[ AxisIndex ];
+			Measurement.MeasurementText = MeasurementText;
+		}
+	}
+
+	// There may already be some objects selected as we switch into VR mode, so we'll pretend that just happened so
+	// that we can make sure all transitions complete properly
+	OnNewObjectsSelected();
+}
+
+
+void ATransformGizmo::UpdateGizmo( const EGizmoHandleTypes GizmoType, const ECoordSystem GizmoCoordinateSpace, const FTransform& LocalToWorld, const FBox& LocalBounds, const FVector ViewLocation, bool bAllHandlesVisible, 
+	UActorComponent* DraggingHandle, const TArray< UActorComponent* >& HoveringOverHandles, const float GizmoHoverScale, const float GizmoHoverAnimationDuration )
+{
+	Super::UpdateGizmo( GizmoType, GizmoCoordinateSpace, LocalToWorld, LocalBounds, ViewLocation, bAllHandlesVisible,
+		DraggingHandle, HoveringOverHandles, GizmoHoverScale, GizmoHoverAnimationDuration );
+
+	const float WorldScaleFactor = GetWorld()->GetWorldSettings()->WorldToMeters / 100.0f;
+
+	// Position the gizmo at the location of the first selected actor
+	const bool bSweep = false;
+	this->SetActorTransform( LocalToWorld, bSweep );
+
+	const FBox WorldSpaceBounds = LocalBounds.TransformBy( LocalToWorld );
+	float ScaleCompensationForTinyGizmo = 1.0f;
+	{
+		const float AvgBounds = FMath::Lerp( WorldSpaceBounds.GetSize().GetAbsMin(), WorldSpaceBounds.GetSize().GetAbsMax(), 0.5f );
+		const float GizmoSizeInRoomSpace = ( AvgBounds / WorldScaleFactor );
+		if( GizmoSizeInRoomSpace < VREd::MinActorSizeForTransformGizmo->GetFloat() )
+		{
+			ScaleCompensationForTinyGizmo = GizmoSizeInRoomSpace / VREd::MinActorSizeForTransformGizmo->GetFloat();
+		}
+	}
+
+	// Increase scale with distance, to make gizmo handles easier to click on
+	// @todo vreditor: Should probably be a curve, not linear
+	// @todo vreditor: Should take FOV into account (especially in non-stereo/HMD mode)
+	const float WorldSpaceDistanceToGizmoBounds = FMath::Sqrt( WorldSpaceBounds.ComputeSquaredDistanceToPoint( ViewLocation ) );
+	const float GizmoScaleUpClose = VREd::GizmoScale->GetFloat();
+	const float GizmoScale( ( GizmoScaleUpClose + ( WorldSpaceDistanceToGizmoBounds / WorldScaleFactor ) * VREd::GizmoDistanceScaleFactor->GetFloat() ) * ScaleCompensationForTinyGizmo * WorldScaleFactor );
+
+	// Update animation
+	float AnimationAlpha = GetAnimationAlpha();
+
+	// Update all the handles
+	bool bIsHoveringOrDraggingScaleGizmo = false;
+	for ( UGizmoHandleGroup* HandleGroup : AllHandleGroups )
+	{
+		if ( HandleGroup != nullptr )
+		{
+			bool bIsHoveringOrDraggingThisHandleGroup = false;
+			HandleGroup->UpdateGizmoHandleGroup( LocalToWorld, LocalBounds, ViewLocation, bAllHandlesVisible, DraggingHandle, 
+				HoveringOverHandles, AnimationAlpha, GizmoScale, GizmoHoverScale, GizmoHoverAnimationDuration, /* Out */ bIsHoveringOrDraggingThisHandleGroup );
+			
+			if( HandleGroup->GetHandleType() == EGizmoHandleTypes::Scale && bIsHoveringOrDraggingThisHandleGroup )
+			{
+				bIsHoveringOrDraggingScaleGizmo = true;
+			}
+		}
+	}
+
+	// Place measurements
+	{
+		for( int32 AxisIndex = 0; AxisIndex < 3; ++AxisIndex )
+		{
+			FTransformGizmoMeasurement& Measurement = Measurements[ AxisIndex ];
+
+			// Update measurement text
+			{
+				const float LocalSpaceLengthOfBoundsAlongAxis = LocalBounds.GetSize()[ AxisIndex ];
+
+				FVector LocalSpaceAxisLengthVector = FVector::ZeroVector;
+				LocalSpaceAxisLengthVector[ AxisIndex ] = LocalSpaceLengthOfBoundsAlongAxis;
+
+				// Transform to world space
+				const FVector WorldSpaceAxisLengthVector = LocalToWorld.TransformVector( LocalSpaceAxisLengthVector );
+				const float WorldSpaceLengthOfBoundsAlongAxis = WorldSpaceAxisLengthVector.Size();
+
+				FText BestSizeString;
+				{
+					const FNumericUnit<float> BestUnit = FUnitConversion::QuantizeUnitsToBestFit( WorldSpaceLengthOfBoundsAlongAxis, EUnit::Centimeters );
+
+					FNumberFormattingOptions NumberFormattingOptions;
+					NumberFormattingOptions.MaximumFractionalDigits = 2;	// @todo sizemap: We could make the number of digits customizable in the UI
+					NumberFormattingOptions.MinimumFractionalDigits = 0;
+					NumberFormattingOptions.MinimumIntegralDigits = 1;
+
+					BestSizeString =
+						FText::FromString( FString::Printf( TEXT( "%s %s" ),
+							*FText::AsNumber( BestUnit.Value, &NumberFormattingOptions ).ToString(),
+							FUnitConversion::GetUnitDisplayString( BestUnit.Units ) ) );
+				}
+
+				Measurement.MeasurementText->SetText( BestSizeString );
+			}
+
+
+			// Position the measurement in the center of the axis along the top of the bounds
+			{
+				// Figure which bounding box edge on this axis is closest to the viewer.  There are four edges to choose from!
+				FVector GizmoSpaceClosestEdge0, GizmoSpaceClosestEdge1;
+				{
+					float ClosestSquaredDistance = TNumericLimits<float>::Max();
+					for( int32 EdgeIndex = 0; EdgeIndex < 4; ++EdgeIndex )
+					{
+						FVector GizmoSpaceEdge0, GizmoSpaceEdge1;
+						GetBoundingBoxEdge( LocalBounds, AxisIndex, EdgeIndex, /* Out */ GizmoSpaceEdge0, /* Out */ GizmoSpaceEdge1 );
+
+						const FVector WorldSpaceEdge0 = LocalToWorld.TransformPosition( GizmoSpaceEdge0 );
+						const FVector WorldSpaceEdge1 = LocalToWorld.TransformPosition( GizmoSpaceEdge1 );
+
+						const float SquaredDistance = FMath::PointDistToSegmentSquared( ViewLocation, WorldSpaceEdge0, WorldSpaceEdge1 );
+						if( SquaredDistance < ClosestSquaredDistance )
+						{
+							ClosestSquaredDistance = SquaredDistance;
+
+							GizmoSpaceClosestEdge0 = GizmoSpaceEdge0;
+							GizmoSpaceClosestEdge1 = GizmoSpaceEdge1;
+						}
+					}
+
+				}
+
+				const FVector EdgeVector = GizmoSpaceClosestEdge1 - GizmoSpaceClosestEdge0;
+				const FVector GizmoSpaceEdgeCenter = GizmoSpaceClosestEdge0 + EdgeVector * 0.5f;
+
+				const float UpOffsetAmount = GizmoScale * 5.0f;	// Push the text up a bit so that it doesn't overlap so much with our handles
+				const FVector GizmoSpaceUpVector = LocalToWorld.InverseTransformVectorNoScale( FVector::UpVector );
+
+				Measurement.MeasurementText->SetRelativeLocation( GizmoSpaceEdgeCenter + GizmoSpaceUpVector * UpOffsetAmount );
+
+				const int32 TextFacingAxisIndex = AxisIndex == 0 ? 1 : ( AxisIndex == 1 ? 0 : 0 );
+				FVector TextFacingDirection = UGizmoHandleGroup::GetAxisVector( TextFacingAxisIndex, ETransformGizmoHandleDirection::Positive );
+
+				// Make sure to face the camera
+				const FVector WorldSpaceEdgeCenterToViewDirection = ( LocalToWorld.TransformPosition( GizmoSpaceEdgeCenter ) - ViewLocation ).GetSafeNormal();
+				if( FVector::DotProduct( LocalToWorld.TransformVectorNoScale( TextFacingDirection ), WorldSpaceEdgeCenterToViewDirection ) > 0.0f )
+				{
+					TextFacingDirection *= -1.0f;
+				}
+
+				const FQuat GizmoSpaceOrientation = TextFacingDirection.ToOrientationQuat();
+				Measurement.MeasurementText->SetRelativeRotation( GizmoSpaceOrientation );
+
+				Measurement.MeasurementText->SetRelativeScale3D( FVector( GizmoScale ) );
+			}
+		}
+	}
+
+	// Update visibility
+	{
+		const bool bShouldShowMeasurements =
+			GetShowMeasurementText() ||
+			bIsHoveringOrDraggingScaleGizmo;
+
+		for( int32 AxisIndex = 0; AxisIndex < 3; ++AxisIndex )
+		{
+			FTransformGizmoMeasurement& Measurement = Measurements[ AxisIndex ];
+
+			Measurement.MeasurementText->SetVisibility( bShouldShowMeasurements );
+		}
+	}
+}
diff --git a/Engine/Source/Editor/VREditor/Gizmo/VREditorTransformGizmo.h b/Engine/Source/Editor/VREditor/Gizmo/VREditorTransformGizmo.h
new file mode 100644
--- /dev/null
+++ b/Engine/Source/Editor/VREditor/Gizmo/VREditorTransformGizmo.h
@@ -0,0 +1,52 @@
+// Copyright 1998-2015 Epic Games, Inc. All Rights Reserved.
+
+#pragma once
+
+#include "VIBaseTransformGizmo.h"
+#include "VREditorTransformGizmo.generated.h"
+
+/**
+ * A transform gizmo that allows you to interact with selected objects by moving, scaling and rotating.
+ */
+UCLASS()
+class ATransformGizmo : public ABaseTransformGizmo
+{
+	GENERATED_BODY()
+
+public:
+
+	/** Default constructor that sets up CDO properties */
+	ATransformGizmo();
+
+	/** Called by the world interaction system after we've been spawned into the world, to allow
+	    us to create components and set everything up nicely for the selected objects that we'll be
+		used to manipulate */
+	virtual void UpdateGizmo( const EGizmoHandleTypes GizmoType, const ECoordSystem GizmoCoordinateSpace, const FTransform& LocalToWorld, const FBox& LocalBounds, const FVector ViewLocation, bool bAllHandlesVisible, 
+		class UActorComponent* DraggingHandle, const TArray< UActorComponent* >& HoveringOverHandles, const float GizmoHoverScale, const float GizmoHoverAnimationDuration ) override;
+
+private:
+
+	/** Measurements, one for each axis */
+	UPROPERTY()
+	FTransformGizmoMeasurement Measurements[ 3 ];
+
+	/** Bounding box translation handles */
+	UPROPERTY()
+	class UVREditorTranslationGizmoHandleGroup* TranslationGizmoHandleGroup;
+
+	/** Bounding box plane translation handles */
+	UPROPERTY()
+	class UVREditorPlaneTranslationGizmoHandleGroup* PlaneTranslationGizmoHandleGroup;
+
+	/** Bounding box rotation handles */
+	UPROPERTY()
+	class UVREditorRotationGizmoHandleGroup* RotationGizmoHandleGroup;
+
+	/** Bounding box stretch handles */
+	UPROPERTY()
+	class UStretchGizmoHandleGroup* StretchGizmoHandleGroup;
+
+	/** Bounding box uniform scale handle */
+	UPROPERTY()
+	class UUniformScaleGizmoHandleGroup* UniformScaleGizmoHandleGroup;
+};
\ No newline at end of file
diff --git a/Engine/Source/Editor/VREditor/UI/VREditorBaseUserWidget.cpp b/Engine/Source/Editor/VREditor/UI/VREditorBaseUserWidget.cpp
new file mode 100644
--- /dev/null
+++ b/Engine/Source/Editor/VREditor/UI/VREditorBaseUserWidget.cpp
@@ -0,0 +1,18 @@
+// Copyright 1998-2015 Epic Games, Inc. All Rights Reserved.
+
+#include "VREditorModule.h"
+#include "VREditorBaseUserWidget.h"
+#include "VREditorFloatingUI.h"
+
+UVREditorBaseUserWidget::UVREditorBaseUserWidget( const FObjectInitializer& ObjectInitializer )
+	: Super( ObjectInitializer ),
+	  Owner( nullptr )
+{
+}
+
+
+void UVREditorBaseUserWidget::SetOwner( class AVREditorFloatingUI* NewOwner )
+{
+	Owner = NewOwner;
+}
+
diff --git a/Engine/Source/Editor/VREditor/UI/VREditorBaseUserWidget.h b/Engine/Source/Editor/VREditor/UI/VREditorBaseUserWidget.h
new file mode 100644
--- /dev/null
+++ b/Engine/Source/Editor/VREditor/UI/VREditorBaseUserWidget.h
@@ -0,0 +1,45 @@
+// Copyright 1998-2015 Epic Games, Inc. All Rights Reserved.
+
+#pragma once
+
+
+#include "Blueprint/UserWidget.h"
+#include "VREditorBaseUserWidget.generated.h"
+
+
+/**
+ * Base class for all of the VR widgets
+ */
+UCLASS( Abstract )
+class UVREditorBaseUserWidget : public UUserWidget
+{
+	GENERATED_BODY()
+
+public:
+	
+	/** Default constructor that sets up CDO properties */
+	UVREditorBaseUserWidget( const FObjectInitializer& ObjectInitializer );
+
+	/** Gets the owner of this widget */
+	class AVREditorFloatingUI* GetOwner()
+	{
+		return Owner.Get();
+	}
+
+	/** Gets the owner of this widget (const) */
+	const class AVREditorFloatingUI* GetOwner() const
+	{
+		return Owner.Get();
+	}
+
+	/** Sets the UI system that owns this widget */
+	void SetOwner( class AVREditorFloatingUI* NewOwner );
+
+
+protected:
+
+	/** The UI system that owns this widget */
+	UPROPERTY()
+	TWeakObjectPtr<class AVREditorFloatingUI> Owner;
+
+};
\ No newline at end of file
diff --git a/Engine/Source/Editor/VREditor/UI/VREditorDockableWindow.cpp b/Engine/Source/Editor/VREditor/UI/VREditorDockableWindow.cpp
new file mode 100644
--- /dev/null
+++ b/Engine/Source/Editor/VREditor/UI/VREditorDockableWindow.cpp
@@ -0,0 +1,443 @@
+// Copyright 1998-2015 Epic Games, Inc. All Rights Reserved.
+
+#include "VREditorModule.h"
+#include "VREditorDockableWindow.h"
+#include "VREditorUISystem.h"
+#include "VREditorMode.h"
+#include "VREditorBaseUserWidget.h"
+#include "WidgetComponent.h"
+#include "VREditorWidgetComponent.h"
+#include "ViewportWorldInteraction.h"
+#include "VREditorInteractor.h"
+
+namespace VREd
+{
+	static FAutoConsoleVariable DockWindowThickness( TEXT( "VREd.DockWindowTickness" ), 1.0f, TEXT( "How thick the window is" ) );
+	static FAutoConsoleVariable DockUISelectionBarVerticalOffset( TEXT( "VREd.DockUISelectionBarVerticalOffset" ), 2.0f, TEXT( "Z Distance between the selectionbar and the UI" ) );
+	static FAutoConsoleVariable DockUICloseButtonOffsetFromCorner( TEXT( "VREd.DockUICloseButtonOffsetFromCorner" ), 1.5f, TEXT( "How far away from the corner (along each 2D axis) of the UI the close button draws" ) );
+	static FAutoConsoleVariable DockUIFadeAnimationDuration( TEXT( "VREd.DockUIFadeAnimationDuration" ), 0.15f, TEXT( "How quick the fade animation should complete in" ) );
+	static FAutoConsoleVariable DockUIHoverScale( TEXT( "VREd.DockUIHoverScale" ), 1.15f, TEXT( "How big the selection bar gets when you hover over it" ) );
+	static FAutoConsoleVariable DockUIHoverAnimationDuration( TEXT( "VREd.DockUIHoverAnimationDuration" ), 0.15f, TEXT( "How quick the hover animation should complete in" ) );
+}
+
+
+AVREditorDockableWindow::AVREditorDockableWindow()
+	: Super(),
+	DockSelectDistance( 0.0f )
+{
+	{
+		UStaticMesh* WindowMesh = nullptr;
+		{
+			static ConstructorHelpers::FObjectFinder<UStaticMesh> ObjectFinder( TEXT( "/Engine/VREditor/UI/SM_ContentWindow_01" ) );
+			WindowMesh = ObjectFinder.Object;
+			check( WindowMesh != nullptr );
+		}
+
+		WindowMeshComponent = CreateDefaultSubobject<UStaticMeshComponent>( TEXT( "WindowMesh" ) );
+		WindowMeshComponent->SetStaticMesh( WindowMesh );
+		WindowMeshComponent->SetMobility( EComponentMobility::Movable );
+		WindowMeshComponent->SetupAttachment( RootComponent );
+
+		WindowMeshComponent->SetCollisionEnabled( ECollisionEnabled::QueryOnly );
+		WindowMeshComponent->SetCollisionResponseToAllChannels( ECollisionResponse::ECR_Ignore );
+		WindowMeshComponent->SetCollisionResponseToChannel( COLLISION_GIZMO, ECollisionResponse::ECR_Block );
+		WindowMeshComponent->SetCollisionObjectType( COLLISION_GIZMO );
+
+		WindowMeshComponent->bGenerateOverlapEvents = false;
+		WindowMeshComponent->SetCanEverAffectNavigation( false );
+		WindowMeshComponent->bCastDynamicShadow = false;
+		WindowMeshComponent->bCastStaticShadow = false;
+		WindowMeshComponent->bAffectDistanceFieldLighting = false;
+	}
+
+	{
+		UStaticMesh* SelectionMesh = nullptr;
+		{
+			static ConstructorHelpers::FObjectFinder<UStaticMesh> ObjectFinder( TEXT( "/Engine/VREditor/UI/SelectionBarMesh" ) );
+			SelectionMesh = ObjectFinder.Object;
+			check( SelectionMesh != nullptr );
+		}
+
+		SelectionBarMeshComponent = CreateDefaultSubobject<UStaticMeshComponent>( TEXT( "SelectionBarMesh" ) );
+		SelectionBarMeshComponent->SetStaticMesh( SelectionMesh );
+		SelectionBarMeshComponent->SetMobility( EComponentMobility::Movable );
+		SelectionBarMeshComponent->SetupAttachment( RootComponent );
+
+		SelectionBarMeshComponent->SetCollisionEnabled( ECollisionEnabled::QueryOnly );
+		SelectionBarMeshComponent->bGenerateOverlapEvents = false;
+		SelectionBarMeshComponent->SetCanEverAffectNavigation( false );
+		SelectionBarMeshComponent->bCastDynamicShadow = false;
+		SelectionBarMeshComponent->bCastStaticShadow = false;
+		SelectionBarMeshComponent->bAffectDistanceFieldLighting = false;
+
+		UMaterial* HoverMaterial = nullptr;
+		{
+			static ConstructorHelpers::FObjectFinder<UMaterial> ObjectFinder( TEXT( "/Engine/VREditor/TransformGizmo/TransformGizmoMaterial" ) );
+			HoverMaterial = ObjectFinder.Object;
+			check( HoverMaterial != nullptr );
+		}
+		SelectionBarMID = UMaterialInstanceDynamic::Create( HoverMaterial, GetTransientPackage() );
+		check( SelectionBarMID != nullptr );
+		SelectionBarMeshComponent->SetMaterial( 0, SelectionBarMID );
+
+		UMaterial* TranslucentHoverMaterial = nullptr;
+		{
+			static ConstructorHelpers::FObjectFinder<UMaterial> ObjectFinder( TEXT( "/Engine/VREditor/TransformGizmo/TranslucentTransformGizmoMaterial" ) );
+			TranslucentHoverMaterial = ObjectFinder.Object;
+			check( TranslucentHoverMaterial != nullptr );
+		}
+		SelectionBarTranslucentMID = UMaterialInstanceDynamic::Create( TranslucentHoverMaterial, GetTransientPackage() );
+		check( SelectionBarTranslucentMID != nullptr );
+		SelectionBarMeshComponent->SetMaterial( 1, SelectionBarTranslucentMID );
+	}
+
+	{
+		UStaticMesh* CloseButtonMesh = nullptr;
+		{
+			static ConstructorHelpers::FObjectFinder<UStaticMesh> ObjectFinder( TEXT( "/Engine/VREditor/UI/CloseButtonMesh" ) );
+			CloseButtonMesh = ObjectFinder.Object;
+			check( CloseButtonMesh != nullptr );
+		}
+
+		CloseButtonMeshComponent= CreateDefaultSubobject<UStaticMeshComponent>( TEXT( "CloseButtonMesh" ) );
+		CloseButtonMeshComponent->SetStaticMesh( CloseButtonMesh );
+		CloseButtonMeshComponent->SetMobility( EComponentMobility::Movable );
+		CloseButtonMeshComponent->SetupAttachment( RootComponent );
+
+		CloseButtonMeshComponent->SetCollisionEnabled( ECollisionEnabled::QueryOnly );
+		CloseButtonMeshComponent->SetCollisionResponseToAllChannels( ECollisionResponse::ECR_Ignore );
+		CloseButtonMeshComponent->SetCollisionResponseToChannel( COLLISION_GIZMO, ECollisionResponse::ECR_Block );
+		CloseButtonMeshComponent->SetCollisionObjectType( COLLISION_GIZMO );
+
+		CloseButtonMeshComponent->bGenerateOverlapEvents = false;
+		CloseButtonMeshComponent->SetCanEverAffectNavigation( false );
+		CloseButtonMeshComponent->bCastDynamicShadow = false;
+		CloseButtonMeshComponent->bCastStaticShadow = false;
+		CloseButtonMeshComponent->bAffectDistanceFieldLighting = false;
+
+		UMaterial* HoverMaterial = nullptr;
+		{
+			static ConstructorHelpers::FObjectFinder<UMaterial> ObjectFinder( TEXT( "/Engine/VREditor/TransformGizmo/TransformGizmoMaterial" ) );
+			HoverMaterial = ObjectFinder.Object;
+			check( HoverMaterial != nullptr );
+		}
+		CloseButtonMID = UMaterialInstanceDynamic::Create( HoverMaterial, GetTransientPackage() );
+		check( CloseButtonMID != nullptr );
+		CloseButtonMeshComponent->SetMaterial( 0, CloseButtonMID );
+
+		UMaterial* TranslucentHoverMaterial = nullptr;
+		{
+			static ConstructorHelpers::FObjectFinder<UMaterial> ObjectFinder( TEXT( "/Engine/VREditor/TransformGizmo/TranslucentTransformGizmoMaterial" ) );
+			TranslucentHoverMaterial = ObjectFinder.Object;
+			check( TranslucentHoverMaterial != nullptr );
+		}
+		CloseButtonTranslucentMID = UMaterialInstanceDynamic::Create( TranslucentHoverMaterial, GetTransientPackage() );
+		check( CloseButtonTranslucentMID != nullptr );
+		CloseButtonMeshComponent->SetMaterial( 1, CloseButtonTranslucentMID );
+	}
+
+	// The selection bar and close button will not be initially visible.  They'll appear when the user aims
+	// their laser toward the UI
+	SelectionBarMeshComponent->SetVisibility( false );
+	CloseButtonMeshComponent->SetVisibility( false );
+
+	// Create the drag operation
+	DragOperationComponent = CreateDefaultSubobject<UViewportDragOperationComponent>( TEXT( "DragOperation" ) );
+	DragOperationComponent->SetDragOperationClass( UDockableWindowDragOperation::StaticClass() );
+}
+
+
+AVREditorDockableWindow::~AVREditorDockableWindow()
+{
+	WindowMeshComponent = nullptr;
+	SelectionBarMeshComponent = nullptr;
+	CloseButtonMeshComponent = nullptr;
+	SelectionBarMID = nullptr;
+	SelectionBarTranslucentMID = nullptr;
+	CloseButtonMID = nullptr;
+	CloseButtonTranslucentMID = nullptr;
+	DragOperationComponent = nullptr;
+}
+
+void AVREditorDockableWindow::SetupWidgetComponent()
+{
+	Super::SetupWidgetComponent();
+
+	SetSelectionBarColor( GetOwner().GetOwner().GetColor( UVREditorMode::EColors::UISelectionBarColor ) );
+	SetCloseButtonColor( GetOwner().GetOwner().GetColor( UVREditorMode::EColors::UICloseButtonColor ) );
+}
+
+
+void AVREditorDockableWindow::TickManually( float DeltaTime )
+{
+	Super::TickManually( DeltaTime );
+
+	if( WidgetComponent->IsVisible() )
+	{
+		const float WorldScaleFactor = GetOwner().GetOwner().GetWorldScaleFactor();
+		const FVector AnimatedScale = CalculateAnimatedScale();
+
+
+		// Update whether the user is aiming toward us or not
+		bIsLaserAimingTowardUI = false;
+		{
+			const FTransform UICapsuleTransform = this->GetActorTransform();
+
+			const FVector UICapsuleStart = FVector( 0.0f, 0.0f, -GetSize().Y * 0.4f ) * WorldScaleFactor * AnimatedScale;
+			const FVector UICapsuleEnd = FVector( 0.0f, 0.0f, GetSize().Y * 0.4f ) * WorldScaleFactor * AnimatedScale;
+			const float UICapsuleLocalRadius = GetSize().X * 0.5f * WorldScaleFactor * AnimatedScale.Y;
+			const float MinDistanceToUICapsule = 10.0f * WorldScaleFactor * AnimatedScale.Y;	// @todo vreditor tweak
+			const FVector UIForwardVector = FVector::ForwardVector;
+			const float MinDotForAimingAtOtherHand = -1.1f;	// @todo vreditor tweak
+
+			for( UViewportInteractor* Interactor : GetOwner().GetOwner().GetWorldInteraction().GetInteractors() )
+			{
+				if( GetOwner().GetOwner().IsHandAimingTowardsCapsule( Interactor, UICapsuleTransform, UICapsuleStart, UICapsuleEnd, UICapsuleLocalRadius, MinDistanceToUICapsule, UIForwardVector, MinDotForAimingAtOtherHand ) )
+				{
+					bIsLaserAimingTowardUI = true;
+				}
+			}
+
+			if( bIsLaserAimingTowardUI )
+			{
+				AimingAtMeFadeAlpha += DeltaTime / VREd::DockUIFadeAnimationDuration->GetFloat();
+			}
+			else
+			{
+				AimingAtMeFadeAlpha -= DeltaTime / VREd::DockUIFadeAnimationDuration->GetFloat();;
+			}
+			AimingAtMeFadeAlpha = FMath::Clamp( AimingAtMeFadeAlpha, 0.0f, 1.0f );
+		}
+
+		const float AnimationOvershootAmount = 1.0f;	// @todo vreditor tweak
+		float EasedAimingAtMeFadeAlpha = UVREditorMode::OvershootEaseOut( AimingAtMeFadeAlpha, AnimationOvershootAmount );
+
+		// Only show our extra buttons and controls if the user is roughly aiming toward us.  This just reduces clutter.
+		SelectionBarMeshComponent->SetVisibility( EasedAimingAtMeFadeAlpha > KINDA_SMALL_NUMBER ? true : false );
+		CloseButtonMeshComponent->SetVisibility( EasedAimingAtMeFadeAlpha > KINDA_SMALL_NUMBER ? true : false );
+
+		EasedAimingAtMeFadeAlpha = FMath::Max( 0.001f, EasedAimingAtMeFadeAlpha );
+
+		// Update the window border mesh
+		{
+			const float WindowMeshSize = 100.0f;	// Size of imported mesh, we need to inverse compensate for
+
+			const FVector WindowMeshScale = FVector(
+				VREd::DockWindowThickness->GetFloat(),
+				GetSize().X / WindowMeshSize,
+				GetSize().Y / WindowMeshSize ) * AnimatedScale * WorldScaleFactor;
+			WindowMeshComponent->SetRelativeScale3D( WindowMeshScale );			
+		}
+
+		// Update the selection bar
+		{
+			if( bIsHoveringOverSelectionBar )
+			{
+				SelectionBarHoverAlpha += DeltaTime / VREd::DockUIHoverAnimationDuration->GetFloat();
+			}
+			else
+			{
+				SelectionBarHoverAlpha -= DeltaTime / VREd::DockUIHoverAnimationDuration->GetFloat();;
+			}
+			SelectionBarHoverAlpha = FMath::Clamp( SelectionBarHoverAlpha, 0.0f, 1.0f );
+
+			// How big the selection bar should be
+			const FVector SelectionBarSize( 1.0f, GetSize().X * 0.8f, 6.0f );	// @todo vreditor tweak
+			FVector SelectionBarScale = SelectionBarSize * AnimatedScale * WorldScaleFactor;
+			SelectionBarScale *= FMath::Lerp( 1.0f, VREd::DockUIHoverScale->GetFloat(), SelectionBarHoverAlpha );
+
+			// Scale vertically based on our fade alpha
+			SelectionBarScale.Z *= EasedAimingAtMeFadeAlpha;
+
+			SelectionBarMeshComponent->SetRelativeScale3D( SelectionBarScale );
+			const FVector SelectionBarRelativeLocation = FVector(
+				0.0f,
+				0.0f,
+				-( GetSize().Y * 0.5f + SelectionBarSize.Z * 0.5f + VREd::DockUISelectionBarVerticalOffset->GetFloat() ) ) * AnimatedScale * WorldScaleFactor;
+			SelectionBarMeshComponent->SetRelativeLocation( SelectionBarRelativeLocation );
+
+			SetSelectionBarColor( GetOwner().GetOwner().GetColor( bIsHoveringOverSelectionBar ? UVREditorMode::EColors::UISelectionBarHoverColor : UVREditorMode::EColors::UISelectionBarColor ) );
+		}
+
+		// Update the close button
+		{
+			if( bIsHoveringOverCloseButton )
+			{
+				CloseButtonHoverAlpha += DeltaTime / VREd::DockUIHoverAnimationDuration->GetFloat();
+			}
+			else
+			{
+				CloseButtonHoverAlpha -= DeltaTime / VREd::DockUIHoverAnimationDuration->GetFloat();;
+			}
+			CloseButtonHoverAlpha = FMath::Clamp( CloseButtonHoverAlpha, 0.0f, 1.0f );
+
+			// How big the close button should be
+			const FVector CloseButtonSize( 1.0f, 2.0f, 2.0f );	// @todo vreditor tweak
+			FVector CloseButtonScale = CloseButtonSize * AnimatedScale * WorldScaleFactor * EasedAimingAtMeFadeAlpha;
+			CloseButtonScale *= FMath::Lerp( 1.0f, VREd::DockUIHoverScale->GetFloat(), CloseButtonHoverAlpha );
+
+
+			CloseButtonMeshComponent->SetRelativeScale3D( CloseButtonScale );
+			const FVector CloseButtonRelativeLocation = FVector(
+				0.0f,
+				-( GetSize().X * 0.5f + CloseButtonSize.Y * 0.5f + VREd::DockUICloseButtonOffsetFromCorner->GetFloat() ),
+				( GetSize().Y * 0.5f + CloseButtonSize.Z * 0.5f + VREd::DockUICloseButtonOffsetFromCorner->GetFloat() ) ) * AnimatedScale * WorldScaleFactor;
+			CloseButtonMeshComponent->SetRelativeLocation( CloseButtonRelativeLocation );
+
+			SetCloseButtonColor( GetOwner().GetOwner().GetColor( bIsHoveringOverCloseButton ? UVREditorMode::EColors::UICloseButtonHoverColor : UVREditorMode::EColors::UICloseButtonColor ) );
+		}
+	}
+} 
+
+void AVREditorDockableWindow::UpdateRelativeRoomTransform()
+{
+	const FTransform RoomToWorld = GetOwner().GetOwner().GetRoomTransform();
+	const FTransform WorldToRoom = RoomToWorld.Inverse();
+
+	const FTransform WindowToWorldTransform = GetActorTransform();
+	const FTransform WindowToRoomTransform = WindowToWorldTransform * WorldToRoom;
+
+	const FVector RoomSpaceWindowLocation = WindowToRoomTransform.GetLocation() / GetOwner().GetOwner().GetWorldScaleFactor();;
+	const FQuat RoomSpaceWindowRotation = WindowToRoomTransform.GetRotation();
+
+	SetRelativeOffset( RoomSpaceWindowLocation );
+	SetLocalRotation( RoomSpaceWindowRotation.Rotator() );
+}
+
+UStaticMeshComponent* AVREditorDockableWindow::GetCloseButtonMeshComponent() const
+{
+	return CloseButtonMeshComponent;
+}
+
+UStaticMeshComponent* AVREditorDockableWindow::GetSelectionBarMeshComponent() const
+{
+	return SelectionBarMeshComponent;
+}
+
+float AVREditorDockableWindow::GetDockSelectDistance() const
+{
+	return DockSelectDistance;
+}
+
+void AVREditorDockableWindow::OnPressed( UViewportInteractor* Interactor, const FHitResult& InHitResult, bool& bOutResultedInDrag )
+{
+	bOutResultedInDrag = false;
+	UVREditorInteractor* VRInteractor = Cast<UVREditorInteractor>( Interactor );
+	if( VRInteractor != nullptr )
+	{
+		if( InHitResult.Component == GetCloseButtonMeshComponent() )
+		{
+			// Close the window
+			const bool bShouldShow = false;
+			const bool bShowOnHand = false;
+			const bool bRefreshQuickMenu = true;
+			GetOwner().ShowEditorUIPanel( this, VRInteractor, bShouldShow, bShowOnHand, bRefreshQuickMenu );
+		}
+		else if( InHitResult.Component == GetSelectionBarMeshComponent() )
+		{
+			bOutResultedInDrag = true;
+
+			FVector LaserPointerStart, LaserPointerEnd;
+			DockSelectDistance = ( InHitResult.TraceStart - InHitResult.Location ).Size();
+			GetOwner().StartDraggingDockUI( this, VRInteractor, DockSelectDistance );
+		}
+	}
+}
+
+void AVREditorDockableWindow::OnHover( UViewportInteractor* Interactor )
+{
+
+}
+
+void AVREditorDockableWindow::OnHoverEnter( UViewportInteractor* Interactor, const FHitResult& InHitResult )
+{
+	if ( SelectionBarMeshComponent == InHitResult.GetComponent() )
+	{
+		bIsHoveringOverSelectionBar = true;
+	}
+
+	if ( CloseButtonMeshComponent == InHitResult.GetComponent() )
+	{
+		bIsHoveringOverCloseButton = true;
+	}
+}
+
+void AVREditorDockableWindow::OnHoverLeave( UViewportInteractor* Interactor, const UActorComponent* NewComponent )
+{
+	UActorComponent* OtherInteractorHoveredComponent = nullptr;
+	if( Interactor->GetOtherInteractor() != nullptr )
+	{
+		OtherInteractorHoveredComponent = Interactor->GetOtherInteractor()->GetHoverComponent();
+	}
+
+	if ( OtherInteractorHoveredComponent != SelectionBarMeshComponent && NewComponent != SelectionBarMeshComponent )
+	{
+		bIsHoveringOverSelectionBar = false;
+	}
+
+	if ( OtherInteractorHoveredComponent != CloseButtonMeshComponent && NewComponent != CloseButtonMeshComponent )
+	{
+		bIsHoveringOverCloseButton = false;
+	}
+}
+
+void AVREditorDockableWindow::OnDragRelease( UViewportInteractor* Interactor )
+{
+	UVREditorInteractor* VREditorInteractor = Cast<UVREditorInteractor>( Interactor );
+	if( VREditorInteractor != nullptr )
+	{
+		UVREditorUISystem& UISystem = GetOwner();
+		UISystem.StopDraggingDockUI( VREditorInteractor );
+	}
+}
+
+UViewportDragOperationComponent* AVREditorDockableWindow::GetDragOperationComponent()
+{
+	return DragOperationComponent;
+}
+
+void AVREditorDockableWindow::SetSelectionBarColor( const FLinearColor& LinearColor )
+{		
+	static FName StaticColorParameterName( "Color" );
+	SelectionBarMID->SetVectorParameterValue( StaticColorParameterName, LinearColor );
+	SelectionBarTranslucentMID->SetVectorParameterValue( StaticColorParameterName, LinearColor );
+}
+
+void AVREditorDockableWindow::SetCloseButtonColor( const FLinearColor& LinearColor )
+{
+	static FName StaticColorParameterName( "Color" );
+	CloseButtonMID->SetVectorParameterValue( StaticColorParameterName, LinearColor );
+	CloseButtonTranslucentMID->SetVectorParameterValue( StaticColorParameterName, LinearColor );
+}
+
+/************************************************************************/
+/* Dockable window drag operation                                       */
+/******************	******************************************************/
+void UDockableWindowDragOperation::ExecuteDrag( UViewportInteractor* Interactor, IViewportInteractableInterface* Interactable  )
+{
+	UVREditorInteractor* VRInteractor = Cast<UVREditorInteractor>( Interactor );
+	AVREditorDockableWindow* DockableWindow = Cast<AVREditorDockableWindow>( Interactable );
+	if ( VRInteractor && DockableWindow )
+	{
+		UVREditorUISystem& UISystem = DockableWindow->GetOwner();
+
+		const bool bIsAbsolute = ( UISystem.GetOwner().GetHMDDeviceType() == EHMDDeviceType::DT_SteamVR );
+		float SlideDeltaY = VRInteractor->GetSlideDelta();
+
+		float NewUIScale = DockableWindow->GetScale() + SlideDeltaY;
+		if ( NewUIScale <= UISystem.GetMinDockWindowSize() )
+		{
+			NewUIScale = UISystem.GetMinDockWindowSize();
+		}
+		else if ( NewUIScale >= UISystem.GetMaxDockWindowSize() )
+		{
+			NewUIScale = UISystem.GetMaxDockWindowSize();
+		}
+
+		DockableWindow->SetScale( NewUIScale );
+
+		const FTransform DockedUIToWorld = UISystem.MakeDockableUITransform( DockableWindow, VRInteractor, DockableWindow->GetDockSelectDistance() );
+		DockableWindow->SetActorTransform( DockedUIToWorld );
+		DockableWindow->UpdateRelativeRoomTransform();
+	}
+}
diff --git a/Engine/Source/Editor/VREditor/UI/VREditorDockableWindow.h b/Engine/Source/Editor/VREditor/UI/VREditorDockableWindow.h
new file mode 100644
--- /dev/null
+++ b/Engine/Source/Editor/VREditor/UI/VREditorDockableWindow.h
@@ -0,0 +1,130 @@
+// Copyright 1998-2015 Epic Games, Inc. All Rights Reserved.
+
+#pragma once
+
+#include "VREditorFloatingUI.h"
+#include "ViewportInteractableInterface.h"
+#include "VREditorDockableWindow.generated.h"
+
+/**
+ * An interactive floating UI panel that can be dragged around
+ */
+UCLASS()
+class AVREditorDockableWindow : public AVREditorFloatingUI, public IViewportInteractableInterface
+{
+	GENERATED_BODY()
+	
+public:
+
+	/** Default constructor */
+	AVREditorDockableWindow();
+
+	/** Destructor */
+	~AVREditorDockableWindow();
+	
+	/** Updates the meshes for the UI */
+	virtual void TickManually( float DeltaTime ) override;
+
+	/** Updates the last dragged relative position */
+	void UpdateRelativeRoomTransform();
+	
+	/** Gets the close button component */
+	UStaticMeshComponent* GetCloseButtonMeshComponent() const;
+
+	/** Gets the selection bar component */
+	UStaticMeshComponent* GetSelectionBarMeshComponent() const;
+
+	/** Gets the distance between the interactor and the window when sta	rting drag */
+	float GetDockSelectDistance() const;
+
+	// IViewportInteractableInterface
+	virtual void OnPressed( UViewportInteractor* Interactor, const FHitResult& InHitResult, bool& bOutResultedInDrag ) override;
+	virtual void OnHover( UViewportInteractor* Interactor ) override;
+	
+	/** Enter hover with laser changes the color of SelectionMesh and CloseButtonMesh */
+	virtual void OnHoverEnter( UViewportInteractor* Interactor, const FHitResult& InHitResult ) override;
+	
+	/** Leaving hover with laser changes the color of SelectionMesh and CloseButtonMesh */
+	virtual void OnHoverLeave( UViewportInteractor* Interactor, const UActorComponent* NewComponent ) override;
+	virtual void OnDragRelease( UViewportInteractor* Interactor ) override;
+	virtual class UViewportDragOperationComponent* GetDragOperationComponent() override;
+
+protected:
+
+	// AVREditorFloatingUI overrides
+	virtual void SetupWidgetComponent() override;
+
+private:
+
+	/** Set the color on the dynamic materials of the selection bar */
+	void SetSelectionBarColor( const FLinearColor& LinearColor );
+
+	/** Set the color on the dynamic materials of the close button */
+	void SetCloseButtonColor( const FLinearColor& LinearColor );
+
+	/** The dockable window mesh */
+	UPROPERTY()
+	class UStaticMeshComponent* WindowMeshComponent;
+
+	/** Mesh underneath the window for easy selecting and dragging */
+	UPROPERTY()
+	class UStaticMeshComponent* SelectionBarMeshComponent;
+
+	/** Mesh that represents the close button for this UI */
+	UPROPERTY()
+	class UStaticMeshComponent* CloseButtonMeshComponent;
+
+	/** Selection bar dynamic material  (opaque) */
+	UPROPERTY()
+	class UMaterialInstanceDynamic* SelectionBarMID;
+
+	/** Select bar dynamic material (translucent) */
+	UPROPERTY()
+	class UMaterialInstanceDynamic* SelectionBarTranslucentMID;
+
+	/** Close button dynamic material  (opaque) */
+	UPROPERTY()
+	class UMaterialInstanceDynamic* CloseButtonMID;
+
+	/** Close button dynamic material (translucent) */
+	UPROPERTY()
+	class UMaterialInstanceDynamic* CloseButtonTranslucentMID;
+
+	UPROPERTY()
+	class UViewportDragOperationComponent* DragOperationComponent;
+
+	/** True if at least one hand's laser is aiming toward the UI */
+	bool bIsLaserAimingTowardUI;
+
+	/** Scalar that ramps up toward 1.0 after the user aims toward the UI */
+	float AimingAtMeFadeAlpha;
+
+	/** True if we're hovering over the selection bar */
+	bool bIsHoveringOverSelectionBar;
+
+	/** Scalar that will advance toward 1.0 over time as we hover over the selection bar */
+	float SelectionBarHoverAlpha;
+
+	/** True if we're hovering over the close button */
+	bool bIsHoveringOverCloseButton;
+
+	/** Scalar that will advance toward 1.0 over time as we hover over the close button */
+	float CloseButtonHoverAlpha;
+
+	/** Distance from interactor laser to the handle when starting dragging */
+	float DockSelectDistance;
+};
+
+/**
+ * Calculation for dragging a dockable window
+ */
+UCLASS()
+class UDockableWindowDragOperation : public UViewportDragOperation
+{
+	GENERATED_BODY()
+
+public:
+
+	// IViewportDragOperation
+	virtual void ExecuteDrag( class UViewportInteractor* Interactor, IViewportInteractableInterface* Interactable ) override;
+};
\ No newline at end of file
diff --git a/Engine/Source/Editor/VREditor/UI/VREditorFloatingUI.cpp b/Engine/Source/Editor/VREditor/UI/VREditorFloatingUI.cpp
new file mode 100644
--- /dev/null
+++ b/Engine/Source/Editor/VREditor/UI/VREditorFloatingUI.cpp
@@ -0,0 +1,297 @@
+// Copyright 1998-2015 Epic Games, Inc. All Rights Reserved.
+
+#include "VREditorModule.h"
+#include "VREditorFloatingUI.h"
+#include "VREditorUISystem.h"
+#include "VREditorBaseUserWidget.h"
+#include "VREditorMode.h"
+#include "WidgetComponent.h"
+#include "VREditorWidgetComponent.h"
+#include "VREditorInteractor.h"
+#include "SLevelViewport.h"	// For tab manager tricks
+
+
+namespace VREd
+{
+	static FAutoConsoleVariable UIFadeSpeed( TEXT( "VREd.UIFadeSpeed" ), 6.0f, TEXT( "How fast UI should fade in and out" ) );
+	static FAutoConsoleVariable UIBrightness( TEXT( "VREd.UIBrightness" ), 0.3f, TEXT( "How bright the UI should be" ) );
+}
+
+
+AVREditorFloatingUI::AVREditorFloatingUI()
+	: Super(),
+	SlateWidget( nullptr ),
+	UserWidget( nullptr ),
+	WidgetComponent( nullptr ),
+	Resolution( 0, 0 ),
+	Owner( nullptr ),
+	UserWidgetClass( nullptr ),
+	bShouldBeVisible(),
+	FadeAlpha( 1.0f ),
+	bCollisionOnShowUI( true ),
+	FadeDelay( 0.0f ),
+	InitialScale( 1.0f )
+{
+	const bool bTransient = true;
+	USceneComponent* SceneComponent = CreateDefaultSubobject<USceneComponent>( TEXT( "SceneComponent" ), bTransient );
+	check( SceneComponent != nullptr );
+	this->RootComponent = SceneComponent;
+
+	WidgetComponent = CreateDefaultSubobject<UVREditorWidgetComponent>( TEXT( "WidgetComponent" ), bTransient );
+	WidgetComponent->SetEditTimeUsable(true);
+	WidgetComponent->SetupAttachment( SceneComponent );
+
+	InitialScale = Scale;
+}
+
+
+void AVREditorFloatingUI::SetupWidgetComponent()
+{
+	WidgetComponent->SetTwoSided( false );	// No VR UI is two-sided
+
+	if( SlateWidget.IsValid() )
+	{
+		// Slate UIs have bogus opacity in their texture's alpha, so ignore texture alpha for VR
+		WidgetComponent->SetOpacityFromTexture( 0.0f );	// Slate UIs have bogus opacity in their texture's alpha, so ignore texture alpha for VR
+		WidgetComponent->SetBackgroundColor( FLinearColor::Black );
+		WidgetComponent->SetBlendMode( EWidgetBlendMode::Opaque );
+	}
+	else
+	{
+		WidgetComponent->SetOpacityFromTexture( 1.0f );
+		WidgetComponent->SetBackgroundColor( FLinearColor::Transparent );
+		WidgetComponent->SetBlendMode( EWidgetBlendMode::Masked );
+	}
+
+	// @todo vreditor: Ideally we use automatic mip map generation, otherwise the UI looks too crunchy at a distance.
+	// However, I tried this and on D3D11 the mips are all black.
+	WidgetComponent->SetDrawSize( FVector2D( Resolution.X, Resolution.Y ) );		// NOTE: Must be called before RegisterComponent() because collision data will be created during registration
+
+	// NOTE: Must be called *after* RegisterComponent() because UWidgetComponent nulls out Widget if no WidgetClass is set (WidgetClass is protected and there is no accessor)
+	if( SlateWidget.IsValid() )
+	{
+		WidgetComponent->SetSlateWidget( SlateWidget.ToSharedRef() );
+	}
+	else
+	{
+		// @todo vreditor unreal: Ideally we would do this in the constructor and not again after.  Due to an apparent bug in UMG, 
+		// we need to re-create the widget in lock-step with the WidgetComponent, otherwise input doesn't function correctly on the 
+		// widget after the widget component is destroyed and recreated with the same user widget.
+		check( UserWidgetClass != nullptr );
+		UserWidget = CreateWidget<UVREditorBaseUserWidget>( GetWorld(), UserWidgetClass );
+		check( UserWidget != nullptr );
+		UserWidget->SetOwner( this );
+
+		WidgetComponent->SetWidget( UserWidget );
+	}
+
+	// @todo vreditor: Is this useful?
+	//WidgetComponent->SetMaxInteractionDistance( 10000.0f );
+
+	// Default to visible
+	ShowUI( true );
+
+	// Set initial opacity
+	UpdateFadingState( 0.0f );
+
+	// Set initial transform
+	UpdateTransformIfDocked();
+}
+
+void AVREditorFloatingUI::SetSlateWidget( UVREditorUISystem& InitOwner, const TSharedRef<SWidget>& InitSlateWidget, const FIntPoint InitResolution, const float InitScale, const EDockedTo InitDockedTo )
+{
+	Owner = &InitOwner;
+	SetVRMode( &Owner->GetOwner() );
+
+	SlateWidget = InitSlateWidget;
+
+	Resolution = InitResolution;
+	check( Resolution.X > 0 && Resolution.Y > 0 );
+
+	Scale = InitScale;
+	InitialScale = Scale;
+
+	SetDockedTo( InitDockedTo );
+
+	SetupWidgetComponent();
+}
+
+
+void AVREditorFloatingUI::SetUMGWidget( UVREditorUISystem& InitOwner, TSubclassOf<UVREditorBaseUserWidget> InitUserWidgetClass, const FIntPoint InitResolution, const float InitScale, const EDockedTo InitDockedTo )
+{
+	Owner = &InitOwner;
+	SetVRMode( &Owner->GetOwner() );
+
+	check( InitUserWidgetClass != nullptr );
+	UserWidgetClass = InitUserWidgetClass;
+
+	Resolution = InitResolution;
+	check( Resolution.X > 0 && Resolution.Y > 0 );
+
+	Scale = InitScale;
+	InitialScale = Scale;
+
+	SetDockedTo( InitDockedTo );
+
+	SetupWidgetComponent();
+}
+
+void AVREditorFloatingUI::Destroyed()
+{
+	if( WidgetComponent != nullptr )
+	{
+		// NOTE: We're nulling out widgets so that we don't have to wait for a GC to free up Slate resources (avoid shutdown crash)
+		WidgetComponent->SetSlateWidget( nullptr );
+		WidgetComponent->SetWidget( nullptr );
+		WidgetComponent = nullptr;
+	}
+
+	this->SlateWidget = nullptr;
+
+	// @todo vreditor unreal: UMG has a bug that prevents you from re-using the user widget for a new widget component
+	// after a previous widget component that was using it was destroyed
+	if( UserWidget != nullptr )
+	{
+		UserWidget->MarkPendingKill();
+		UserWidget = nullptr;
+	}
+
+	Super::Destroyed();
+}
+
+void AVREditorFloatingUI::SetTransform( const FTransform& Transform )
+{
+	const FVector AnimatedScale = CalculateAnimatedScale();
+
+	FTransform AnimatedTransform = Transform;
+	AnimatedTransform.SetScale3D( AnimatedTransform.GetScale3D() * AnimatedScale );
+
+	const float Aspect = ( float ) Resolution.X / ( float ) Resolution.Y;
+
+	RootComponent->SetWorldLocation( AnimatedTransform.GetLocation() );
+	RootComponent->SetWorldRotation( AnimatedTransform.GetRotation() );
+	WidgetComponent->SetWorldScale3D( FVector( 1.0f / AnimatedTransform.GetScale3D().X, 1.0f / ( float ) Resolution.X, 1.0f / ( float ) Resolution.Y / Aspect ) * AnimatedTransform.GetScale3D() );
+}
+
+
+void AVREditorFloatingUI::UpdateFadingState( const float DeltaTime )
+{
+	if ( FadeDelay > 0.f )
+	{
+		FadeDelay -= DeltaTime;
+	}
+	else
+	{
+		if( bShouldBeVisible.GetValue() )
+		{
+			FadeAlpha += VREd::UIFadeSpeed->GetFloat() * DeltaTime;
+		}
+		else
+		{
+			FadeAlpha -= VREd::UIFadeSpeed->GetFloat() * DeltaTime;
+		}
+		FadeAlpha = FMath::Clamp( FadeAlpha, 0.0f, 1.0f );
+
+		if( FadeAlpha > 0.0f + KINDA_SMALL_NUMBER )
+		{
+			// At least a little bit visible
+			if( bHidden )
+			{
+				SetActorHiddenInGame( false );
+				WidgetComponent->SetVisibility( true );
+				FadeDelay = 0.0f;
+			}
+		}
+
+		if( FadeAlpha >= 1.0f - KINDA_SMALL_NUMBER )
+		{
+			// Fully visible
+		}
+		else if( FadeAlpha <= 0.0f + KINDA_SMALL_NUMBER )
+		{
+			// Fully invisible
+			if( !bHidden )
+			{
+				SetActorHiddenInGame( true );
+				WidgetComponent->SetVisibility( false );
+				FadeDelay = 0.0f;
+			}
+		}
+
+ 		// Set material color
+		// NOTE: We intentionally make the UI quite dark here, so it doesn't bloom out in HDR
+		const float UIBrightness = FadeAlpha * VREd::UIBrightness->GetFloat();
+		WidgetComponent->SetTintColorAndOpacity( FLinearColor( UIBrightness, UIBrightness, UIBrightness ).CopyWithNewOpacity( FadeAlpha ) );
+	}
+}
+
+FVector AVREditorFloatingUI::CalculateAnimatedScale() const
+{
+	const float AnimationOvershootAmount = 0.7f;	// @todo vreditor tweak
+	const float EasedAlpha = UVREditorMode::OvershootEaseOut( FadeAlpha, AnimationOvershootAmount );
+
+	// Animate vertically more than horizontally; just looks a little better
+	const float ZScale = FMath::Max( 0.001f, EasedAlpha );
+	const float YScale = FMath::Max( 0.001f, 0.7f + 0.3f * EasedAlpha );
+
+	FVector AnimatedScale = FVector( 1.0f, YScale, ZScale );
+	AnimatedScale.Y *= YScale;
+	AnimatedScale.Z *= ZScale;
+
+	return AnimatedScale;
+}
+
+void AVREditorFloatingUI::SetCollisionOnShow( const bool bInCollisionOnShow )
+{
+	bCollisionOnShowUI = bInCollisionOnShow;
+}
+
+float AVREditorFloatingUI::GetInitialScale() const
+{
+	return InitialScale;
+}
+
+
+void AVREditorFloatingUI::ShowUI( const bool bShow, const bool bAllowFading, const float InitFadeDelay )
+{
+	if( !bShouldBeVisible.IsSet() || bShow != bShouldBeVisible.GetValue() )
+	{
+		bShouldBeVisible = bShow;
+
+		if( bCollisionOnShowUI )
+		{
+			SetActorEnableCollision( bShow );
+		}
+
+		if( !bAllowFading )
+		{
+			SetActorHiddenInGame( !bShow );
+			WidgetComponent->SetVisibility( bShow );
+			FadeAlpha = bShow ? 1.0f : 0.0f;
+		}
+
+		FadeDelay = InitFadeDelay;
+	}
+}
+
+
+FVector2D AVREditorFloatingUI::GetSize() const
+{
+	const float Aspect = (float)Resolution.X / (float)Resolution.Y;
+	return FVector2D( Scale, Scale / Aspect );
+}
+
+float AVREditorFloatingUI::GetScale() const
+{
+	return Scale;
+}
+
+void AVREditorFloatingUI::SetScale( const float NewSize )
+{
+	Scale = NewSize;
+
+	const float WorldScaleFactor = Owner->GetOwner().GetWorldScaleFactor();
+	const FVector NewScale( Scale * WorldScaleFactor );
+	const float Aspect = (float)Resolution.X / (float)Resolution.Y;
+	WidgetComponent->SetWorldScale3D( FVector( 1.0f / NewScale.X, 1.0f / (float)Resolution.X, 1.0f / (float)Resolution.Y / Aspect ) * NewScale );
+}
\ No newline at end of file
diff --git a/Engine/Source/Editor/VREditor/UI/VREditorFloatingUI.h b/Engine/Source/Editor/VREditor/UI/VREditorFloatingUI.h
new file mode 100644
--- /dev/null
+++ b/Engine/Source/Editor/VREditor/UI/VREditorFloatingUI.h
@@ -0,0 +1,138 @@
+// Copyright 1998-2015 Epic Games, Inc. All Rights Reserved.
+
+#pragma once
+
+#include "VREditorBaseActor.h"
+#include "VREditorFloatingUI.generated.h"
+
+
+/**
+ * Represents an interactive floating UI panel in the VR Editor
+ */
+UCLASS()
+class AVREditorFloatingUI : public AVREditorBaseActor
+{
+	GENERATED_BODY()
+
+public:
+
+	/** Default constructor which sets up safe defaults */
+	AVREditorFloatingUI();
+
+	/** Creates a FVREditorFloatingUI using a Slate widget, and sets up safe defaults */
+	void SetSlateWidget( class UVREditorUISystem& InitOwner, const TSharedRef<SWidget>& InitSlateWidget, const FIntPoint InitResolution, const float InitScale, const EDockedTo InitDockedTo );
+
+	/** Creates a FVREditorFloatingUI using a UMG user widget, and sets up safe defaults */
+	void SetUMGWidget( class UVREditorUISystem& InitOwner, class TSubclassOf<class UVREditorBaseUserWidget> InitUserWidgetClass, const FIntPoint InitResolution, const float InitScale, const EDockedTo InitDockedTo );
+
+	/** @return Returns true if the UI is visible (or wants to be visible -- it might be transitioning */
+	bool IsUIVisible() const
+	{
+		return bShouldBeVisible.GetValue();
+	}
+
+	/** Shows or hides the UI (also enables collision, and performs a transition effect) */
+	void ShowUI( const bool bShow, const bool bAllowFading = true, const float InitFadeDelay = 0.0f );
+
+	/** Returns the widget component for this UI, or nullptr if not spawned right now */
+	class UVREditorWidgetComponent* GetWidgetComponent()
+	{
+		return WidgetComponent;
+	}
+
+	/** Returns the actual size of the UI in either axis, after scaling has been applied.  Does not take into animation or world scaling */
+	FVector2D GetSize() const;
+
+	/** Gets the scale */
+	float GetScale() const;
+
+	/** Sets a new size for the UI */
+	void SetScale( const float NewSize );
+
+	/** Sets the UI transform */
+	virtual void SetTransform( const FTransform& Transform ) override;
+
+	/** AActor overrides */
+	virtual void Destroyed() override;
+
+	/** Returns the owner of this object */
+	UVREditorUISystem& GetOwner()
+	{
+		check( Owner != nullptr );
+		return *Owner;
+	}
+
+	/** Returns the owner of this object (const) */
+	const UVREditorUISystem& GetOwner() const
+	{
+		check( Owner != nullptr );
+		return *Owner;
+	}
+
+	/** Returns casted userwidget */
+	template<class T>
+	T* GetUserWidget() const
+	{
+		return CastChecked<T>( UserWidget );
+	}
+
+	/** Sets if the collision is on when showing the UI*/
+	void SetCollisionOnShow( const bool bInCollisionOnShow );
+
+	/** Gets the initial size of this UI */
+	float GetInitialScale() const;
+
+protected:
+
+	/** Returns a scale to use for this widget that takes into account animation */
+	FVector CalculateAnimatedScale() const;
+
+	/** Called to finish setting everything up, after a widget has been assigned */
+	virtual void SetupWidgetComponent();
+
+private:
+
+	/** Called after spawning, and every Tick, to update opacity of the widget */
+	virtual void UpdateFadingState( const float DeltaTime );
+
+protected:
+	/** Slate widget we're drawing, or null if we're drawing a UMG user widget */
+	TSharedPtr<SWidget> SlateWidget;
+	
+	/** UMG user widget we're drawing, or nullptr if we're drawing a Slate widget */
+	UPROPERTY()
+	UVREditorBaseUserWidget* UserWidget;
+	
+	/** When in a spawned state, this is the widget component to represent the widget */
+	UPROPERTY()
+	class UVREditorWidgetComponent* WidgetComponent;
+
+	/** Resolution we should draw this UI at, regardless of scale */
+	FIntPoint Resolution;
+
+private:
+
+	/** Owning object */
+	class UVREditorUISystem* Owner;
+
+	/** Class of the UMG widget we're spawning */
+	UPROPERTY()
+	UClass* UserWidgetClass;
+
+	/** True if the UI wants to be visible, or false if it wants to be hidden.  Remember, we might still be visually 
+	    transitioning between states */
+	TOptional<bool> bShouldBeVisible;
+
+	/** Fade alpha, for visibility transitions */
+	float FadeAlpha;
+
+	/** Collision on when showing UI */
+	bool bCollisionOnShowUI;
+
+	/** Delay to fading in or out. Set on ShowUI and cleared on finish of fade in/out */
+	float FadeDelay;
+
+	/** The starting scale of this UI */
+	float InitialScale;
+};
+
diff --git a/Engine/Source/Editor/VREditor/UI/VREditorQuickMenu.cpp b/Engine/Source/Editor/VREditor/UI/VREditorQuickMenu.cpp
new file mode 100644
--- /dev/null
+++ b/Engine/Source/Editor/VREditor/UI/VREditorQuickMenu.cpp
@@ -0,0 +1,435 @@
+// Copyright 1998-2015 Epic Games, Inc. All Rights Reserved.
+
+#include "VREditorModule.h"
+#include "VREditorQuickMenu.h"
+#include "VREditorUISystem.h"
+#include "VREditorMode.h"
+#include "ViewportWorldInteraction.h"
+#include "VREditorFloatingUI.h"
+#include "VREditorTransformGizmo.h"
+#include "SLevelViewport.h"	// For Simulate toggle
+#include "ImageUtils.h"
+
+
+#define LOCTEXT_NAMESPACE "VREditorQuickMenu"
+
+UVREditorQuickMenu::UVREditorQuickMenu( const FObjectInitializer& ObjectInitializer )
+	: Super( ObjectInitializer )
+{
+}
+
+
+void UVREditorQuickMenu::OnExitVRButtonClicked()
+{
+	GetOwner()->GetOwner().GetOwner().StartExitingVRMode();
+}
+
+
+void UVREditorQuickMenu::OnDeleteButtonClicked()
+{
+	GetOwner()->GetOwner().GetOwner().GetWorldInteraction().DeleteSelectedObjects();
+}
+
+
+void UVREditorQuickMenu::OnUndoButtonClicked()
+{
+	GetOwner()->GetOwner().GetOwner().GetWorldInteraction().Undo();
+}
+
+
+void UVREditorQuickMenu::OnRedoButtonClicked()
+{
+	GetOwner()->GetOwner().GetOwner().GetWorldInteraction().Redo();
+}
+
+
+void UVREditorQuickMenu::OnCopyButtonClicked()
+{
+	GetOwner()->GetOwner().GetOwner().GetWorldInteraction().Copy();
+}
+
+
+void UVREditorQuickMenu::OnPasteButtonClicked()
+{
+	GetOwner()->GetOwner().GetOwner().GetWorldInteraction().Paste();
+}
+
+
+bool UVREditorQuickMenu::OnTranslationSnapButtonClicked( const bool bIsChecked )
+{
+	GUnrealEd->Exec( GetWorld(), *FString::Printf( TEXT( "MODE GRID=%d" ), bIsChecked ? 1 : 0 ) );
+
+	return IsTranslationSnapEnabled();
+}
+
+
+bool UVREditorQuickMenu::IsTranslationSnapEnabled() const
+{
+	return GetDefault<ULevelEditorViewportSettings>()->GridEnabled;
+}
+
+
+FText UVREditorQuickMenu::OnTranslationSnapSizeButtonClicked()
+{
+	const TArray<float>& GridSizes = GEditor->GetCurrentPositionGridArray();
+	const int32 CurrentGridSize = GetDefault<ULevelEditorViewportSettings>()->CurrentPosGridSize;
+
+	int32 NewGridSize = CurrentGridSize + 1;
+	if( NewGridSize >= GridSizes.Num() )
+	{
+		NewGridSize = 0;
+	}
+
+	GEditor->SetGridSize( NewGridSize );
+
+	return GetTranslationSnapSizeText();
+}
+
+
+FText UVREditorQuickMenu::GetTranslationSnapSizeText() const
+{
+	return FText::AsNumber( GEditor->GetGridSize() );
+}
+
+
+FText UVREditorQuickMenu::OnRotationSnapSizeButtonClicked()
+{
+	const TArray<float>& GridSizes = GEditor->GetCurrentRotationGridArray();
+	const int32 CurrentGridSize = GetDefault<ULevelEditorViewportSettings>()->CurrentRotGridSize;
+	const ERotationGridMode CurrentGridMode = GetDefault<ULevelEditorViewportSettings>()->CurrentRotGridMode;
+
+	int32 NewGridSize = CurrentGridSize + 1;
+	if( NewGridSize >= GridSizes.Num() )
+	{
+		NewGridSize = 0;
+	}
+
+	GEditor->SetRotGridSize( NewGridSize, CurrentGridMode );
+
+	return GetRotationSnapSizeText();
+}
+
+
+FText UVREditorQuickMenu::GetRotationSnapSizeText() const
+{
+	return FText::AsNumber( GEditor->GetRotGridSize().Yaw );
+}
+
+
+FText UVREditorQuickMenu::OnScaleSnapSizeButtonClicked()
+{
+	const ULevelEditorViewportSettings* ViewportSettings = GetMutableDefault<ULevelEditorViewportSettings>();
+	const int32 NumGridSizes = ViewportSettings->ScalingGridSizes.Num();
+
+	const int32 CurrentGridSize = GetDefault<ULevelEditorViewportSettings>()->CurrentScalingGridSize;
+
+	int32 NewGridSize = CurrentGridSize + 1;
+	if( NewGridSize >= NumGridSizes )
+	{
+		NewGridSize = 0;
+	}
+
+	GEditor->SetScaleGridSize( NewGridSize );
+
+	return GetScaleSnapSizeText();
+}
+
+
+FText UVREditorQuickMenu::GetScaleSnapSizeText() const
+{
+	return FText::AsNumber( GEditor->GetScaleGridSize() );
+}
+
+
+bool UVREditorQuickMenu::OnRotationSnapButtonClicked( const bool bIsChecked )
+{
+	GUnrealEd->Exec( GetWorld(), *FString::Printf( TEXT( "MODE ROTGRID=%d" ), bIsChecked ? 1 : 0 ) );
+
+	return IsRotationSnapEnabled();
+}
+
+
+bool UVREditorQuickMenu::IsRotationSnapEnabled() const
+{
+	return GetDefault<ULevelEditorViewportSettings>()->RotGridEnabled;
+}
+
+
+bool UVREditorQuickMenu::OnScaleSnapButtonClicked( const bool bIsChecked )
+{
+	GUnrealEd->Exec( GetWorld(), *FString::Printf( TEXT( "MODE SCALEGRID=%d" ), bIsChecked ? 1 : 0 ) );
+
+	return IsScaleSnapEnabled();
+}
+
+
+bool UVREditorQuickMenu::IsScaleSnapEnabled() const
+{
+	return GetDefault<ULevelEditorViewportSettings>()->SnapScaleEnabled;
+}
+
+
+FText UVREditorQuickMenu::OnGizmoCoordinateSystemButtonClicked()
+{
+	GetOwner()->GetOwner().GetOwner().GetWorldInteraction().CycleTransformGizmoCoordinateSpace(); 
+	return GetGizmoCoordinateSystemText();
+}
+
+
+FText UVREditorQuickMenu::GetGizmoCoordinateSystemText() const
+{
+	const ECoordSystem CurrentCoordSystem = GetOwner()->GetOwner().GetOwner().GetWorldInteraction().GetTransformGizmoCoordinateSpace(); //@todo VREditor
+	return CurrentCoordSystem == COORD_World ? LOCTEXT( "WorldCoordinateSystem", "World Space" ) : LOCTEXT( "LocalCoordinateSystem", "Local Space" );
+}
+
+
+FText UVREditorQuickMenu::OnGizmoModeButtonClicked()
+{
+	const EGizmoHandleTypes CurrentGizmoMode = GetOwner()->GetOwner().GetOwner().GetWorldInteraction().GetCurrentGizmoType();
+	GetOwner()->GetOwner().GetOwner().CycleTransformGizmoHandleType();
+	return GetGizmoModeText();
+}
+
+
+FText UVREditorQuickMenu::GetGizmoModeText() const
+{
+	const EGizmoHandleTypes CurrentGizmoType = GetOwner()->GetOwner().GetOwner().GetWorldInteraction().GetCurrentGizmoType();
+	FText GizmoTypeText;
+	switch( CurrentGizmoType )
+	{
+		case EGizmoHandleTypes::All:
+			GizmoTypeText = LOCTEXT( "AllGizmoType", "Universal" );
+			break;
+
+		case EGizmoHandleTypes::Translate:
+			GizmoTypeText = LOCTEXT( "TranslateGizmoType", "Translate" );
+			break;
+
+		case EGizmoHandleTypes::Rotate:
+			GizmoTypeText = LOCTEXT( "RotationGizmoType", "Rotation" );
+			break;
+
+		case EGizmoHandleTypes::Scale:
+			GizmoTypeText = LOCTEXT( "ScaleGizmoType", "Scale" );
+			break;
+
+		default:
+			check( 0 );	// Unrecognized type
+			break;
+	}
+
+	return GizmoTypeText;
+}
+
+
+bool UVREditorQuickMenu::OnSimulateButtonClicked( const bool bIsChecked )
+{
+	if( bIsChecked != GEditor->bIsSimulatingInEditor )
+	{
+		if( GEditor->PlayWorld == nullptr && !GEditor->bIsSimulatingInEditor )
+		{
+			const bool bSimulateInEditor = false;
+			const bool bAtPlayerStart = true;
+
+			// If there is an active level view port, play the game in it.
+			TSharedRef<ILevelViewport> LevelViewport = StaticCastSharedRef<SLevelViewport>( GetOwner()->GetOwner().GetOwner().GetLevelViewportPossessedForVR().AsShared() );
+			GUnrealEd->RequestPlaySession(
+				bAtPlayerStart,
+				LevelViewport,
+				bSimulateInEditor );
+		}
+		else if( GEditor->bIsSimulatingInEditor )
+		{
+			GEditor->RequestEndPlayMap();
+		}
+		else if( GUnrealEd->PlayWorld != nullptr )
+		{
+			// There is already a play world active which means simulate in editor is happening
+			// Toggle to pie 
+			GUnrealEd->RequestToggleBetweenPIEandSIE();
+		}
+	}
+
+	return IsSimulatingEnabled();
+}
+
+
+bool UVREditorQuickMenu::IsSimulatingEnabled() const
+{
+	return GEditor->bIsSimulateInEditorQueued || GEditor->bIsSimulatingInEditor;
+}
+
+
+bool UVREditorQuickMenu::OnContentBrowserButtonClicked( const bool bIsChecked )
+{
+	GetOwner()->GetOwner().TogglePanelVisibility( UVREditorUISystem::EEditorUIPanel::ContentBrowser );
+	return IsContentBrowserVisible();
+}
+
+
+bool UVREditorQuickMenu::IsContentBrowserVisible() const
+{
+	return ( GetOwner()->GetOwner().IsShowingEditorUIPanel( UVREditorUISystem::EEditorUIPanel::ContentBrowser ) );
+}
+
+
+bool UVREditorQuickMenu::OnWorldOutlinerButtonClicked( const bool bIsChecked )
+{
+	GetOwner()->GetOwner().TogglePanelVisibility( UVREditorUISystem::EEditorUIPanel::WorldOutliner );
+	return IsWorldOutlinerVisible();
+}
+
+
+bool UVREditorQuickMenu::IsWorldOutlinerVisible() const
+{
+	return ( GetOwner()->GetOwner().IsShowingEditorUIPanel( UVREditorUISystem::EEditorUIPanel::WorldOutliner ) );
+}
+
+
+bool UVREditorQuickMenu::OnActorDetailsButtonClicked( const bool bIsChecked )
+{
+	GetOwner()->GetOwner().TogglePanelVisibility( UVREditorUISystem::EEditorUIPanel::ActorDetails );
+	return IsActorDetailsVisible();
+}
+
+
+bool UVREditorQuickMenu::IsActorDetailsVisible() const
+{
+	return ( GetOwner()->GetOwner().IsShowingEditorUIPanel( UVREditorUISystem::EEditorUIPanel::ActorDetails ) );
+}
+
+
+bool UVREditorQuickMenu::OnModesButtonClicked( const bool bIsChecked )
+{
+	GetOwner()->GetOwner().TogglePanelVisibility( UVREditorUISystem::EEditorUIPanel::Modes );
+	return IsModesVisible();
+}
+
+
+bool UVREditorQuickMenu::IsModesVisible() const
+{
+	return ( GetOwner()->GetOwner().IsShowingEditorUIPanel( UVREditorUISystem::EEditorUIPanel::Modes ) );
+}
+
+
+bool UVREditorQuickMenu::OnGameModeButtonClicked( const bool bIsChecked )
+{
+	FEditorViewportClient* ViewportClient = GetOwner()->GetOwner().GetOwner().GetLevelViewportPossessedForVR().GetViewportClient().Get();
+	const bool bIsInGameView = ViewportClient->IsInGameView();
+	if( bIsChecked != bIsInGameView )
+	{
+		ViewportClient->SetGameView( bIsChecked );
+	}
+
+	return IsGameModeEnabled();
+}
+
+
+bool UVREditorQuickMenu::IsGameModeEnabled() const
+{
+	const FEditorViewportClient* ViewportClient = GetOwner()->GetOwner().GetOwner().GetLevelViewportPossessedForVR().GetViewportClient().Get();
+	return ViewportClient->IsInGameView();
+}
+
+
+bool UVREditorQuickMenu::OnTutorialButtonClicked( const bool bIsChecked )
+{
+	GetOwner()->GetOwner().TogglePanelVisibility( UVREditorUISystem::EEditorUIPanel::Tutorial );
+	return IsTutorialVisible();
+}
+
+bool UVREditorQuickMenu::IsTutorialVisible() const
+{
+	return ( GetOwner()->GetOwner().IsShowingEditorUIPanel( UVREditorUISystem::EEditorUIPanel::Tutorial ) );
+}
+
+bool UVREditorQuickMenu::OnLightButtonClicked( const bool bIsChecked ) 
+{
+	EControllerHand ControllerHand = GetOwner()->GetDockedTo() == AVREditorFloatingUI::EDockedTo::LeftArm ? EControllerHand::Right : EControllerHand::Left;
+	UVREditorInteractor* Interactor = GetOwner()->GetOwner().GetOwner().GetHandInteractor( ControllerHand );
+	GetOwner()->GetOwner().GetOwner().ToggleFlashlight( Interactor );
+	return IsLightVisible();
+}
+
+bool UVREditorQuickMenu::IsLightVisible() const
+{
+	return (GetOwner()->GetOwner().GetOwner().IsFlashlightOn());
+}
+
+void UVREditorQuickMenu::OnScreenshotButtonClicked()
+{
+	// @todo vreditor: update after direct buffer grab changes
+	
+	FString GeneratedFilename = TEXT("");
+	FString Filename = TEXT("");
+	FScreenshotRequest::CreateViewportScreenShotFilename(GeneratedFilename);
+	const bool bRemovePath = false;
+	GeneratedFilename = FPaths::GetBaseFilename(GeneratedFilename, bRemovePath);
+	FFileHelper::GenerateNextBitmapFilename(GeneratedFilename, TEXT("png"), Filename);
+
+	FViewport* Viewport = GetOwner()->GetOwner().GetOwner().GetLevelViewportPossessedForVR().GetActiveViewport();
+	TSharedRef<SWidget> WindowRef = GetOwner()->GetOwner().GetOwner().GetLevelViewportPossessedForVR().AsWidget();
+
+
+	TArray<FColor> OutImageData;
+	FIntVector OutImageSize;
+
+	if (FSlateApplication::Get().TakeScreenshot(WindowRef, OutImageData, OutImageSize))
+	{
+		int32 Width = OutImageSize.X;
+		int32 Height = OutImageSize.Y;
+
+		TArray<FColor> ScaledBitmap;
+
+
+		ScaledBitmap = OutImageData;
+		//Clear the alpha channel before saving
+		for (int32 Index = 0; Index < Width*Height; Index++)
+		{
+			ScaledBitmap[Index].A = 255;
+		}
+
+
+		TArray<uint8> CompressedBitmap;
+		FImageUtils::CompressImageArray(Width, Height, ScaledBitmap, CompressedBitmap);
+
+
+		//Save locally
+		const bool bTree = true;
+		const FString FinalFileName = Filename;
+		IFileManager::Get().MakeDirectory(*FPaths::GetPath(FinalFileName), bTree);
+		FFileHelper::SaveArrayToFile(CompressedBitmap, *FinalFileName);
+
+	}
+
+
+}
+
+void UVREditorQuickMenu::OnPlayButtonClicked()
+{
+	GetOwner()->GetOwner().GetOwner().StartExitingVRMode( EVREditorExitType::PIE_VR );
+}
+
+bool UVREditorQuickMenu::OnAssetEditorButtonClicked( const bool bIsChecked )
+{
+	GetOwner()->GetOwner().TogglePanelVisibility( UVREditorUISystem::EEditorUIPanel::AssetEditor );
+	return IsAssetEditorVisible();
+}
+
+bool UVREditorQuickMenu::IsAssetEditorVisible() const
+{
+	return ( GetOwner()->GetOwner().IsShowingEditorUIPanel( UVREditorUISystem::EEditorUIPanel::AssetEditor ) );
+}
+
+bool UVREditorQuickMenu::OnWorldSettingsButtonClicked( const bool bIsChecked )
+{
+	GetOwner()->GetOwner().TogglePanelVisibility( UVREditorUISystem::EEditorUIPanel::WorldSettings);
+	return IsWorldSettingsVisible();
+}
+
+bool UVREditorQuickMenu::IsWorldSettingsVisible() const
+{
+	return ( GetOwner()->GetOwner().IsShowingEditorUIPanel( UVREditorUISystem::EEditorUIPanel::WorldSettings ) );
+}
+
+#undef LOCTEXT_NAMESPACE
\ No newline at end of file
diff --git a/Engine/Source/Editor/VREditor/UI/VREditorQuickMenu.h b/Engine/Source/Editor/VREditor/UI/VREditorQuickMenu.h
new file mode 100644
--- /dev/null
+++ b/Engine/Source/Editor/VREditor/UI/VREditorQuickMenu.h
@@ -0,0 +1,162 @@
+// Copyright 1998-2015 Epic Games, Inc. All Rights Reserved.
+
+#pragma once
+
+
+#include "VREditorBaseUserWidget.h"
+#include "VREditorQuickMenu.generated.h"
+
+
+/**
+ * A simple quick-access menu for VR editing, with frequently-used features
+ */
+UCLASS( Blueprintable )
+class UVREditorQuickMenu : public UVREditorBaseUserWidget
+{
+	GENERATED_BODY()
+
+public:
+	
+	/** Default constructor that sets up CDO properties */
+	UVREditorQuickMenu( const FObjectInitializer& ObjectInitializer );
+
+	/** 
+	 * Event handlers.  These are called from the Blueprint when buttons are clicked. 
+	 */
+
+	UFUNCTION( BlueprintCallable, Category="Default" )
+	void OnExitVRButtonClicked();
+
+	UFUNCTION( BlueprintCallable, Category="Default" )
+	void OnDeleteButtonClicked();
+
+	UFUNCTION( BlueprintCallable, Category = "Default" )
+	void OnUndoButtonClicked();
+
+	UFUNCTION( BlueprintCallable, Category = "Default" )
+	void OnRedoButtonClicked();
+
+	UFUNCTION( BlueprintCallable, Category = "Default" )
+	void OnCopyButtonClicked();
+
+	UFUNCTION( BlueprintCallable, Category = "Default" )
+	void OnPasteButtonClicked();
+
+	UFUNCTION( BlueprintCallable, Category = "Default" )
+	bool OnTranslationSnapButtonClicked( const bool bIsChecked );
+
+	UFUNCTION( BlueprintPure, Category = "Default" )
+	bool IsTranslationSnapEnabled() const;
+
+	UFUNCTION( BlueprintCallable, Category = "Default" )
+	FText OnTranslationSnapSizeButtonClicked();
+
+	UFUNCTION( BlueprintPure, Category = "Default" )
+	FText GetTranslationSnapSizeText() const;
+
+	UFUNCTION( BlueprintCallable, Category = "Default" )
+	bool OnRotationSnapButtonClicked( const bool bIsChecked );
+
+	UFUNCTION( BlueprintPure, Category = "Default" )
+	bool IsRotationSnapEnabled() const;
+
+	UFUNCTION( BlueprintCallable, Category = "Default" )
+	FText OnRotationSnapSizeButtonClicked();
+
+	UFUNCTION( BlueprintPure, Category = "Default" )
+	FText GetRotationSnapSizeText() const;
+
+	UFUNCTION( BlueprintCallable, Category = "Default" )
+	bool OnScaleSnapButtonClicked( const bool bIsChecked );
+
+	UFUNCTION( BlueprintPure, Category = "Default" )
+	bool IsScaleSnapEnabled() const;
+
+	UFUNCTION( BlueprintCallable, Category = "Default" )
+	FText OnScaleSnapSizeButtonClicked();
+
+	UFUNCTION( BlueprintPure, Category = "Default" )
+	FText GetScaleSnapSizeText() const;
+
+	UFUNCTION( BlueprintCallable, Category = "Default" )
+	FText OnGizmoCoordinateSystemButtonClicked();
+
+	UFUNCTION( BlueprintPure, Category = "Default" )
+	FText GetGizmoCoordinateSystemText() const;
+
+	UFUNCTION( BlueprintCallable, Category = "Default" )
+	FText OnGizmoModeButtonClicked();
+
+	UFUNCTION( BlueprintPure, Category = "Default" )
+	FText GetGizmoModeText() const;
+
+	UFUNCTION( BlueprintCallable, Category = "Default" )
+	bool OnSimulateButtonClicked( const bool bIsChecked );
+
+	UFUNCTION( BlueprintPure, Category = "Default" )
+	bool IsSimulatingEnabled() const;
+
+	UFUNCTION( BlueprintCallable, Category = "Default" )
+	bool OnContentBrowserButtonClicked( const bool bIsChecked );
+
+	UFUNCTION( BlueprintPure, Category = "Default" )
+	bool IsContentBrowserVisible() const;
+
+	UFUNCTION( BlueprintCallable, Category = "Default" )
+	bool OnWorldOutlinerButtonClicked( const bool bIsChecked );
+
+	UFUNCTION( BlueprintPure, Category = "Default" )
+	bool IsWorldOutlinerVisible() const;
+
+	UFUNCTION( BlueprintCallable, Category = "Default" )
+	bool OnActorDetailsButtonClicked( const bool bIsChecked );
+
+	UFUNCTION( BlueprintPure, Category = "Default" )
+	bool IsActorDetailsVisible() const;
+
+	UFUNCTION( BlueprintCallable, Category = "Default" )
+	bool OnModesButtonClicked( const bool bIsChecked );
+
+	UFUNCTION( BlueprintPure, Category = "Default" )
+	bool IsModesVisible() const;
+
+	UFUNCTION( BlueprintCallable, Category = "Default" )
+	bool OnGameModeButtonClicked( const bool bIsChecked );
+
+	UFUNCTION( BlueprintPure, Category = "Default" )
+	bool IsGameModeEnabled() const;
+
+	UFUNCTION( BlueprintCallable, Category = "Default" )
+	bool OnTutorialButtonClicked( const bool bIsChecked );
+
+	UFUNCTION( BlueprintPure, Category = "Default" )
+	bool IsTutorialVisible() const;
+
+	UFUNCTION(BlueprintCallable, Category = "Default")
+	bool OnLightButtonClicked(const bool bIsChecked);
+
+	UFUNCTION(BlueprintPure, Category = "Default")
+	bool IsLightVisible() const;
+
+	UFUNCTION(BlueprintCallable, Category = "Default")
+	void OnScreenshotButtonClicked();
+
+	UFUNCTION( BlueprintCallable, Category = "Default" )
+	void OnPlayButtonClicked();
+
+	UFUNCTION( BlueprintCallable, Category = "Default" )
+	bool OnAssetEditorButtonClicked( const bool bIsChecked );
+
+	UFUNCTION( BlueprintPure, Category = "Default" )
+	bool IsAssetEditorVisible() const;
+
+	UFUNCTION(BlueprintCallable, Category = "Default")
+	bool OnWorldSettingsButtonClicked(const bool bIsChecked);
+
+	UFUNCTION(BlueprintPure, Category = "Default")
+	bool IsWorldSettingsVisible() const;
+
+	UFUNCTION(BlueprintImplementableEvent, meta=(DisplayName = "Refresh"))
+	void RefreshUI();
+
+};
\ No newline at end of file
diff --git a/Engine/Source/Editor/VREditor/UI/VREditorRadialMenu.cpp b/Engine/Source/Editor/VREditor/UI/VREditorRadialMenu.cpp
new file mode 100644
--- /dev/null
+++ b/Engine/Source/Editor/VREditor/UI/VREditorRadialMenu.cpp
@@ -0,0 +1,229 @@
+// Copyright 1998-2015 Epic Games, Inc. All Rights Reserved.
+
+#include "VREditorModule.h"
+#include "VREditorRadialMenu.h"
+#include "VREditorRadialMenuItem.h"
+#include "VREditorFloatingUI.h"
+#include "VREditorUISystem.h"
+#include "VREditorMode.h"
+#include "ViewportWorldInteraction.h"
+#include "VREditorTransformGizmo.h"
+#include "VREditorInteractor.h"
+#include "VREditorMotionControllerInteractor.h"
+
+#include "Editor/LevelEditor/Public/LevelEditor.h"
+
+#define LOCTEXT_NAMESPACE "VREditorRadialMenu"
+
+namespace VREd
+{
+	static FAutoConsoleVariable RadialMenuRiftSelectionRadius( TEXT( "VREd.RadialMenuRiftSelectionRadius" ), 0.3f, TEXT( "The radius that is used for selecting buttons on the radial menu on Oculus Rift" ) );
+	static FAutoConsoleVariable RadialMenuViveSelectionRadius( TEXT( "VREd.RadialMenuViveSelectionRadius" ), 0.3f, TEXT( "The radius that is used for selecting buttons on the radial menu on Vive" ) );
+}
+
+UVREditorRadialMenu::UVREditorRadialMenu( const FObjectInitializer& ObjectInitializer )
+: Super( ObjectInitializer ),
+CurrentItem( nullptr ),
+PreviousItem( nullptr )
+{
+}
+
+
+void UVREditorRadialMenu::Update( const UVREditorInteractor* Interactor )
+{
+	const UVREditorMotionControllerInteractor* MotionControllerInteractor = Cast< UVREditorMotionControllerInteractor >( Interactor );
+	if( MotionControllerInteractor )
+	{
+		TrackpadPosition = MotionControllerInteractor->GetTrackpadPosition();
+
+		if ( !GetOwner()->bHidden )
+		{
+			// Check if position is from the center
+			if ( !IsInMenuRadius() )
+			{
+				const float AnglePerItem = 360 / MenuItems.Num();
+				float Angle = FRotator::NormalizeAxis( FMath::RadiansToDegrees( FMath::Atan2( TrackpadPosition.X, TrackpadPosition.Y ) ) );
+				TrackpadAngle = Angle;
+
+				Angle += AnglePerItem / 2.0f;
+				if ( Angle < 0 )
+				{
+					Angle = Angle + 360;
+				}
+
+				const int32 Index = ( Angle / AnglePerItem );
+				UVREditorRadialMenuItem* NewMenuItem = MenuItems[ Index ];
+				if ( NewMenuItem )
+				{
+					CurrentItem = NewMenuItem;
+				}
+
+				// Update the visuals
+				if ( CurrentItem != PreviousItem )
+				{
+					if ( CurrentItem )
+					{
+						CurrentItem->OnEnterHover();
+					}
+
+					if ( PreviousItem )
+					{
+						PreviousItem->OnLeaveHover();
+					}
+				}
+
+				PreviousItem = CurrentItem;
+			}
+			else
+			{
+				if ( CurrentItem )
+				{
+					CurrentItem->OnLeaveHover();
+					CurrentItem = nullptr;
+					PreviousItem = nullptr;
+				}
+			}
+		}
+	}
+}
+
+void UVREditorRadialMenu::SetButtons( UVREditorRadialMenuItem* InitTopItem, UVREditorRadialMenuItem* InitTopRightItem, UVREditorRadialMenuItem* InitRightItem, UVREditorRadialMenuItem* InitBottomRightItem,
+	UVREditorRadialMenuItem* InitBottomItem, UVREditorRadialMenuItem* InitLeftBottomItem, UVREditorRadialMenuItem* InitLeftItem, UVREditorRadialMenuItem* InitTopLeftItem )
+{
+	InitTopItem->OnPressedDelegate.AddUObject( this, &UVREditorRadialMenu::OnGizmoCycle );
+	MenuItems.Add( InitTopItem );
+	UpdateGizmoTypeLabel();
+
+	InitTopRightItem->OnPressedDelegate.AddUObject( this, &UVREditorRadialMenu::OnCopyButtonClicked );
+	InitTopRightItem->SetLabel( FString( "Copy" ) );
+	MenuItems.Add( InitTopRightItem );
+
+	InitRightItem->OnPressedDelegate.AddUObject( this, &UVREditorRadialMenu::OnPasteButtonClicked );
+	InitRightItem->SetLabel( FString( "Paste" ) );
+	MenuItems.Add( InitRightItem );
+
+	InitBottomRightItem->OnPressedDelegate.AddUObject( this, &UVREditorRadialMenu::OnRedoButtonClicked );
+	InitBottomRightItem->SetLabel( FString( "Redo" ) );
+	MenuItems.Add( InitBottomRightItem );
+
+	InitBottomItem->OnPressedDelegate.AddUObject( this, &UVREditorRadialMenu::OnSnapActorsToGroundClicked );
+	InitBottomItem->SetLabel( FString( "Snap To Ground" ) );
+	MenuItems.Add( InitBottomItem );
+
+	InitLeftBottomItem->OnPressedDelegate.AddUObject( this, &UVREditorRadialMenu::OnUndoButtonClicked );
+	InitLeftBottomItem->SetLabel( FString( "Undo" ) );
+	MenuItems.Add( InitLeftBottomItem );
+
+	InitLeftItem->OnPressedDelegate.AddUObject( this, &UVREditorRadialMenu::OnDuplicateButtonClicked );
+	InitLeftItem->SetLabel( FString( "Duplicate" ) );
+	MenuItems.Add( InitLeftItem );
+
+	InitTopLeftItem->OnPressedDelegate.AddUObject( this, &UVREditorRadialMenu::OnDeleteButtonClicked );
+	InitTopLeftItem->SetLabel( FString( "Delete" ) );
+	MenuItems.Add( InitTopLeftItem );
+}
+
+void UVREditorRadialMenu::ResetItem()
+{
+	if (PreviousItem)
+	{
+		PreviousItem->OnLeaveHover();
+	}
+
+	if (CurrentItem)
+	{
+		CurrentItem->OnLeaveHover();
+	}
+}
+
+void UVREditorRadialMenu::SelectCurrentItem()
+{
+	if (IsInMenuRadius())
+	{
+		GetOwner()->GetOwner().TogglePanelsVisibility();
+	}
+	else if (CurrentItem)
+	{
+		CurrentItem->OnPressed();
+	}
+}
+
+bool UVREditorRadialMenu::IsInMenuRadius() const
+{
+	const float RadialMenuSelectionRadius = Owner->GetOwner().GetOwner().GetHMDDeviceType() == EHMDDeviceType::DT_OculusRift ? VREd::RadialMenuRiftSelectionRadius->GetFloat() : VREd::RadialMenuViveSelectionRadius->GetFloat();
+	return FVector2D::Distance( FVector2D::ZeroVector, TrackpadPosition ) < RadialMenuSelectionRadius;
+}
+
+void UVREditorRadialMenu::UpdateGizmoTypeLabel()
+{
+	FString Result;
+
+	switch (GetOwner()->GetOwner().GetOwner().GetWorldInteraction().GetCurrentGizmoType())
+	{
+	case EGizmoHandleTypes::All:
+		Result = LOCTEXT( "AllGizmoType", "Universal Gizmo" ).ToString();
+		break;
+	case EGizmoHandleTypes::Translate:
+		Result = LOCTEXT( "TranslateGizmoType", "Translate Gizmo" ).ToString();
+		break;
+	case EGizmoHandleTypes::Rotate:
+		Result = LOCTEXT( "RotateGizmoType", "Rotate Gizmo" ).ToString();
+		break;
+	case EGizmoHandleTypes::Scale:
+		Result = LOCTEXT( "ScaleGizmoType", "Scale Gizmo" ).ToString();
+		break;
+	default:
+		break;
+	}
+
+	MenuItems[0]->SetLabel( Result );
+}
+
+void UVREditorRadialMenu::OnGizmoCycle()
+{
+	GetOwner()->GetOwner().GetOwner().CycleTransformGizmoHandleType();
+	UpdateGizmoTypeLabel();
+}
+
+void UVREditorRadialMenu::OnDuplicateButtonClicked()
+{
+	GetOwner()->GetOwner().GetOwner().GetWorldInteraction().Duplicate();
+}
+
+
+void UVREditorRadialMenu::OnDeleteButtonClicked()
+{
+	GetOwner()->GetOwner().GetOwner().GetWorldInteraction().DeleteSelectedObjects();
+}
+
+
+void UVREditorRadialMenu::OnUndoButtonClicked()
+{
+	GetOwner()->GetOwner().GetOwner().GetWorldInteraction().Undo();
+}
+
+
+void UVREditorRadialMenu::OnRedoButtonClicked()
+{
+	GetOwner()->GetOwner().GetOwner().GetWorldInteraction().Redo();
+}
+
+
+void UVREditorRadialMenu::OnCopyButtonClicked()
+{
+	GetOwner()->GetOwner().GetOwner().GetWorldInteraction().Copy();
+}
+
+
+void UVREditorRadialMenu::OnPasteButtonClicked()
+{
+	GetOwner()->GetOwner().GetOwner().GetWorldInteraction().Paste();
+}
+
+void UVREditorRadialMenu::OnSnapActorsToGroundClicked()
+{
+	GetOwner()->GetOwner().GetOwner().SnapSelectedActorsToGround();
+}
+
+
+#undef LOCTEXT_NAMESPACE
diff --git a/Engine/Source/Editor/VREditor/UI/VREditorRadialMenu.h b/Engine/Source/Editor/VREditor/UI/VREditorRadialMenu.h
new file mode 100644
--- /dev/null
+++ b/Engine/Source/Editor/VREditor/UI/VREditorRadialMenu.h
@@ -0,0 +1,73 @@
+// Copyright 1998-2015 Epic Games, Inc. All Rights Reserved.
+
+#pragma once
+
+
+#include "VREditorBaseUserWidget.h"
+#include "VREditorRadialMenu.generated.h"
+
+class UVREditorRadialMenuItem;
+
+/**
+ * A simple radial menu for VR editing, with frequently-used features
+ */
+UCLASS( Blueprintable )
+class UVREditorRadialMenu : public UVREditorBaseUserWidget
+{
+	GENERATED_BODY()
+
+public:
+	/** Default constructor that sets up CDO properties */
+	UVREditorRadialMenu( const FObjectInitializer& ObjectInitializer );
+
+	/** Update the state of the widget with the trackpad location */
+	void Update( const class UVREditorInteractor* Interactor );
+
+	/** To bind the buttons from UMG */
+	UFUNCTION(BlueprintCallable, Category = "RadialItems")
+	void SetButtons( UVREditorRadialMenuItem* InitTopItem, UVREditorRadialMenuItem* InitTopRightItem, UVREditorRadialMenuItem* InitRightItem, UVREditorRadialMenuItem* InitBottomRightItem,
+		UVREditorRadialMenuItem* InitBottomItem, UVREditorRadialMenuItem* InitLeftBottomItem, UVREditorRadialMenuItem* InitLeftItem, UVREditorRadialMenuItem* InitTopLeftItem );
+
+	/** Resets the current item to not hovering */
+	void ResetItem();
+
+	/** Press the current item */
+	void SelectCurrentItem();
+
+	/** If the users has his trackpad or analog stick inside the radius where buttons will not be selected */
+	bool IsInMenuRadius() const;
+
+	/** Current position of the trackpad to render line in UMG */
+	UPROPERTY( EditAnywhere, BlueprintReadOnly, Category = "RadialMenu")
+	FVector2D TrackpadPosition;
+
+	/** Current angle of the trackpad position (-180 to 180) to visualize in UMG */
+	UPROPERTY( EditAnywhere, BlueprintReadOnly, Category = "RadialMenu" )
+	float TrackpadAngle;
+
+private:
+	/** All the buttons in the radial menu */
+	UPROPERTY()
+	TArray<UVREditorRadialMenuItem*> MenuItems;
+
+	/** The current widget currently selected */
+	UPROPERTY()
+	UVREditorRadialMenuItem* CurrentItem;
+
+	/** The previous item that was selected */
+	UPROPERTY()
+	UVREditorRadialMenuItem* PreviousItem;
+
+	/** Updates the label of the gizmo type item */
+	void UpdateGizmoTypeLabel();
+
+	/** Temporary functions for delegates of the buttons */
+	void OnGizmoCycle();
+	void OnDuplicateButtonClicked();
+	void OnDeleteButtonClicked();
+	void OnUndoButtonClicked();
+	void OnRedoButtonClicked();
+	void OnCopyButtonClicked();
+	void OnPasteButtonClicked();
+	void OnSnapActorsToGroundClicked();
+};
\ No newline at end of file
diff --git a/Engine/Source/Editor/VREditor/UI/VREditorRadialMenuItem.cpp b/Engine/Source/Editor/VREditor/UI/VREditorRadialMenuItem.cpp
new file mode 100644
--- /dev/null
+++ b/Engine/Source/Editor/VREditor/UI/VREditorRadialMenuItem.cpp
@@ -0,0 +1,36 @@
+// Copyright 1998-2015 Epic Games, Inc. All Rights Reserved.
+
+#include "VREditorModule.h"
+#include "VREditorRadialMenuItem.h"
+
+UVREditorRadialMenuItem::UVREditorRadialMenuItem( const FObjectInitializer& ObjectInitializer )
+	: Super( ObjectInitializer )
+{
+
+}
+
+float UVREditorRadialMenuItem::GetAngle() const
+{
+	return FRotator::NormalizeAxis( FMath::RadiansToDegrees(FMath::Atan2( RenderTransformPivot.X - RenderTransform.Translation.X,  RenderTransformPivot.X - RenderTransform.Translation.Y ) ) );
+}
+
+void UVREditorRadialMenuItem::OnEnterHover_Implementation()
+{
+
+}
+
+void UVREditorRadialMenuItem::OnLeaveHover_Implementation()
+{
+	
+}
+
+void UVREditorRadialMenuItem::OnPressed_Implementation()
+{
+	OnPressedDelegate.Broadcast();
+}
+
+void UVREditorRadialMenuItem::SetLabel_Implementation(const FString& Text)
+{
+
+}
+
diff --git a/Engine/Source/Editor/VREditor/UI/VREditorRadialMenuItem.h b/Engine/Source/Editor/VREditor/UI/VREditorRadialMenuItem.h
new file mode 100644
--- /dev/null
+++ b/Engine/Source/Editor/VREditor/UI/VREditorRadialMenuItem.h
@@ -0,0 +1,45 @@
+// Copyright 1998-2015 Epic Games, Inc. All Rights Reserved.
+
+#pragma once
+
+
+#include "VREditorBaseUserWidget.h"
+#include "VREditorRadialMenuItem.generated.h"
+
+DECLARE_MULTICAST_DELEGATE(FVRQuickMenuItemPressedCallback)
+
+/**
+ * A simple quick-access menu for VR editing, with frequently-used features
+ */
+UCLASS( Blueprintable )
+class UVREditorRadialMenuItem : public UVREditorBaseUserWidget
+{
+	GENERATED_BODY()
+
+public:
+	/** Default constructor that sets up CDO properties */
+	UVREditorRadialMenuItem( const FObjectInitializer& ObjectInitializer );
+
+	/** Update UI visuals */
+	UFUNCTION(BlueprintNativeEvent, meta = (DisplayName = "On Enter Hover"))
+	void OnEnterHover();
+
+	/** Update UI visuals */
+	UFUNCTION(BlueprintNativeEvent, meta = (DisplayName = "On Leave Hover"))
+	void OnLeaveHover();
+
+	/** Broadcasts the PressedDelegate */
+	UFUNCTION(BlueprintNativeEvent, meta = (DisplayName = "On Pressed"))
+	void OnPressed();
+
+	/** Changes the text for this widget */
+	UFUNCTION(BlueprintNativeEvent, meta = (DisplayName = "Set Label "))
+	void SetLabel( const FString& Text );
+
+	/** Gets the angle from the pivot */
+	UFUNCTION( BlueprintCallable, meta = (DisplayName = "Get Angle "), Category = "RadialMenu" )
+	float GetAngle() const;
+	
+	/** Broadcast when pressed */
+	FVRQuickMenuItemPressedCallback OnPressedDelegate;
+};
\ No newline at end of file
diff --git a/Engine/Source/Editor/VREditor/UI/VREditorUISystem.cpp b/Engine/Source/Editor/VREditor/UI/VREditorUISystem.cpp
new file mode 100644
--- /dev/null
+++ b/Engine/Source/Editor/VREditor/UI/VREditorUISystem.cpp
@@ -0,0 +1,1521 @@
+// Copyright 1998-2015 Epic Games, Inc. All Rights Reserved.
+
+#include "VREditorModule.h"
+#include "VREditorUISystem.h"
+#include "VREditorMode.h"
+#include "VREditorFloatingUI.h"
+#include "VREditorDockableWindow.h"
+#include "VREditorQuickMenu.h"
+#include "VREditorRadialMenu.h"
+#include "VREditorRadialMenuItem.h"
+#include "IHeadMountedDisplay.h"
+#include "ViewportWorldInteraction.h"
+#include "ViewportInteractor.h"
+#include "VREditorInteractor.h"
+#include "VREditorMotionControllerInteractor.h"
+
+// UI
+#include "WidgetComponent.h"
+#include "VREditorWidgetComponent.h"
+
+// Content Browser
+#include "Editor/ContentBrowser/Public/ContentBrowserModule.h"
+#include "Editor/ContentBrowser/Public/IContentBrowserSingleton.h"
+
+// World Outliner
+#include "Editor/SceneOutliner/Public/SceneOutliner.h"
+
+// Actor Details, Modes
+#include "LevelEditor.h"
+
+// World Settings
+#include "Editor/PropertyEditor/Public/PropertyEditorModule.h"
+#include "Editor/PropertyEditor/Public/IDetailsView.h"
+
+#include "SLevelViewport.h"
+#include "SScaleBox.h"
+#include "SDPIScaler.h"
+#include "SWidget.h"
+#include "SDockTab.h"
+#include "SColorPicker.h"
+
+namespace VREd
+{
+	static FAutoConsoleVariable ContentBrowserUIResolutionX( TEXT( "VREd.ContentBrowserUIResolutionX" ), 1920, TEXT( "Horizontal resolution to use for content browser UI render targets" ) );
+	static FAutoConsoleVariable ContentBrowserUIResolutionY( TEXT( "VREd.ContentBrowserUIResolutionY" ), 1200, TEXT( "Vertical resolution to use for content browser UI render targets" ) );
+	static FAutoConsoleVariable DefaultEditorUIResolutionX( TEXT( "VREd.DefaultEditorUIResolutionX" ), 1024, TEXT( "Horizontal resolution to use for VR editor UI render targets" ) );
+	static FAutoConsoleVariable DefaultEditorUIResolutionY( TEXT( "VREd.DefaultEditorUIResolutionY" ), 1024, TEXT( "Vertical resolution to use for VR editor UI render targets" ) );
+	static FAutoConsoleVariable QuickMenuUIResolutionX( TEXT( "VREd.QuickMenuUIResolutionX" ), 1024, TEXT( "Horizontal resolution to use for Quick Menu VR UI render targets" ) );
+	static FAutoConsoleVariable QuickMenuUIResolutionY( TEXT( "VREd.QuickMenuUIResolutionY" ), 900, TEXT( "Vertical resolution to use for Quick Menu VR UI render targets" ) );
+	static FAutoConsoleVariable ContentBrowserUISize( TEXT( "VREd.ContentBrowserUISize" ), 60.0f, TEXT( "How big content browser UIs should be" ) );
+	static FAutoConsoleVariable EditorUISize( TEXT( "VREd.EditorUISize" ), 50.0f, TEXT( "How big editor UIs should be" ) );
+	static FAutoConsoleVariable ContentBrowserUIScale( TEXT( "VREd.ContentBrowserUIScale" ), 2.0f, TEXT( "How much to scale up (or down) the content browser for VR" ) );
+	static FAutoConsoleVariable EditorUIScale( TEXT( "VREd.EditorUIScale" ), 2.0f, TEXT( "How much to scale up (or down) editor UIs for VR" ) );
+	static FAutoConsoleVariable EditorUIOffsetFromHand( TEXT( "VREd.EditorUIOffsetFromHand" ), 12.0f, TEXT( "How far to offset editor UIs from the origin of the hand mesh" ) );
+	static FAutoConsoleVariable AssetEditorUIResolutionX( TEXT( "VREd.AssetEditorUIResolutionX" ), 1920, TEXT( "Horizontal resolution to use for VR editor asset editor UI render targets" ) );
+	static FAutoConsoleVariable AssetEditorUIResolutionY( TEXT( "VREd.AssetEditorUIResolutionY" ), 1080, TEXT( "Vertical resolution to use for VR editor asset editor UI render targets" ) );
+	static FAutoConsoleVariable RadialMenuFadeDelay( TEXT( "VREd.RadialMenuFadeDelay" ), 0.2f, TEXT( "The delay for the radial menu after selecting a button" ) );
+	static FAutoConsoleVariable UIAbsoluteScrollSpeed( TEXT( "VREd.UIAbsoluteScrollSpeed" ), 8.0f, TEXT( "How fast the UI scrolls when dragging the touchpad" ) );
+	static FAutoConsoleVariable UIRelativeScrollSpeed( TEXT( "VREd.UIRelativeScrollSpeed" ), 0.75f, TEXT( "How fast the UI scrolls when holding an analog stick" ) );
+	static FAutoConsoleVariable MinUIScrollDeltaForInertia( TEXT( "VREd.MinUIScrollDeltaForInertia" ), 0.25f, TEXT( "Minimum amount of touch pad input before inertial UI scrolling kicks in" ) );
+	static FAutoConsoleVariable MinDockDragDistance( TEXT( "VREd.MinDockDragDistance" ), 10.0f, TEXT( "Minimum amount of distance needed to drag dock from hand" ) );
+	static FAutoConsoleVariable DoubleClickTime( TEXT( "VREd.DoubleClickTime" ), 0.25f, TEXT( "Minimum duration between clicks for a double click event to fire" ) );
+	static FAutoConsoleVariable UIPressHapticFeedbackStrength( TEXT( "VREd.UIPressHapticFeedbackStrength" ), 0.4f, TEXT( "Strenth of haptic when clicking on the UI" ) );
+	static FAutoConsoleVariable UIAssetEditorSummonedOnHandHapticFeedbackStrength( TEXT( "VREd.UIAssetEditorSummonedOnHandHapticFeedbackStrength" ), 1.0f, TEXT( "Strenth of haptic to play to remind a user which hand an asset editor was spawned on" ) );
+	static FAutoConsoleVariable MaxDockWindowSize( TEXT( "VREd.MaxDockWindowSize" ), 250, TEXT( "Maximum size for dockable windows" ) );
+	static FAutoConsoleVariable MinDockWindowSize( TEXT( "VREd.MinDockWindowSize" ), 40, TEXT( "Minimum size for dockable windows" ) );
+	static FAutoConsoleVariable DockSnapTime( TEXT( "VREd.DockSnapTime" ), 0.1f, TEXT( "Lerp time for snapping dock to hand and back" ) );
+	static FAutoConsoleVariable DockDragSpeed( TEXT( "VREd.DockDragSpeed" ), 0.05f, TEXT( "Dock dragging time delay" ) );
+	static FAutoConsoleVariable ColorPickerDockSpawnOffset( TEXT( "VREd.ColorPickerDockSpawnOffset" ),  3.f, TEXT( "Offset of where the color picker spawns" ) );
+
+	// Tutorial UI commands
+	static FAutoConsoleVariable TutorialUIResolutionX( TEXT( "VREd.TutorialUI.Resolution.X" ), 1920, TEXT( "The X resolution for the tutorial UI panel" ) );
+	static FAutoConsoleVariable TutorialUIResolutionY( TEXT( "VREd.TutorialUI.Resolution.Y" ), 1080, TEXT( "The Y resolution for the tutorial UI panel" ) );
+	static FAutoConsoleVariable TutorialUISize( TEXT( "VREd.TutorialUI.Size" ), 200, TEXT( "The room space size for the tutorial UI panel" ) );
+	static FAutoConsoleVariable TutorialUIYaw( TEXT( "VREd.TutorialUI.Yaw" ), 180, TEXT( "The yaw rotation for the tutorial UI panel" ) );
+	static FAutoConsoleVariable TutorialUIPitch( TEXT( "VREd.TutorialUI.Pitch" ), 45, TEXT( "The pitch rotation for the tutorial UI panel" ) );
+	static FAutoConsoleVariable TutorialUILocationX( TEXT( "VREd.TutorialUI.Location.X" ), 200, TEXT( "The X location for the tutorial UI panel" ) );
+	static FAutoConsoleVariable TutorialUILocationY( TEXT( "VREd.TutorialUI.Location.Y" ), 0, TEXT( "The Y location for the tutorial UI panel" ) );
+	static FAutoConsoleVariable TutorialUILocationZ( TEXT( "VREd.TutorialUI.Location.Z" ), 40, TEXT( "The Z location for the tutorial UI panel" ) );
+
+	static const FTransform DefaultModesTransform(			FRotator( -30,  130, 0 ), FVector( 40, -45,  40 ), FVector( 1 ) );	// Top left
+	static const FTransform DefaultWorldOutlinerTransform(	FRotator( -30, -130, 0 ), FVector( 40,  45,  40 ), FVector( 1 ) );	// Top right
+	static const FTransform DefaultContentBrowserTransform( FRotator(  20,  160, 0 ), FVector( 50, -65, -40 ), FVector( 1 ) );  // Bottom left
+	static const FTransform DefaultActorDetailsTransform(	FRotator(  20, -160, 0 ), FVector( 50,  65, -40 ), FVector( 1 ) );	// Bottom right
+
+	static const FTransform DefaultColorPickerTransform(	FRotator( -10, 180, 0), FVector( 30, 35, 0 ), FVector( 1 ) );
+	
+}
+
+
+UVREditorUISystem::UVREditorUISystem( const FObjectInitializer& Initializer ) : 
+	Super( Initializer ),
+	VRMode( nullptr ),
+	FloatingUIs(),
+	QuickMenuUI( nullptr ),
+	QuickRadialMenu( nullptr ),
+	RadialMenuHideDelayTime( 0.0f ),
+	QuickMenuWidgetClass( nullptr ),
+	QuickRadialWidgetClass( nullptr ),
+	TutorialWidgetClass( nullptr ),
+	InteractorDraggingUI( nullptr ),
+	DraggingUIOffsetTransform( FTransform::Identity ),
+	DraggingUI( nullptr ),
+	ColorPickerUI( nullptr ),
+	bRefocusViewport(false),
+	StartDragUISound( nullptr ),
+	StopDragUISound( nullptr ),
+	HideUISound( nullptr ),
+	ShowUISound( nullptr ),
+	bDraggedDockFromHandPassedThreshold( false ),
+	LastDraggingHoverLocation( FVector::ZeroVector ),
+	bSetDefaultLayout( true )
+{
+	EditorUIPanels.SetNumZeroed( (int32)EEditorUIPanel::TotalCount );
+	
+	// Set default layout transform in correct order of enum
+	DefaultWindowTransforms.Add( VREd::DefaultContentBrowserTransform );
+	DefaultWindowTransforms.Add( VREd::DefaultWorldOutlinerTransform );
+	DefaultWindowTransforms.Add( VREd::DefaultActorDetailsTransform );
+	DefaultWindowTransforms.Add( VREd::DefaultModesTransform );
+}
+
+void UVREditorUISystem::Init()
+{
+	// Register to find out about VR events
+	GetOwner().GetWorldInteraction().OnViewportInteractionInputAction().AddUObject( this, &UVREditorUISystem::OnVRAction );
+	GetOwner().GetWorldInteraction().OnViewportInteractionHoverUpdate().AddUObject( this, &UVREditorUISystem::OnVRHoverUpdate );
+
+	QuickMenuWidgetClass = LoadClass<UVREditorQuickMenu>( nullptr, TEXT( "/Engine/VREditor/UI/VRQuickMenu.VRQuickMenu_C" ) );
+	check( QuickMenuWidgetClass != nullptr );
+
+	QuickRadialWidgetClass = LoadClass<UVREditorRadialMenu>( nullptr, TEXT( "/Engine/VREditor/UI/VRRadialQuickMenu.VRRadialQuickMenu_C" ) );
+	check( QuickRadialWidgetClass != nullptr );
+
+	TutorialWidgetClass = LoadClass<UVREditorBaseUserWidget>( nullptr, TEXT( "/Engine/VREditor/Tutorial/UI_VR_Tutorial_00.UI_VR_Tutorial_00_C" ) );
+	check( TutorialWidgetClass != nullptr );
+
+	// Load sounds
+	StartDragUISound = LoadObject<USoundCue>( nullptr, TEXT( "/Engine/VREditor/Sounds/VR_negative_Cue" ) );
+	check( StartDragUISound != nullptr );
+
+	StopDragUISound = LoadObject<USoundCue>( nullptr, TEXT( "/Engine/VREditor/Sounds/VR_negative_Cue" ) );
+	check( StopDragUISound != nullptr );
+
+	HideUISound = LoadObject<USoundCue>( nullptr, TEXT( "/Engine/VREditor/Sounds/VR_close_Cue" ) );
+	check( HideUISound != nullptr );
+
+	ShowUISound = LoadObject<USoundCue>( nullptr, TEXT( "/Engine/VREditor/Sounds/VR_open_Cue" ) );
+	check( ShowUISound != nullptr );
+
+	
+	// Create all of our UI panels
+	CreateUIs();
+
+
+	// Bind the color picker creation & destruction overrides
+	SColorPicker::OnColorPickerNonModalCreateOverride.BindUObject(this, &UVREditorUISystem::CreateVRColorPicker);
+	SColorPicker::OnColorPickerDestroyOverride.BindUObject(this, &UVREditorUISystem::DestroyVRColorPicker);
+}
+
+
+void UVREditorUISystem::Shutdown()
+{
+	if ( VRMode != nullptr && &VRMode->GetWorldInteraction() != nullptr )
+	{
+		GetOwner().GetWorldInteraction().OnViewportInteractionInputAction().RemoveAll( this );
+		GetOwner().GetWorldInteraction().OnViewportInteractionHoverUpdate().RemoveAll( this );
+	}
+
+	// Unbind the color picker creation & destruction overrides
+	SColorPicker::OnColorPickerNonModalCreateOverride.Unbind();
+	SColorPicker::OnColorPickerDestroyOverride.Unbind();
+
+	CleanUpActorsBeforeMapChangeOrSimulate();
+
+	DefaultWindowTransforms.Empty();
+
+	QuickMenuWidgetClass = nullptr;
+	QuickRadialWidgetClass = nullptr;
+	TutorialWidgetClass = nullptr;
+	StartDragUISound = nullptr;
+	StopDragUISound = nullptr;
+	HideUISound = nullptr;
+	ShowUISound = nullptr;
+	VRMode = nullptr;
+	DraggingUI = nullptr;
+	ColorPickerUI = nullptr;
+}
+
+void UVREditorUISystem::OnVRAction( FEditorViewportClient& ViewportClient, UViewportInteractor* Interactor,
+	const FViewportActionKeyInput& Action, bool& bOutIsInputCaptured, bool& bWasHandled )
+{
+	UVREditorMotionControllerInteractor* VREditorInteractor = Cast<UVREditorMotionControllerInteractor>( Interactor );
+	if( VREditorInteractor && 
+		!( VREditorInteractor->GetDraggingMode() == EViewportInteractionDraggingMode::ActorsWithGizmo ||
+		  VREditorInteractor->GetDraggingMode() == EViewportInteractionDraggingMode::Interactable ) )
+	{
+		if ( !bWasHandled )
+		{
+			if ( Action.ActionType == VRActionTypes::ConfirmRadialSelection )
+			{
+				const EViewportInteractionDraggingMode DraggingMode = Interactor->GetDraggingMode();
+				if ( Action.Event == IE_Pressed && 
+					 ( DraggingMode != EViewportInteractionDraggingMode::ActorsFreely && DraggingMode != EViewportInteractionDraggingMode::World && 
+					   ( DraggingMode != EViewportInteractionDraggingMode::AssistingDrag && Interactor->GetOtherInteractor() != nullptr && Interactor->GetOtherInteractor()->GetDraggingMode() != EViewportInteractionDraggingMode::World ) ) )
+				{
+					UVREditorRadialMenu* RadialMenu = QuickRadialMenu->GetUserWidget<UVREditorRadialMenu>();
+					if ( IsShowingRadialMenu( VREditorInteractor ) )
+					{
+						RadialMenu->Update( VREditorInteractor );
+					}
+
+					RadialMenu->SelectCurrentItem();
+					bWasHandled = true;
+				}
+			}
+			else if ( Action.ActionType == ViewportWorldActionTypes::SelectAndMove_LightlyPressed )
+			{
+				FVector LaserPointerStart, LaserPointerEnd;
+				if ( VREditorInteractor->GetLaserPointer( LaserPointerStart, LaserPointerEnd ) )
+				{
+					FHitResult HitResult = VREditorInteractor->GetHitResultFromLaserPointer();
+					if ( HitResult.Actor.IsValid() )
+					{
+						// Only allow clicks to our own widget components
+						UWidgetComponent* WidgetComponent = Cast<UWidgetComponent>( HitResult.GetComponent() );
+						if ( WidgetComponent != nullptr && IsWidgetAnEditorUIWidget( WidgetComponent ) )
+						{
+							// Always mark the event as handled so that the editor doesn't try to select the widget component
+							bWasHandled = true;
+
+							if ( Action.Event != IE_Repeat )
+							{
+								// If the Modifier button is held down, treat this like a right click instead of a left click
+								const bool bIsRightClicking =
+									( Action.Event == IE_Pressed && VREditorInteractor->IsModifierPressed() ) ||
+									( Action.Event == IE_Released && VREditorInteractor->IsRightClickingOnUI() );
+
+								FVector2D LastLocalHitLocation = WidgetComponent->GetLastLocalHitLocation();
+
+								FVector2D LocalHitLocation;
+								WidgetComponent->GetLocalHitLocation( HitResult.ImpactPoint, LocalHitLocation );
+
+								// If we weren't already hovering over this widget, then we'll reset the last hit location
+								if ( WidgetComponent != VREditorInteractor->GetHoveringOverWidgetComponent() )
+								{
+									LastLocalHitLocation = LocalHitLocation;
+
+									if ( UVREditorWidgetComponent* VRWidgetComponent = Cast<UVREditorWidgetComponent>( VREditorInteractor->GetHoveringOverWidgetComponent() ) )
+									{
+										VRWidgetComponent->SetIsHovering( false );
+									}
+								}
+
+								FWidgetPath WidgetPathUnderFinger = FWidgetPath( WidgetComponent->GetHitWidgetPath( HitResult.ImpactPoint, /*bIgnoreEnabledStatus*/ false ) );
+								if ( WidgetPathUnderFinger.IsValid() )
+								{
+									TSet<FKey> PressedButtons;
+									if ( Action.Event == IE_Pressed )
+									{
+										PressedButtons.Add( bIsRightClicking ? EKeys::RightMouseButton : EKeys::LeftMouseButton );
+									}
+
+									FPointerEvent PointerEvent(
+										1 + (uint8)VREditorInteractor->GetControllerSide(),
+										LocalHitLocation,
+										LastLocalHitLocation,
+										PressedButtons,
+										bIsRightClicking ? EKeys::RightMouseButton : EKeys::LeftMouseButton,
+										0.0f,	// Wheel delta
+										FModifierKeysState() );
+
+									VREditorInteractor->SetHoveringOverWidgetComponent( WidgetComponent );
+
+									if ( UVREditorWidgetComponent* VRWidgetComponent = Cast<UVREditorWidgetComponent>( VREditorInteractor->GetHoveringOverWidgetComponent() ) )
+									{
+										VRWidgetComponent->SetIsHovering( true );
+									}
+
+									FReply Reply = FReply::Unhandled();
+									if ( Action.Event == IE_Pressed )
+									{
+										const double CurrentTime = FPlatformTime::Seconds();
+										if (CurrentTime - VREditorInteractor->GetLastClickReleaseTime() <= VREd::DoubleClickTime->GetFloat())
+										{
+											// Trigger a double click event!
+											Reply = FSlateApplication::Get().RoutePointerDoubleClickEvent(WidgetPathUnderFinger, PointerEvent);
+										}
+										else
+										{
+											Reply = FSlateApplication::Get().RoutePointerDownEvent(WidgetPathUnderFinger, PointerEvent);
+										}
+
+										// In case of selecting a level in the content browser the VREditormode is closed, this makes sure nothing happens after that.
+										if( !IVREditorModule::Get().IsVREditorModeActive() )
+										{
+											bWasHandled = true;
+											return;
+										}
+										
+										VREditorInteractor->SetIsClickingOnUI( true );
+										VREditorInteractor->SetIsRightClickingOnUI( bIsRightClicking );
+										VREditorInteractor->SetLastClickReleaseTime( CurrentTime );
+										bOutIsInputCaptured = true;
+
+										// Play a haptic effect on press
+										VREditorInteractor->PlayHapticEffect( VREd::UIPressHapticFeedbackStrength->GetFloat() );
+									}
+									else if ( Action.Event == IE_Released )
+									{
+										Reply = FSlateApplication::Get().RoutePointerUpEvent( WidgetPathUnderFinger, PointerEvent );
+									}
+								}
+							}
+						}
+					}
+				}
+
+				if ( Action.Event == IE_Released )
+				{
+					bool bWasRightClicking = false;
+					if ( VREditorInteractor->IsClickingOnUI() )
+					{
+						if ( VREditorInteractor->IsRightClickingOnUI() )
+						{
+							bWasRightClicking = true;
+						}
+						VREditorInteractor->SetIsClickingOnUI( false );
+						VREditorInteractor->SetIsRightClickingOnUI( false );
+						bOutIsInputCaptured = false;
+					}
+
+					if ( !bWasHandled )
+					{
+						TSet<FKey> PressedButtons;
+						FPointerEvent PointerEvent(
+							1 + (uint8)VREditorInteractor->GetControllerSide(),
+							FVector2D::ZeroVector,
+							FVector2D::ZeroVector,
+							PressedButtons,
+							bWasRightClicking ? EKeys::RightMouseButton : EKeys::LeftMouseButton,
+							0.0f,	// Wheel delta
+							FModifierKeysState() );
+
+						FWidgetPath EmptyWidgetPath;
+
+						VREditorInteractor->SetIsClickingOnUI( false );
+						VREditorInteractor->SetIsRightClickingOnUI( false );
+						FReply Reply = FSlateApplication::Get().RoutePointerUpEvent( EmptyWidgetPath, PointerEvent );
+					}
+				}
+			}
+		}
+	}
+}
+
+
+void UVREditorUISystem::OnVRHoverUpdate( FEditorViewportClient& ViewportClient, UViewportInteractor* Interactor, FVector& HoverImpactPoint, bool& bWasHandled )
+{
+	UVREditorMotionControllerInteractor* VREditorInteractor = Cast<UVREditorMotionControllerInteractor>( Interactor );
+	if( VREditorInteractor != nullptr )
+	{
+		if( !bWasHandled && Interactor->GetDraggingMode() != EViewportInteractionDraggingMode::Interactable )
+		{
+			FVector LaserPointerStart, LaserPointerEnd;
+			if( Interactor->GetLaserPointer( LaserPointerStart, LaserPointerEnd ) )
+			{
+				FHitResult HitResult = Interactor->GetHitResultFromLaserPointer();
+				if( HitResult.Actor.IsValid() )
+				{
+					// Only allow clicks to our own widget components
+					UWidgetComponent* WidgetComponent = Cast<UWidgetComponent>( HitResult.GetComponent() );
+					if( WidgetComponent != nullptr && IsWidgetAnEditorUIWidget( WidgetComponent ) )
+					{
+						FVector2D LastLocalHitLocation = WidgetComponent->GetLastLocalHitLocation();
+
+						FVector2D LocalHitLocation;
+						WidgetComponent->GetLocalHitLocation( HitResult.ImpactPoint, LocalHitLocation );
+
+						// If we weren't already hovering over this widget, then we'll reset the last hit location
+						if( WidgetComponent != VREditorInteractor->GetHoveringOverWidgetComponent() )
+						{
+							LastLocalHitLocation = LocalHitLocation;
+
+							if ( UVREditorWidgetComponent* VRWidgetComponent = Cast<UVREditorWidgetComponent>( VREditorInteractor->GetHoveringOverWidgetComponent() ) )
+							{
+								VRWidgetComponent->SetIsHovering(false);
+							}
+						}
+
+						// @todo vreditor UI: There is a CursorRadius optional arg that we migth want to make use of wherever we call GetHitWidgetPath()!
+						FWidgetPath WidgetPathUnderFinger = FWidgetPath( WidgetComponent->GetHitWidgetPath( HitResult.ImpactPoint, /*bIgnoreEnabledStatus*/ false ) );
+						if ( WidgetPathUnderFinger.IsValid() )
+						{
+							TSet<FKey> PressedButtons;
+							FPointerEvent PointerEvent(
+								1 + (uint8)VREditorInteractor->GetControllerSide(),
+								LocalHitLocation,
+								LastLocalHitLocation,
+								LocalHitLocation - LastLocalHitLocation,
+								PressedButtons,
+								FModifierKeysState());
+
+							FSlateApplication::Get().RoutePointerMoveEvent(WidgetPathUnderFinger, PointerEvent, false);
+
+							bWasHandled = true;
+							HoverImpactPoint = HitResult.ImpactPoint;
+							VREditorInteractor->SetHoveringOverWidgetComponent( WidgetComponent );
+							VREditorInteractor->SetIsHoveringOverUI( true );
+
+							if ( UVREditorWidgetComponent* VRWidgetComponent = Cast<UVREditorWidgetComponent>(WidgetComponent) )
+							{
+								VRWidgetComponent->SetIsHovering(true);
+							}
+
+							// Route the mouse scrolling
+							if ( VREditorInteractor->IsTrackpadPositionValid( 1 ) )
+							{
+								const bool bIsAbsolute = ( GetOwner().GetHMDDeviceType() == EHMDDeviceType::DT_SteamVR );
+
+								float ScrollDelta = 0.0f;
+								if( VREditorInteractor->IsTouchingTrackpad() || !bIsAbsolute )
+								{
+									if( bIsAbsolute )
+									{
+										const float ScrollSpeed = VREd::UIAbsoluteScrollSpeed->GetFloat();
+										ScrollDelta = ( VREditorInteractor->GetTrackpadPosition().Y - VREditorInteractor->GetLastTrackpadPosition().Y ) * ScrollSpeed;
+									}
+									else
+									{
+										const float ScrollSpeed = VREd::UIRelativeScrollSpeed->GetFloat();
+										ScrollDelta = VREditorInteractor->GetTrackpadPosition().Y * ScrollSpeed;
+									}
+								}
+
+								// If using a trackpad (Vive), invert scroll direction so that it feels more like scrolling on a mobile device
+								if( GetOwner().GetHMDDeviceType() == EHMDDeviceType::DT_SteamVR )
+								{
+									ScrollDelta *= -1.0f;
+								}
+
+								if( !FMath::IsNearlyZero( ScrollDelta ) )
+								{
+									FPointerEvent MouseWheelEvent(
+										1 + (uint8)VREditorInteractor->GetControllerSide(),
+										LocalHitLocation,
+										LastLocalHitLocation,
+										PressedButtons,
+										EKeys::MouseWheelAxis,
+										ScrollDelta,
+										FModifierKeysState() );
+
+									FSlateApplication::Get().RouteMouseWheelOrGestureEvent(WidgetPathUnderFinger, MouseWheelEvent, nullptr);
+
+									VREditorInteractor->SetUIScrollVelocity( 0.0f );
+
+									// Don't apply inertia unless the user dragged a decent amount this frame
+									if( bIsAbsolute && FMath::Abs( ScrollDelta ) >= VREd::MinUIScrollDeltaForInertia->GetFloat() )
+									{
+										// Don't apply inertia if our data is sort of old
+										const FTimespan CurrentTime = FTimespan::FromSeconds( FPlatformTime::Seconds() );
+										if( CurrentTime - VREditorInteractor->GetLastTrackpadPositionUpdateTime() < FTimespan::FromSeconds( 1.0f / 30.0f ) )
+										{
+											//GWarn->Logf( TEXT( "INPUT: UIScrollVelocity=%0.2f" ), ScrollDelta );
+											VREditorInteractor->SetUIScrollVelocity( ScrollDelta );
+										}
+									}
+								}
+							}
+							else
+							{
+								if( !FMath::IsNearlyZero( VREditorInteractor->GetUIScrollVelocity() ) )
+								{
+									// Apply UI scrolling inertia
+									const float ScrollDelta = VREditorInteractor->GetUIScrollVelocity();
+									{
+										FPointerEvent MouseWheelEvent(
+											1 + (uint8)VREditorInteractor->GetControllerSide(),
+											LocalHitLocation,
+											LastLocalHitLocation,
+											PressedButtons,
+											EKeys::MouseWheelAxis,
+											ScrollDelta,
+											FModifierKeysState() );
+
+										FSlateApplication::Get().RouteMouseWheelOrGestureEvent( WidgetPathUnderFinger, MouseWheelEvent, nullptr );
+									}
+
+									// Apply damping
+									FVector ScrollVelocityVector( VREditorInteractor->GetUIScrollVelocity(), 0.0f, 0.0f );
+									const bool bVelocitySensitive = false;
+									GetOwner().GetWorldInteraction().ApplyVelocityDamping( ScrollVelocityVector, bVelocitySensitive );
+									VREditorInteractor->SetUIScrollVelocity( ScrollVelocityVector.X );
+								}
+								else
+								{
+									VREditorInteractor->SetUIScrollVelocity( 0.0f );
+								}
+							}
+						}
+					}
+				}
+			}
+		}
+
+		// If nothing was hovered, make sure we tell Slate about that
+		if( !bWasHandled && VREditorInteractor->GetHoveringOverWidgetComponent() != nullptr )
+		{
+			if ( UVREditorWidgetComponent* VRWidgetComponent = Cast<UVREditorWidgetComponent>( VREditorInteractor->GetHoveringOverWidgetComponent() ) )
+			{
+				VRWidgetComponent->SetIsHovering( false );
+			}
+
+			const FVector2D LastLocalHitLocation = VREditorInteractor->GetHoveringOverWidgetComponent()->GetLastLocalHitLocation();
+			VREditorInteractor->SetHoveringOverWidgetComponent( nullptr );
+
+			TSet<FKey> PressedButtons;
+			FPointerEvent PointerEvent(
+				1 + (uint8)VREditorInteractor->GetControllerSide(),
+				LastLocalHitLocation,
+				LastLocalHitLocation,
+				FVector2D::ZeroVector,
+				PressedButtons,
+				FModifierKeysState() );
+
+			FSlateApplication::Get().RoutePointerMoveEvent( FWidgetPath(), PointerEvent, false );
+		}
+	}
+}
+
+void UVREditorUISystem::Tick( FEditorViewportClient* ViewportClient, const float DeltaTime )
+{
+	if ( bRefocusViewport )
+	{
+		bRefocusViewport = false;
+		FSlateApplication::Get().SetUserFocus(0, ViewportClient->GetEditorViewportWidget());
+	}
+
+	// Figure out if one hand is "aiming toward" the other hand.  We'll fade in a UI on the hand being
+	// aimed at when the user does this.
+	if( QuickMenuUI != nullptr )
+	{
+		UVREditorInteractor* HandInteractorWithQuickMenu = nullptr;
+		if( QuickMenuUI->IsUIVisible() )
+		{
+			HandInteractorWithQuickMenu = GetOwner().GetHandInteractor( QuickMenuUI->GetDockedTo() == AVREditorFloatingUI::EDockedTo::LeftArm ? EControllerHand::Left : EControllerHand::Right );
+		}
+
+		const float WorldScaleFactor = GetOwner().GetWorldScaleFactor();
+
+		UViewportInteractor* HandInteractorThatNeedsQuickMenu = nullptr;
+		for( UViewportInteractor* Interactor : GetOwner().GetWorldInteraction().GetInteractors() )
+		{
+			bool bShowQuickMenuOnArm = false;
+
+			UViewportInteractor* OtherInteractor = Interactor->GetOtherInteractor();
+			if( OtherInteractor != nullptr )
+			{
+				// @todo vreditor tweak: Weird to hard code this here. Probably should be an accessor on the hand itself, and based on the actual device type
+				const FTransform UICapsuleTransform = OtherInteractor->GetTransform();
+				const FVector UICapsuleStart = FVector( -9.0f, 0.0f, 0.0f ) * WorldScaleFactor;
+				const FVector UICapsuleEnd = FVector( -18.0f, 0.0f, 0.0f ) * WorldScaleFactor;
+				const float UICapsuleLocalRadius = 6.0f * WorldScaleFactor;
+				const float MinDistanceToUICapsule = 8.0f * WorldScaleFactor; // @todo vreditor tweak
+				const FVector UIForwardVector = FVector::UpVector;
+				const float MinDotForAimingAtOtherHand = 0.25f;	// @todo vreditor tweak
+
+				if( GetOwner().IsHandAimingTowardsCapsule( Interactor, UICapsuleTransform, UICapsuleStart, UICapsuleEnd, UICapsuleLocalRadius, MinDistanceToUICapsule, UIForwardVector, MinDotForAimingAtOtherHand ) )
+				{
+					bShowQuickMenuOnArm = true;
+				}
+
+				if( bShowQuickMenuOnArm )
+				{
+					HandInteractorThatNeedsQuickMenu = OtherInteractor;
+				}
+			}
+		}
+		
+		// If we don't need a quick menu, or if a different hand needs to spawn it, then kill the existing menu
+		if( QuickMenuUI->IsUIVisible() && HandInteractorWithQuickMenu != nullptr && 
+			( HandInteractorThatNeedsQuickMenu == nullptr || HandInteractorThatNeedsQuickMenu != HandInteractorWithQuickMenu )  )
+		{
+			// Despawn
+			HandInteractorWithQuickMenu->SetHasUIOnForearm( false );
+			QuickMenuUI->ShowUI( false );
+		}
+
+		if( HandInteractorThatNeedsQuickMenu != nullptr && !QuickMenuUI->IsUIVisible() )
+		{
+			UVREditorMotionControllerInteractor* MotionControllerHandInteractorThatNeedsQuickMenu = Cast<UVREditorMotionControllerInteractor>( HandInteractorThatNeedsQuickMenu );
+			if ( MotionControllerHandInteractorThatNeedsQuickMenu )
+			{
+				const AVREditorFloatingUI::EDockedTo DockTo = ( MotionControllerHandInteractorThatNeedsQuickMenu->GetControllerSide() == EControllerHand::Left ) ? AVREditorFloatingUI::EDockedTo::LeftArm : AVREditorFloatingUI::EDockedTo::RightArm;
+				QuickMenuUI->SetDockedTo( DockTo );
+				QuickMenuUI->ShowUI( true );
+				MotionControllerHandInteractorThatNeedsQuickMenu->SetHasUIOnForearm( true );
+			}
+		}
+	}
+
+	// Close the radial menu if it was not updated for a while
+	RadialMenuHideDelayTime += DeltaTime;
+	if( RadialMenuHideDelayTime >= VREd::RadialMenuFadeDelay->GetFloat() )
+	{
+		UVREditorInteractor* Interactor = GetOwner().GetHandInteractor( QuickRadialMenu->GetDockedTo() == AVREditorFloatingUI::EDockedTo::LeftHand ? EControllerHand::Left : EControllerHand::Right );
+		HideRadialMenu( Interactor );
+	}
+	
+	// Tick all of our floating UIs
+	for( AVREditorFloatingUI* FloatingUIPtr : FloatingUIs )
+	{
+		if( FloatingUIPtr != nullptr )
+		{
+			FloatingUIPtr->TickManually( DeltaTime );
+		}
+	}
+}
+
+void UVREditorUISystem::Render( const FSceneView* SceneView, FViewport* Viewport, FPrimitiveDrawInterface* PDI )
+{
+	// ...
+}
+
+void UVREditorUISystem::CreateUIs()
+{
+	const FIntPoint DefaultResolution( VREd::DefaultEditorUIResolutionX->GetInt(), VREd::DefaultEditorUIResolutionY->GetInt() );
+
+	const bool bIsVREditorDemo = FParse::Param( FCommandLine::Get(), TEXT( "VREditorDemo" ) );	// @todo vreditor: Remove this when no longer needed
+																								
+	// @todo vreditor: Tweak
+	if ( QuickMenuWidgetClass != nullptr )
+	{
+		const FIntPoint Resolution( VREd::QuickMenuUIResolutionX->GetInt(), VREd::QuickMenuUIResolutionY->GetInt() );
+
+		const bool bWithSceneComponent = false;
+		QuickMenuUI = GetOwner().SpawnTransientSceneActor< AVREditorFloatingUI >(TEXT("QuickMenu"), bWithSceneComponent);
+		QuickMenuUI->SetUMGWidget( *this, QuickMenuWidgetClass, Resolution, 30.0f, AVREditorFloatingUI::EDockedTo::Nothing );
+		QuickMenuUI->ShowUI( false );
+		QuickMenuUI->SetRelativeOffset( FVector( -11.0f, 0.0f, 3.0f ) );
+		FloatingUIs.Add( QuickMenuUI );
+	}
+
+	// Create the radial UI
+	if ( QuickRadialWidgetClass != nullptr )
+	{
+		TSharedPtr<SWidget> WidgetToDraw;
+		QuickRadialMenu = GetOwner().SpawnTransientSceneActor< AVREditorFloatingUI >(TEXT("QuickRadialmenu"), false);
+		QuickRadialMenu->SetUMGWidget( *this, QuickRadialWidgetClass, DefaultResolution, 40.0f, AVREditorFloatingUI::EDockedTo::Nothing );
+		QuickRadialMenu->ShowUI( false );
+		QuickRadialMenu->SetActorEnableCollision( false );
+		QuickRadialMenu->SetRelativeOffset( FVector( 10.0f, 0.0f, 10.f ) );
+		QuickRadialMenu->SetCollisionOnShow( false );
+		FloatingUIs.Add(QuickRadialMenu);
+	}
+
+	// Make some editor UIs!
+	{
+		{
+			const FIntPoint Resolution( VREd::ContentBrowserUIResolutionX->GetInt(), VREd::ContentBrowserUIResolutionY->GetInt() );
+
+			IContentBrowserSingleton& ContentBrowserSingleton = FModuleManager::Get().LoadModuleChecked<FContentBrowserModule>( "ContentBrowser" ).Get();;
+
+			// @todo vreditor UI: We've turned off a LOT of content browser features that are clunky to use in VR right (pop ups, text fields, etc.)
+			FContentBrowserConfig Config;
+			Config.bCanShowClasses = false;
+			Config.bCanShowRealTimeThumbnails = false;
+			Config.InitialAssetViewType = EAssetViewType::Tile;
+			Config.bCanShowDevelopersFolder = false;
+			Config.bCanShowFolders = false;
+			Config.bUsePathPicker = false;
+			Config.bCanShowFilters = false;
+			Config.bCanShowAssetSearch = false;
+			Config.bUseSourcesView = true;
+			Config.bExpandSourcesView = true;
+			Config.ThumbnailLabel = EThumbnailLabel::NoLabel;
+			Config.ThumbnailScale = 0.4f;
+
+			if( bIsVREditorDemo )
+			{
+				Config.bUsePathPicker = false;
+				Config.bCanShowFilters = false;
+				Config.bShowAssetPathTree = false;
+				Config.bAlwaysShowCollections = true;
+				Config.SelectedCollectionName.Name = FName( "Meshes" );	// @todo vreditor: hard-coded collection name and type
+				Config.SelectedCollectionName.Type = ECollectionShareType::CST_Shared;
+			}
+			else
+			{
+				Config.bCanShowFilters = true;
+				Config.bUsePathPicker = true;
+				Config.bShowAssetPathTree = true;
+				Config.bAlwaysShowCollections = false;
+			}
+
+			Config.bShowBottomToolbar = false;
+			Config.bCanShowLockButton = false;
+			TSharedRef<SWidget> ContentBrowser = ContentBrowserSingleton.CreateContentBrowser( "VRContentBrowser", nullptr, &Config );
+
+			TSharedRef<SWidget> WidgetToDraw =
+				SNew( SDPIScaler )
+				.DPIScale( VREd::ContentBrowserUIScale->GetFloat() )
+				[
+					ContentBrowser
+				]
+			;
+
+			const bool bWithSceneComponent = false;
+			AVREditorFloatingUI* ContentBrowserUI = GetOwner().SpawnTransientSceneActor< AVREditorDockableWindow >( TEXT( "ContentBrowserUI" ), bWithSceneComponent );
+			ContentBrowserUI->SetSlateWidget( *this, WidgetToDraw, Resolution, VREd::ContentBrowserUISize->GetFloat(), AVREditorFloatingUI::EDockedTo::Nothing );
+			ContentBrowserUI->ShowUI( false );
+			FloatingUIs.Add( ContentBrowserUI );
+
+			EditorUIPanels[ (int32)EEditorUIPanel::ContentBrowser ] = ContentBrowserUI;
+		}
+
+		{
+			FSceneOutlinerModule& SceneOutlinerModule = FModuleManager::Get().LoadModuleChecked<FSceneOutlinerModule>( "SceneOutliner" );
+
+			SceneOutliner::FInitializationOptions InitOptions;
+			InitOptions.Mode = ESceneOutlinerMode::ActorBrowsing;
+
+			TSharedRef<ISceneOutliner> SceneOutliner = SceneOutlinerModule.CreateSceneOutliner( InitOptions, FOnActorPicked() /* Not used for outliner when in browsing mode */ );
+			TSharedRef<SWidget> WidgetToDraw =
+				SNew( SDPIScaler )
+				.DPIScale( VREd::EditorUIScale->GetFloat() )
+				[
+					SceneOutliner
+				]
+			;
+
+			const bool bWithSceneComponent = false;
+			AVREditorFloatingUI* WorldOutlinerUI = GetOwner().SpawnTransientSceneActor< AVREditorDockableWindow >( TEXT( "WorldOutlinerUI" ), bWithSceneComponent );
+			WorldOutlinerUI->SetSlateWidget( *this, WidgetToDraw, DefaultResolution, VREd::EditorUISize->GetFloat(), AVREditorFloatingUI::EDockedTo::Nothing );
+			WorldOutlinerUI->ShowUI( false );
+			FloatingUIs.Add( WorldOutlinerUI );
+
+			EditorUIPanels[ (int32)EEditorUIPanel::WorldOutliner ] = WorldOutlinerUI;
+		}
+
+		{
+			const TSharedRef< ILevelEditor >& LevelEditor = FModuleManager::GetModuleChecked<FLevelEditorModule>( "LevelEditor" ).GetFirstLevelEditor().ToSharedRef();
+
+			const FName TabIdentifier = NAME_None;	// No tab for us!
+			TSharedRef<SWidget> ActorDetails = LevelEditor->CreateActorDetails( TabIdentifier );
+
+			TSharedRef<SWidget> WidgetToDraw =
+				SNew( SDPIScaler )
+				.DPIScale( VREd::EditorUIScale->GetFloat() )
+				[
+					ActorDetails
+				]
+			;
+
+			const bool bWithSceneComponent = false;
+			AVREditorFloatingUI* ActorDetailsUI = GetOwner().SpawnTransientSceneActor< AVREditorDockableWindow >( TEXT( "ActorDetailsUI" ), bWithSceneComponent );
+			ActorDetailsUI->SetSlateWidget( *this, WidgetToDraw, DefaultResolution, VREd::EditorUISize->GetFloat(), AVREditorFloatingUI::EDockedTo::Nothing );
+			ActorDetailsUI->ShowUI( false );
+			FloatingUIs.Add( ActorDetailsUI );
+
+			EditorUIPanels[ (int32)EEditorUIPanel::ActorDetails ] = ActorDetailsUI;
+		}
+
+		{
+			const TSharedRef< ILevelEditor >& LevelEditor = FModuleManager::GetModuleChecked<FLevelEditorModule>( "LevelEditor" ).GetFirstLevelEditor().ToSharedRef();
+
+			TSharedRef<SWidget> Modes = LevelEditor->CreateToolBox();
+
+			TSharedRef<SWidget> WidgetToDraw =
+				SNew( SDPIScaler )
+				.DPIScale( VREd::EditorUIScale->GetFloat() )
+				[
+					Modes
+				]
+			;
+
+			const bool bWithSceneComponent = false;
+			AVREditorFloatingUI* ModesUI = GetOwner().SpawnTransientSceneActor< AVREditorDockableWindow >( TEXT( "ModesUI" ), bWithSceneComponent );
+			ModesUI->SetSlateWidget( *this, WidgetToDraw, DefaultResolution, VREd::EditorUISize->GetFloat(), AVREditorFloatingUI::EDockedTo::Nothing );
+			ModesUI->ShowUI( false );
+			FloatingUIs.Add( ModesUI );
+
+			// @todo vreditor placement: This is required to force the modes UI to refresh -- otherwise it looks empty
+			GLevelEditorModeTools().ActivateMode( FBuiltinEditorModes::EM_Placement );
+
+			EditorUIPanels[ (int32)EEditorUIPanel::Modes ] = ModesUI;			
+		}
+
+		// Create the tutorial dockable window
+		{
+			const bool bShowAtStart = 
+				!bIsVREditorDemo && 
+				GetOwner().IsActuallyUsingVR();	// Don't show tutorial at startup when in "forced VR" mode, because you can't interact with UI using the mouse currently
+
+			AVREditorFloatingUI* TutorialUI = GetOwner().SpawnTransientSceneActor< AVREditorDockableWindow >( TEXT( "TutorialUI" ), false );
+			TutorialUI->SetUMGWidget( *this, TutorialWidgetClass, FIntPoint( VREd::TutorialUIResolutionX->GetFloat(), VREd::TutorialUIResolutionY->GetFloat() ), VREd::TutorialUISize->GetFloat(), bShowAtStart ? AVREditorFloatingUI::EDockedTo::Room : AVREditorFloatingUI::EDockedTo::Nothing );
+			if( bShowAtStart )
+			{
+				TutorialUI->SetRelativeOffset( FVector( VREd::TutorialUILocationX->GetFloat(), VREd::TutorialUILocationY->GetFloat(), VREd::TutorialUILocationZ->GetFloat() ) );
+				TutorialUI->SetLocalRotation( FRotator( VREd::TutorialUIPitch->GetFloat(), VREd::TutorialUIYaw->GetFloat(), 0 ) );
+			}
+			
+			TutorialUI->ShowUI( bShowAtStart );
+
+			FloatingUIs.Add( TutorialUI );
+
+			EditorUIPanels[ (int32)EEditorUIPanel::Tutorial ] = TutorialUI;		
+		}
+
+		{
+			const FIntPoint Resolution(VREd::AssetEditorUIResolutionX->GetInt(), VREd::AssetEditorUIResolutionY->GetInt());
+
+			const bool bWithSceneComponent = false;
+			AVREditorDockableWindow* TabManagerUI = GetOwner().SpawnTransientSceneActor< AVREditorDockableWindow >( TEXT( "AssetEditor" ), bWithSceneComponent );
+			TabManagerUI->SetSlateWidget( *this, SNullWidget::NullWidget, Resolution, VREd::EditorUISize->GetFloat(), AVREditorFloatingUI::EDockedTo::Nothing );
+			TabManagerUI->ShowUI( false );
+
+			// @todo vreditor: Could use "Hovering" instead for better performance with many UIs, but needs to be manually refreshed in some cases
+			TabManagerUI->GetWidgetComponent()->SetDrawingPolicy( EVREditorWidgetDrawingPolicy::Always );
+
+			FloatingUIs.Add( TabManagerUI );
+
+			EditorUIPanels[ (int32)EEditorUIPanel::AssetEditor ] = TabManagerUI;
+
+			TSharedPtr<SWindow> TabManagerWindow = TabManagerUI->GetWidgetComponent()->GetSlateWindow();
+			TSharedRef<SWindow> TabManagerWindowRef = TabManagerWindow.ToSharedRef();
+			ProxyTabManager = MakeShareable( new FProxyTabmanager( TabManagerWindowRef ) );
+
+			ProxyTabManager->OnTabOpened.AddUObject( this, &UVREditorUISystem::OnProxyTabLaunched );
+			ProxyTabManager->OnAttentionDrawnToTab.AddUObject( this, &UVREditorUISystem::OnAttentionDrawnToTab );
+
+			// We're going to start stealing tabs from the global tab manager inserting them into the world instead.
+			FGlobalTabmanager::Get()->SetProxyTabManager( ProxyTabManager );
+
+			FAssetEditorManager::Get().OnAssetEditorOpened().AddUObject( this, &UVREditorUISystem::OnAssetEditorOpened );
+		}
+		// create the world settings menu
+		{
+			FPropertyEditorModule& PropertyEditorModule = FModuleManager::Get().LoadModuleChecked< FPropertyEditorModule >( "PropertyEditor" );
+
+			FDetailsViewArgs DetailsViewArgs( false, false, true, FDetailsViewArgs::HideNameArea, false, GUnrealEd );
+			DetailsViewArgs.bShowActorLabel = false;
+
+			TSharedRef<IDetailsView> WorldSettings = PropertyEditorModule.CreateDetailView( DetailsViewArgs );
+
+			UWorld* World = GetOwner().GetWorld();
+			if ( World != NULL )
+			{
+				WorldSettings->SetObject( World->GetWorldSettings() );
+			}
+
+			TSharedRef<SWidget> WidgetToDraw = 
+				SNew( SDPIScaler )
+				.DPIScale( VREd::EditorUIScale->GetFloat() )
+				[
+					WorldSettings
+				]
+			;
+
+			const bool bWithSceneComponent = false;
+			AVREditorFloatingUI* WorldSettingsUI = GetOwner().SpawnTransientSceneActor< AVREditorDockableWindow >( TEXT( "WorldSettingsUI" ), bWithSceneComponent );
+			WorldSettingsUI->SetSlateWidget( *this, WidgetToDraw, DefaultResolution, VREd::EditorUISize->GetFloat(), AVREditorFloatingUI::EDockedTo::Nothing );
+			WorldSettingsUI->ShowUI( false );
+			FloatingUIs.Add( WorldSettingsUI );
+
+			EditorUIPanels[ (int32)EEditorUIPanel::WorldSettings ] = WorldSettingsUI;
+		}
+	}
+}
+
+void UVREditorUISystem::CleanUpActorsBeforeMapChangeOrSimulate()
+{
+	for( AVREditorFloatingUI* FloatingUIPtr : FloatingUIs )
+	{
+		if( FloatingUIPtr != nullptr )
+		{
+			FloatingUIPtr->Destroy( false, false );
+			FloatingUIPtr = nullptr;
+		}
+	}
+
+	FloatingUIs.Reset();
+	EditorUIPanels.Reset();
+	QuickRadialMenu = nullptr;
+	QuickMenuUI = nullptr;
+
+	ProxyTabManager.Reset();
+
+	// Remove the proxy tab manager, we don't want to steal tabs any more.
+	FGlobalTabmanager::Get()->SetProxyTabManager(TSharedPtr<FProxyTabmanager>());
+	FAssetEditorManager::Get().OnAssetEditorOpened().RemoveAll(this);
+}
+
+void UVREditorUISystem::OnAssetEditorOpened(UObject* Asset)
+{
+	// We need to disable drag drop on the tabs spawned in VR mode.
+	TArray<IAssetEditorInstance*> Editors = FAssetEditorManager::Get().FindEditorsForAsset(Asset);
+	for ( IAssetEditorInstance* Editor : Editors )
+	{
+		Editor->GetAssociatedTabManager()->SetCanDoDragOperation(false);
+	}
+}
+
+bool UVREditorUISystem::IsWidgetAnEditorUIWidget( const UActorComponent* WidgetComponent ) const
+{
+	if( WidgetComponent != nullptr && WidgetComponent->IsA( UWidgetComponent::StaticClass() ) )
+	{
+		for( AVREditorFloatingUI* FloatingUIPtr : FloatingUIs )
+		{
+			if( FloatingUIPtr != nullptr )
+			{
+				if( WidgetComponent == FloatingUIPtr->GetWidgetComponent() )
+				{
+					return true;
+				}
+			}
+		}
+	}
+
+	return false;
+}
+
+
+bool UVREditorUISystem::IsShowingEditorUIPanel( const EEditorUIPanel EditorUIPanel ) const
+{
+	AVREditorFloatingUI* Panel = EditorUIPanels[ (int32)EditorUIPanel ];
+	if( Panel != nullptr )
+	{
+		return Panel->IsUIVisible();
+	}
+
+	return false;
+}
+
+
+void UVREditorUISystem::ShowEditorUIPanel( const UWidgetComponent* WidgetComponent, UVREditorInteractor* Interactor, const bool bShouldShow, const bool OnHand, const bool bRefreshQuickMenu, const bool bPlaySound )
+{
+	AVREditorFloatingUI* Panel = nullptr;
+	for( AVREditorFloatingUI* CurrentPanel : EditorUIPanels )
+	{
+		if( CurrentPanel && CurrentPanel->GetWidgetComponent() == WidgetComponent )
+		{
+			Panel = CurrentPanel;
+			break;
+		}
+	}
+
+	ShowEditorUIPanel( Panel, Interactor, bShouldShow, OnHand, bRefreshQuickMenu, bPlaySound );
+}
+
+
+void UVREditorUISystem::ShowEditorUIPanel( const EEditorUIPanel EditorUIPanel, UVREditorInteractor* Interactor, const bool bShouldShow, const bool OnHand, const bool bRefreshQuickMenu, const bool bPlaySound )
+{
+	AVREditorFloatingUI* Panel = EditorUIPanels[ (int32)EditorUIPanel ];
+	ShowEditorUIPanel( Panel, Interactor, bShouldShow, OnHand, bRefreshQuickMenu, bPlaySound );
+}
+
+void UVREditorUISystem::ShowEditorUIPanel( AVREditorFloatingUI* Panel, UVREditorInteractor* Interactor, const bool bShouldShow, const bool OnHand, const bool bRefreshQuickMenu, const bool bPlaySound )
+{
+	if( Panel != nullptr )
+	{
+		AVREditorFloatingUI::EDockedTo DockedTo = Panel->GetDockedTo();
+		if( OnHand || DockedTo == AVREditorFloatingUI::EDockedTo::Nothing )
+		{
+			// Hide any panels that are already shown on this hand
+			for( int32 PanelIndex = 0; PanelIndex < (int32)EEditorUIPanel::TotalCount; ++PanelIndex )
+			{
+				AVREditorFloatingUI* OtherPanel = EditorUIPanels[PanelIndex];
+				if( OtherPanel != nullptr && OtherPanel->IsUIVisible() && OtherPanel->GetDockedTo() != AVREditorFloatingUI::EDockedTo::Room )
+				{
+					UVREditorInteractor* OtherPanelDockToVREditorInteractor = GetOwner().GetHandInteractor( OtherPanel->GetDockedTo() == AVREditorFloatingUI::EDockedTo::LeftHand ? EControllerHand::Left : EControllerHand::Right );
+					if( OtherPanelDockToVREditorInteractor == Interactor )
+					{
+						OtherPanel->ShowUI( false );
+						OtherPanel->SetDockedTo( AVREditorFloatingUI::EDockedTo::Nothing );
+						OtherPanelDockToVREditorInteractor->SetHasUIInFront( false );
+					}
+				}
+			}
+			
+			UVREditorMotionControllerInteractor* MotionControllerInteractor = Cast<UVREditorMotionControllerInteractor>( Interactor );
+			if ( MotionControllerInteractor )
+			{
+				const AVREditorFloatingUI::EDockedTo NewDockedTo = MotionControllerInteractor->GetControllerSide() == EControllerHand::Left ? AVREditorFloatingUI::EDockedTo::LeftHand : AVREditorFloatingUI::EDockedTo::RightHand;
+				Panel->SetDockedTo( NewDockedTo );
+			}
+			
+			if ( bShouldShow )
+			{
+				Panel->SetScale( Panel->GetInitialScale() );
+			}
+
+			const FVector EditorUIRelativeOffset( Panel->GetSize().Y * 0.5f + VREd::EditorUIOffsetFromHand->GetFloat(), 0.0f, 0.0f );
+			Panel->SetRelativeOffset( EditorUIRelativeOffset );
+			Panel->SetLocalRotation( FRotator( 90.0f, 180.0f, 0.0f ) ); // @todo vreditor: needs initial rotation
+		}
+
+		Panel->ShowUI( bShouldShow );
+
+		if ( Panel->GetDockedTo() == AVREditorFloatingUI::EDockedTo::LeftHand || Panel->GetDockedTo() == AVREditorFloatingUI::EDockedTo::RightHand )
+		{
+			UVREditorInteractor* PanelDockedToVREditorInteractor = GetOwner().GetHandInteractor( Panel->GetDockedTo() == AVREditorFloatingUI::EDockedTo::LeftHand ? EControllerHand::Left : EControllerHand::Right );
+			PanelDockedToVREditorInteractor->SetHasUIInFront( bShouldShow );
+		}
+
+		if( bPlaySound )
+		{
+			UGameplayStatics::PlaySound2D( GetOwner().GetWorld(), bShouldShow ? ShowUISound : HideUISound );
+		}
+
+		if ( bRefreshQuickMenu && QuickMenuUI )
+		{
+			QuickMenuUI->GetUserWidget<UVREditorQuickMenu>()->RefreshUI();
+		}
+	}
+}
+
+
+bool UVREditorUISystem::IsShowingRadialMenu( UVREditorInteractor* Interactor ) const
+{
+	const EControllerHand DockedToHand = QuickRadialMenu->GetDockedTo() == AVREditorFloatingUI::EDockedTo::LeftHand ? EControllerHand::Left : EControllerHand::Right;
+	const UVREditorInteractor* DockedToInteractor = GetOwner().GetHandInteractor( DockedToHand );
+	return DockedToInteractor == Interactor && !QuickRadialMenu->bHidden;
+}
+
+
+void UVREditorUISystem::UpdateRadialMenu( UVREditorInteractor* Interactor )
+{
+	if( QuickRadialMenu->bHidden )
+	{
+		QuickRadialMenu->ShowUI( true, false );
+	}
+
+	if( !QuickRadialMenu->bHidden )
+	{
+		RadialMenuHideDelayTime = 0.0f;
+		QuickRadialMenu->GetUserWidget<UVREditorRadialMenu>()->Update( Interactor );
+	}
+}
+
+
+void UVREditorUISystem::TryToSpawnRadialMenu( UVREditorInteractor* Interactor )
+{
+	if( Interactor != nullptr )
+	{
+		UVREditorMotionControllerInteractor* MotionControllerInteractor = Cast<UVREditorMotionControllerInteractor>( Interactor );
+		if( MotionControllerInteractor != nullptr )
+		{
+			UVREditorInteractor* DockedToInteractor = nullptr;
+			if( !QuickRadialMenu->bHidden )
+			{
+				const EControllerHand ControllerHand = QuickRadialMenu->GetDockedTo() == AVREditorFloatingUI::EDockedTo::LeftHand ? EControllerHand::Left : EControllerHand::Right;
+				DockedToInteractor = GetOwner().GetHandInteractor( ControllerHand );
+			}
+
+			EViewportInteractionDraggingMode DraggingMode = Interactor->GetDraggingMode();
+
+			bool bNeedsSpawn =
+				( QuickRadialMenu->bHidden || DockedToInteractor != Interactor ) &&
+				DraggingMode != EViewportInteractionDraggingMode::ActorsAtLaserImpact &&	// Don't show radial menu if the hand is busy dragging something around
+				DraggingMode != EViewportInteractionDraggingMode::ActorsFreely &&
+				DraggingMode != EViewportInteractionDraggingMode::ActorsWithGizmo &&
+				DraggingMode != EViewportInteractionDraggingMode::World &&
+				DraggingMode != EViewportInteractionDraggingMode::AssistingDrag &&
+				( InteractorDraggingUI == nullptr || InteractorDraggingUI != Interactor ) &&
+				!Interactor->IsHoveringOverUI();	// Don't show radial menu when aiming at a UI  (too much clutter)
+
+			UVREditorRadialMenu* RadialMenu = QuickRadialMenu->GetUserWidget<UVREditorRadialMenu>();
+			// We need to update the trackpad position in the radialmenu before checking if it can be used
+			RadialMenu->Update( Interactor );
+			if( RadialMenu && RadialMenu->IsInMenuRadius() )
+			{
+				bNeedsSpawn = false;
+			}
+
+			if( bNeedsSpawn )
+			{
+				const AVREditorFloatingUI::EDockedTo DockedTo = MotionControllerInteractor->GetControllerSide() == EControllerHand::Left ? AVREditorFloatingUI::EDockedTo::LeftHand : AVREditorFloatingUI::EDockedTo::RightHand;
+				QuickRadialMenu->SetDockedTo( DockedTo );
+				QuickRadialMenu->ShowUI( true );
+			}
+		}
+	}
+}
+
+
+void UVREditorUISystem::HideRadialMenu( UVREditorInteractor* Interactor )
+{
+	// Only hide the radial menu if the passed interactor is actually the interactor with the radial menu
+	if( IsShowingRadialMenu( Interactor ) )
+	{
+		QuickRadialMenu->ShowUI( false, true, VREd::RadialMenuFadeDelay->GetFloat() );
+	}
+}
+
+
+FTransform UVREditorUISystem::MakeDockableUITransformOnLaser( AVREditorDockableWindow* InitDraggingDockUI, UVREditorInteractor* Interactor, const float DockSelectDistance ) const
+{
+	FTransform HandTransform;
+	FVector HandForward;
+	Interactor->GetTransformAndForwardVector( HandTransform, HandForward );
+	const FTransform InteractorTransform = Interactor->GetTransform();
+
+	const FVector NewLocation = InteractorTransform.GetLocation() + ( HandForward.GetSafeNormal() * DockSelectDistance );
+	FRotator NewRotation = ( InteractorTransform.GetLocation() - NewLocation ).ToOrientationRotator();
+	NewRotation.Roll = -InteractorTransform.GetRotation().Rotator().Roll;
+
+	const FTransform LaserImpactToWorld( NewRotation, NewLocation );
+	return LaserImpactToWorld;
+}
+
+
+FTransform UVREditorUISystem::MakeDockableUITransform( AVREditorDockableWindow* InitDraggingDockUI, UVREditorInteractor* Interactor, const float DockSelectDistance )
+{
+	const FTransform UIOnLaserToWorld = MakeDockableUITransformOnLaser( DraggingUI, Interactor, DockSelectDistance );
+	const FTransform UIToUIOnLaser = DraggingUIOffsetTransform;
+	
+	const FTransform UIToWorld = UIToUIOnLaser * UIOnLaserToWorld;
+
+	if( DraggingUI != nullptr &&  DraggingUI->GetDockedTo() == AVREditorFloatingUI::EDockedTo::Dragging && InteractorDraggingUI && Interactor == InteractorDraggingUI )
+	{
+		UVREditorMotionControllerInteractor* OtherInteractor = Cast<UVREditorMotionControllerInteractor>( InteractorDraggingUI->GetOtherInteractor() );
+		if ( OtherInteractor )
+		{
+			const float WorldScaleFactor = GetOwner().GetWorldScaleFactor();
+			const FVector DraggingUILocation = InteractorDraggingUI->GetHoverLocation();
+			const FVector OtherHandLocation = OtherInteractor->GetTransform().GetLocation();
+			LastDraggingHoverLocation = UIOnLaserToWorld.GetLocation();
+
+			const float Distance = FVector::Dist( UIOnLaserToWorld.GetLocation(), OtherHandLocation ) / WorldScaleFactor;
+
+			// if dragged passed the threshold since starting dragging
+			if ( Distance > VREd::MinDockDragDistance->GetFloat() && !bDraggedDockFromHandPassedThreshold &&
+				( DraggingUI->GetPreviouslyDockedTo() == AVREditorFloatingUI::EDockedTo::LeftHand || DraggingUI->GetPreviouslyDockedTo() == AVREditorFloatingUI::EDockedTo::RightHand ) )
+			{
+				bDraggedDockFromHandPassedThreshold = true;
+			}
+
+			// Snapping to a hand when in range
+			if ( Distance <= VREd::MinDockDragDistance->GetFloat() && ( bDraggedDockFromHandPassedThreshold || DraggingUI->GetPreviouslyDockedTo() == AVREditorFloatingUI::EDockedTo::Room ) )
+			{
+				const FVector EditorUIRelativeOffset( DraggingUI->GetSize().Y * 0.5f + VREd::EditorUIOffsetFromHand->GetFloat(), 0.0f, 0.0f );
+				const FTransform MovingUIGoalTransform = DraggingUI->MakeUITransformLockedToHand( OtherInteractor, false, EditorUIRelativeOffset, FRotator( 90.0f, 180.0f, 0.0f ) );
+				const AVREditorFloatingUI::EDockedTo NewDockedTo = OtherInteractor->GetControllerSide() == EControllerHand::Left ? AVREditorFloatingUI::EDockedTo::LeftHand : AVREditorFloatingUI::EDockedTo::RightHand;
+				DraggingUI->MoveTo( MovingUIGoalTransform, VREd::DockDragSpeed->GetFloat(), NewDockedTo );
+			}
+			else
+			{
+				DraggingUI->MoveTo( UIToWorld, VREd::DockDragSpeed->GetFloat(), AVREditorFloatingUI::EDockedTo::Room );
+			}
+		}
+	}
+
+	return UIToWorld;
+}
+
+
+void UVREditorUISystem::StartDraggingDockUI( AVREditorDockableWindow* InitDraggingDockUI, UVREditorInteractor* Interactor, const float DockSelectDistance )
+{
+	AVREditorFloatingUI::EDockedTo DockTo = InitDraggingDockUI->GetDockedTo();
+	if( DockTo == AVREditorFloatingUI::EDockedTo::LeftHand || DockTo == AVREditorFloatingUI::EDockedTo::RightHand )
+	{
+		bDraggedDockFromHandPassedThreshold = false;
+
+		// Set the laser of the interactor the dockable window was pulled from
+		UViewportInteractor* OtherInteractor = Interactor->GetOtherInteractor();
+		if (OtherInteractor)
+		{
+			UVREditorInteractor* OtherVREditorInteractor = Cast<UVREditorInteractor>(OtherInteractor);
+			if (OtherVREditorInteractor)
+			{
+				OtherVREditorInteractor->SetHasUIInFront(false);
+			}
+		}
+	}
+
+	
+
+	InteractorDraggingUI = Interactor;
+
+	FTransform UIToWorld = InitDraggingDockUI->GetActorTransform();
+	UIToWorld.SetScale3D( FVector( 1.0f ) );
+	const FTransform WorldToUI = UIToWorld.Inverse();
+
+	const FTransform UIOnLaserToWorld = MakeDockableUITransformOnLaser( InitDraggingDockUI, Interactor, DockSelectDistance );
+	const FTransform UIOnLaserToUI = UIOnLaserToWorld * WorldToUI;
+	const FTransform UIToUIOnLaser = UIOnLaserToUI.Inverse();
+	DraggingUIOffsetTransform = UIToUIOnLaser;
+
+	DraggingUI = InitDraggingDockUI;
+	DraggingUI->SetDockedTo( AVREditorFloatingUI::EDockedTo::Dragging );
+
+	UGameplayStatics::PlaySound2D( GetOwner().GetWorld(), StartDragUISound );
+}
+
+AVREditorDockableWindow* UVREditorUISystem::GetDraggingDockUI() const
+{
+	return DraggingUI;
+}
+
+void UVREditorUISystem::StopDraggingDockUI( UVREditorInteractor* VREditorInteractor )
+{
+	if ( IsInteractorDraggingDockUI( VREditorInteractor ) )
+	{
+		// Put the Dock back on the hand it came from or leave it where it is in the room
+		UViewportInteractor* OtherInteractor = VREditorInteractor->GetOtherInteractor();
+		if ( OtherInteractor )
+		{
+			UVREditorMotionControllerInteractor* MotionControllerOtherInteractor = Cast<UVREditorMotionControllerInteractor>( OtherInteractor );
+			if ( MotionControllerOtherInteractor )
+			{
+				bool bOnHand = false;
+				const float Distance = FVector::Dist( LastDraggingHoverLocation, MotionControllerOtherInteractor->GetTransform().GetLocation() ) / GetOwner().GetWorldScaleFactor();
+				if ( Distance > VREd::MinDockDragDistance->GetFloat() )
+				{
+					DraggingUI->SetDockedTo( AVREditorFloatingUI::EDockedTo::Room );
+				}
+				else
+				{
+					bOnHand = true;
+					const FVector EditorUIRelativeOffset( DraggingUI->GetSize().Y * 0.5f + VREd::EditorUIOffsetFromHand->GetFloat(), 0.0f, 0.0f );
+					const FTransform MovingUIGoalTransform = DraggingUI->MakeUITransformLockedToHand( MotionControllerOtherInteractor, false, EditorUIRelativeOffset, FRotator( 90.0f, 180.0f, 0.0f ) );
+					const AVREditorFloatingUI::EDockedTo NewDockedTo = MotionControllerOtherInteractor->GetControllerSide() == EControllerHand::Left ? AVREditorFloatingUI::EDockedTo::LeftHand : AVREditorFloatingUI::EDockedTo::RightHand;
+					DraggingUI->MoveTo( MovingUIGoalTransform, VREd::DockSnapTime->GetFloat(), NewDockedTo );
+				}
+
+				ShowEditorUIPanel( DraggingUI, MotionControllerOtherInteractor, true, bOnHand );
+
+				DraggingUI = nullptr;
+				InteractorDraggingUI->SetDraggingMode( EViewportInteractionDraggingMode::Nothing );
+				InteractorDraggingUI = nullptr;
+
+				UGameplayStatics::PlaySound2D( GetOwner().GetWorld(), StopDragUISound );
+			}
+		}
+	}
+}
+
+bool UVREditorUISystem::IsDraggingDockUI()
+{
+	return DraggingUI != nullptr && DraggingUI->GetDockedTo() == AVREditorFloatingUI::EDockedTo::Dragging;
+}
+
+bool UVREditorUISystem::IsInteractorDraggingDockUI( const UVREditorInteractor* Interactor ) const
+{
+	return InteractorDraggingUI && InteractorDraggingUI == Interactor;
+}
+
+void UVREditorUISystem::TogglePanelsVisibility()
+{
+	bool bAnyPanelsVisible = false;
+	AVREditorFloatingUI* PanelOnHand = nullptr;
+
+	// Check if there is any panel visible and if any is docked to a hand
+	for (AVREditorFloatingUI* Panel : EditorUIPanels)
+	{
+		if (Panel && Panel->IsUIVisible())
+		{
+			bAnyPanelsVisible = true;
+
+			const AVREditorFloatingUI::EDockedTo DockedTo = Panel->GetDockedTo();
+			if (DockedTo == AVREditorFloatingUI::EDockedTo::LeftHand || DockedTo == AVREditorFloatingUI::EDockedTo::RightHand)
+			{
+				PanelOnHand = Panel;
+				break;
+			}
+		}
+	}
+
+	// Hide if there is any UI visible
+	const bool bShowUI = !bAnyPanelsVisible;
+	
+	UGameplayStatics::PlaySound2D( GetOwner().GetWorld(), bShowUI ? ShowUISound : HideUISound );
+
+	if ( bSetDefaultLayout)
+	{
+		bSetDefaultLayout = false;
+		SetDefaultWindowLayout();
+	}
+	else
+	{
+		for ( AVREditorFloatingUI* Panel : EditorUIPanels )
+		{
+			if ( Panel != nullptr && Panel->IsUIVisible() != bShowUI )
+			{
+				if ( bShowUI )
+				{
+					bool SetNewVisibility = true;
+					const AVREditorFloatingUI::EDockedTo DockedTo = Panel->GetDockedTo();
+					if ( DockedTo == AVREditorFloatingUI::EDockedTo::LeftHand || DockedTo == AVREditorFloatingUI::EDockedTo::RightHand )
+					{
+						if ( !PanelOnHand )
+						{
+							UVREditorInteractor* VREditorInteractor = GetOwner().GetHandInteractor( Panel->GetDockedTo() == AVREditorFloatingUI::EDockedTo::LeftHand ? EControllerHand::Left : EControllerHand::Right );
+							VREditorInteractor->SetHasUIInFront( true );
+							PanelOnHand = Panel;
+						}
+						else
+						{
+							SetNewVisibility = false;
+						}
+					}
+					else if ( DockedTo == AVREditorFloatingUI::EDockedTo::Nothing )
+					{
+						SetNewVisibility = false;
+					}
+
+					if ( SetNewVisibility )
+					{
+						Panel->ShowUI( true );
+					}
+				}
+				else
+				{
+					Panel->ShowUI( false );
+				}
+			}
+		}
+	}
+
+	if( !bShowUI && PanelOnHand )
+	{
+		UVREditorInteractor* VREditorInteractor = GetOwner().GetHandInteractor( PanelOnHand->GetDockedTo() == AVREditorFloatingUI::EDockedTo::LeftHand ? EControllerHand::Left : EControllerHand::Right );
+		VREditorInteractor->SetHasUIInFront( false );
+	}
+
+	if (QuickMenuUI)
+	{
+		QuickMenuUI->GetUserWidget<UVREditorQuickMenu>()->RefreshUI();
+	}
+}
+
+float UVREditorUISystem::GetMaxDockWindowSize() const
+{
+	return VREd::MaxDockWindowSize->GetFloat();
+}
+
+float UVREditorUISystem::GetMinDockWindowSize() const
+{
+	return VREd::MinDockWindowSize->GetFloat();
+}
+
+void UVREditorUISystem::OnProxyTabLaunched(TSharedPtr<SDockTab> NewTab)
+{
+	ShowAssetEditor();
+}
+
+void UVREditorUISystem::OnAttentionDrawnToTab(TSharedPtr<SDockTab> NewTab)
+{
+	// @todo vreditor: When clicking on Modes icons in the Modes panel, the LevelEditor tab is invoked which causes an empty asset editor to pop-up
+	static FName LevelEditorTabID( "LevelEditor" );
+	if( NewTab->GetLayoutIdentifier().TabType != LevelEditorTabID )
+	{
+		ShowAssetEditor();
+	}
+}
+
+void UVREditorUISystem::ShowAssetEditor()
+{
+	bRefocusViewport = true;
+
+	// A tab was opened, so make sure the "Asset" UI is visible.  That's where the user can interact
+	// with the newly-opened tab
+	if ( !IsShowingEditorUIPanel(EEditorUIPanel::AssetEditor) )
+	{
+		// Always spawn on a hand.  But which hand?  Well, we'll choose the hand that isn't actively clicking on something using a laser.
+		UVREditorInteractor* VREditorInteractor = GetOwner().GetHandInteractor( EControllerHand::Left ); // Hand that did not clicked with a laser
+		//@todo VREditor: This should include a check on VREditorInteractor->IsClickingOnUI(), however that is broken because of double click on UI.
+		if( !VREditorInteractor->HasUIInFront() )
+		{
+			VREditorInteractor = GetOwner().GetHandInteractor( EControllerHand::Right );
+		}
+		
+		const bool bShouldShow = true;
+		const bool bShowOnHand = true;
+		ShowEditorUIPanel( EEditorUIPanel::AssetEditor, VREditorInteractor, bShouldShow, bShowOnHand );
+
+		// Play haptic effect so user knows to look at their hand that now has UI on it!
+		VREditorInteractor->PlayHapticEffect( VREd::UIAssetEditorSummonedOnHandHapticFeedbackStrength->GetFloat() );
+	}
+}
+
+void UVREditorUISystem::TogglePanelVisibility( const EEditorUIPanel EditorUIPanel )
+{
+	AVREditorFloatingUI* Panel = EditorUIPanels[(int32)EditorUIPanel];
+	if (Panel != nullptr)
+	{
+		const bool bIsShowing = Panel->IsUIVisible();
+		UVREditorInteractor* InteractorWithQuickMenu = GetOwner().GetHandInteractor( QuickMenuUI->GetDockedTo() == AVREditorFloatingUI::EDockedTo::LeftArm ? EControllerHand::Left : EControllerHand::Right );
+		if( Panel->GetDockedTo() == AVREditorFloatingUI::EDockedTo::Room )
+		{
+			Panel->ShowUI( false );
+			Panel->SetDockedTo( AVREditorFloatingUI::EDockedTo::Nothing );
+		}
+		else
+		{
+			ShowEditorUIPanel( Panel, InteractorWithQuickMenu, !bIsShowing, true );
+		}
+	}
+}
+
+void UVREditorUISystem::SetDefaultWindowLayout()
+{
+	for ( int32 PanelIndex = 0; PanelIndex < DefaultWindowTransforms.Num(); ++PanelIndex )
+	{
+		AVREditorFloatingUI* Panel = EditorUIPanels[PanelIndex];
+		if ( Panel )
+		{
+			const AVREditorFloatingUI::EDockedTo DockedTo = Panel->GetDockedTo();
+			if (DockedTo == AVREditorFloatingUI::EDockedTo::LeftHand || DockedTo == AVREditorFloatingUI::EDockedTo::RightHand)
+			{
+				const EControllerHand ControllerHand = DockedTo == AVREditorFloatingUI::EDockedTo::LeftHand ? EControllerHand::Left : EControllerHand::Right;
+				UVREditorInteractor* Interactor = GetOwner().GetHandInteractor( ControllerHand );
+				Interactor->SetHasUIInFront( false );
+			}
+
+			Panel->SetDockedTo( AVREditorFloatingUI::EDockedTo::Room );
+			Panel->ShowUI( true );
+
+			// Make sure the UIs are centered around the direction your head is looking (yaw only!)
+			const FVector RoomSpaceHeadLocation = GetOwner().GetRoomSpaceHeadTransform().GetLocation() / GetOwner().GetWorldScaleFactor();
+			FRotator RoomSpaceHeadYawRotator = GetOwner().GetRoomSpaceHeadTransform().GetRotation().Rotator();
+			RoomSpaceHeadYawRotator.Pitch = 0.0f;
+			RoomSpaceHeadYawRotator.Roll = 0.0f;
+
+			FTransform NewTransform = DefaultWindowTransforms[PanelIndex];
+			NewTransform *= FTransform( RoomSpaceHeadYawRotator.Quaternion(), FVector::ZeroVector );
+			Panel->SetLocalRotation( NewTransform.GetRotation().Rotator() );
+			Panel->SetRelativeOffset( RoomSpaceHeadLocation + NewTransform.GetTranslation() );
+		}
+	}
+}
+void UVREditorUISystem::CreateVRColorPicker(const TSharedRef<SColorPicker>& ColorPicker)
+{
+	// Check that the Color Picker Panel isn't currently showing. Also handles the EEditorUIPanel::ColorPicker being null.
+	const FIntPoint DefaultResolution(VREd::DefaultEditorUIResolutionX->GetInt(), VREd::DefaultEditorUIResolutionY->GetInt());
+
+	TSharedRef<SWidget> WidgetToDraw =
+		SNew(SDPIScaler)
+		.DPIScale(VREd::EditorUIScale->GetFloat())
+		[
+			ColorPicker
+		]
+	;
+
+	// Only spawn the Color picker dockable window if there is none yet
+	bool bJustSpawned = false;
+	if( ColorPickerUI == nullptr )
+	{
+		const bool bWithSceneComponent = false;
+		ColorPickerUI = GetOwner().SpawnTransientSceneActor< AVREditorDockableWindow >( TEXT( "ColorPickerUI" ), bWithSceneComponent );
+		FloatingUIs.Add( ColorPickerUI );
+		EditorUIPanels[ (int32)EEditorUIPanel::ColorPicker ] = ColorPickerUI;
+		bJustSpawned = true;
+	}
+
+	ColorPickerUI->SetSlateWidget( *this, WidgetToDraw, DefaultResolution, VREd::EditorUISize->GetFloat(), AVREditorFloatingUI::EDockedTo::Room );
+
+	// Always spawn based on the location of the menu you are hovering over. To get this information, find the hand hovering over the UI.
+	UVREditorInteractor* VREditorInteractor = GetOwner().GetHandInteractor(EControllerHand::Left);
+	if (!VREditorInteractor->IsHoveringOverUI())
+	{
+		VREditorInteractor = GetOwner().GetHandInteractor(EControllerHand::Right);
+	}
+
+	// Spawn the color picker close to the pressed UI it was opened using UI. Otherwise spawn it in front of the user.
+	UWidgetComponent* HoveredWidgetComponent = VREditorInteractor->GetHoveringOverWidgetComponent();
+	if( HoveredWidgetComponent != nullptr )
+	{
+		FVector LaserPointerStart, LaserPointerEnd;
+		if( VREditorInteractor->GetLaserPointer( LaserPointerStart, LaserPointerEnd ) )
+		{
+			//Calculate new location using the direction of the laser with an offset towards the user
+			FVector LaserPointerDirection = LaserPointerEnd - LaserPointerStart;
+			LaserPointerDirection.Normalize();
+			const FVector HoverOverUILocation = VREditorInteractor->GetHoverLocation();
+			const FVector Offset = LaserPointerDirection * VREd::ColorPickerDockSpawnOffset->GetFloat() * VRMode->GetWorldScaleFactor();
+
+			FTransform NewTransform = FTransform( HoveredWidgetComponent->GetComponentQuat(), HoverOverUILocation - Offset );
+			
+			// Move the color picker if it was already spawned, or instantly set the transform if it was just spawned. We don't want to move the UI from location zero
+			if( bJustSpawned )
+			{
+				ColorPickerUI->SetDockedTo( AVREditorBaseActor::EDockedTo::Room );
+				ColorPickerUI->SetTransform( NewTransform );
+			}
+			else
+			{
+				ColorPickerUI->MoveTo( NewTransform, 0.2f, AVREditorBaseActor::EDockedTo::Room );
+			}
+
+			ColorPickerUI->SetActorTransform( NewTransform );
+			ColorPickerUI->UpdateRelativeRoomTransform();
+		}
+	}
+	else
+	{
+		// Make sure the UIs are centered around the direction your head is looking (yaw only!)
+		const FVector RoomSpaceHeadLocation = VRMode->GetRoomSpaceHeadTransform().GetLocation() / VRMode->GetWorldScaleFactor();
+		FRotator RoomSpaceHeadYawRotator = VRMode->GetRoomSpaceHeadTransform().GetRotation().Rotator();
+		RoomSpaceHeadYawRotator.Pitch = 0.0f;
+		RoomSpaceHeadYawRotator.Roll = 0.0f;
+	
+		FTransform NewTransform = VREd::DefaultColorPickerTransform;
+		NewTransform *= FTransform(RoomSpaceHeadYawRotator.Quaternion(), FVector::ZeroVector);
+		ColorPickerUI->SetLocalRotation(NewTransform.GetRotation().Rotator());
+		ColorPickerUI->SetRelativeOffset(RoomSpaceHeadLocation + NewTransform.GetTranslation());
+	}
+
+	ColorPickerUI->ShowUI( true );
+}
+
+void UVREditorUISystem::DestroyVRColorPicker()
+{
+
+	UVREditorInteractor* VREditorInteractor = nullptr;
+	const bool bShouldShow = false;
+	const bool bShowOnHand = false;
+	const bool bRefreshQuickMenu = true;
+	ShowEditorUIPanel(EEditorUIPanel::ColorPicker, VREditorInteractor, bShouldShow, bShowOnHand, bRefreshQuickMenu);
+}
diff --git a/Engine/Source/Editor/VREditor/UI/VREditorUISystem.h b/Engine/Source/Editor/VREditor/UI/VREditorUISystem.h
new file mode 100644
--- /dev/null
+++ b/Engine/Source/Editor/VREditor/UI/VREditorUISystem.h
@@ -0,0 +1,258 @@
+// Copyright 1998-2015 Epic Games, Inc. All Rights Reserved.
+
+#pragma once
+
+#include "VREditorUISystem.generated.h"
+
+class UVREditorInteractor;
+class UViewportInteractor;
+class SColorPicker;
+
+/**
+ * VR Editor user interface manager
+ */
+UCLASS()
+class UVREditorUISystem : public UObject
+{
+	GENERATED_UCLASS_BODY()
+
+public:
+
+	enum class EEditorUIPanel
+	{
+		ContentBrowser,
+		WorldOutliner,
+		ActorDetails,
+		Modes,
+		Tutorial,
+		AssetEditor,
+		WorldSettings,
+		ColorPicker,
+		// ...
+
+		TotalCount,
+	};
+
+	/** Shuts down the UISystem whenever the mode is exited */
+	void Shutdown();
+
+	/** Initializes default values for the UISystem and creates the UI panels */
+	void Init();
+
+	/** Called by VREditorMode to update us every frame */
+	void Tick( class FEditorViewportClient* ViewportClient, const float DeltaTime );
+
+	/** Called by VREditorMode to draw debug graphics every frame */
+	void Render( const class FSceneView* SceneView, class FViewport* Viewport, class FPrimitiveDrawInterface* PDI );
+
+	/** Sets the owner of this system */
+	void SetOwner( class UVREditorMode* NewOwner )
+	{
+		VRMode = NewOwner;
+	}
+
+	/** Gets the owner of this system */
+	class UVREditorMode& GetOwner()
+	{
+		return *VRMode;
+	}
+
+	/** Gets the owner of this system (const) */
+	const class UVREditorMode& GetOwner() const
+	{
+		return *VRMode;
+	}
+
+	/** Called by our owner right before a map is loaded or switching between Simulate and normal Editor, so we can clean
+	    up our spawned actors */
+	void CleanUpActorsBeforeMapChangeOrSimulate();
+
+	/** Searches to see if the specified widget component is from our editor UI */
+	bool IsWidgetAnEditorUIWidget( const class UActorComponent* WidgetComponent ) const;
+
+	/** Returns true if the specified editor UI panel is currently visible */
+	bool IsShowingEditorUIPanel( const EEditorUIPanel EditorUIPanel ) const;
+
+	/** Sets whether the specified editor UI panel should be visible.  Any other UI floating off this hand will be dismissed when showing it. */
+	void ShowEditorUIPanel( const class UWidgetComponent* WidgetComponent, UVREditorInteractor* Interactor, const bool bShouldShow, const bool OnHand = false, const bool bRefreshQuickMenu = true, const bool bPlaySound = true );
+	void ShowEditorUIPanel( const EEditorUIPanel EditorUIPanel, UVREditorInteractor* Interactor, const bool bShouldShow, const bool OnHand = false, const bool bRefreshQuickMenu = true, const bool bPlaySound = true );
+	void ShowEditorUIPanel( class AVREditorFloatingUI* Panel, UVREditorInteractor* Interactor, const bool bShouldShow, const bool OnHand = false, const bool bRefreshQuickMenu = true, const bool bPlaySound = true );
+
+	/** Returns true if the radial menu is visible on this hand */
+	bool IsShowingRadialMenu( UVREditorInteractor* Interactor ) const;
+
+	/** Tries to spawn the radial menu (if the specified hand isn't doing anything else) */
+	void TryToSpawnRadialMenu( UVREditorInteractor* Interactor );
+
+	/** Updates the radial menu with the current hand trackpad location */
+	void UpdateRadialMenu( UVREditorInteractor* Interactor );
+
+	/** Hides the radial menu if the specified hand is showing it */
+	void HideRadialMenu( UVREditorInteractor* Interactor );
+
+	/** Start dragging a dock window on the hand */
+	void StartDraggingDockUI( class AVREditorDockableWindow* InitDraggingDockUI, UVREditorInteractor* Interactor, const float DockSelectDistance );
+
+	/** Makes up a transform for a dockable UI when dragging it with a laser at the specified distance from the laser origin */
+	FTransform MakeDockableUITransformOnLaser( AVREditorDockableWindow* InitDraggingDockUI, UVREditorInteractor* Interactor, const float DockSelectDistance ) const;
+
+	/** Makes up a transform for a dockable UI when dragging it that includes the original offset from the laser's impact point */
+	FTransform MakeDockableUITransform( AVREditorDockableWindow* InitDraggingDockUI, UVREditorInteractor* Interactor, const float DockSelectDistance );
+	
+	/** Returns the current dragged dock window, nullptr if none */
+	class AVREditorDockableWindow* GetDraggingDockUI() const;
+
+	/** Resets all values to stop dragging the current dock window */
+	void StopDraggingDockUI( UVREditorInteractor* VREditorInteractor );
+
+	/** Are we currently dragging a dock window */
+	bool IsDraggingDockUI();	
+
+	bool IsInteractorDraggingDockUI( const UVREditorInteractor* Interactor ) const;
+
+	/** Hides and unhides all the editor UI panels */
+	void TogglePanelsVisibility();
+
+	/** Get the maximum dock window size */
+	float GetMaxDockWindowSize() const;
+
+	/** Get the minimum dock window size */
+	float GetMinDockWindowSize() const;
+
+	/** Toggles the visibility of the panel, if the panel is in room space it will be hidden and docked to nothing */
+	void TogglePanelVisibility( const EEditorUIPanel EditorUIPanel );
+
+protected:
+
+	/** Called when the user presses a button on their motion controller device */
+	void OnVRAction( class FEditorViewportClient& ViewportClient, UViewportInteractor* Interactor,
+		const struct FViewportActionKeyInput& Action, bool& bOutIsInputCaptured, bool& bWasHandled );
+
+	/** Called every frame to update hover state */
+	void OnVRHoverUpdate( class FEditorViewportClient& ViewportClient, UViewportInteractor* Interactor, FVector& HoverImpactPoint, bool& bWasHandled );
+
+	/** Testing Slate UI */
+	void CreateUIs();
+
+	/**
+	 * Called when the injected proxy tab manager reports a new tab has been launched, 
+	 * signaling an asset editor has been launched, but really it could be any major global tab.
+	 */
+	void OnProxyTabLaunched(TSharedPtr<SDockTab> NewTab);
+
+	/** Called when the injected proxy tab manager reports needing to draw attention to a tab. */
+	void OnAttentionDrawnToTab(TSharedPtr<SDockTab> NewTab);
+
+	/** Show the asset editor if it's not currently visible. */
+	void ShowAssetEditor();
+
+	/** Called when an asset editor opens an asset while in VR Editor Mode. */
+	void OnAssetEditorOpened(UObject* Asset);
+
+	/** Sets the main windows to their default transform */
+	void SetDefaultWindowLayout();
+
+
+	/** Creates a VR-specific color picker. Gets bound to SColorPicker's creation override delegate */
+	void CreateVRColorPicker(const TSharedRef<SColorPicker>& ColorPicker);
+
+	/** Hides the VR-specific color picker. Gets bound to SColorPicker's destruction override delegate */
+	void DestroyVRColorPicker();
+
+
+protected:
+
+	/** Owning object */
+	UPROPERTY()
+	class UVREditorMode* VRMode;
+
+	/** All of the floating UIs.  These may or may not be visible (spawned) */
+	UPROPERTY()
+	TArray< class AVREditorFloatingUI* > FloatingUIs;
+
+	/** Our Quick Menu UI */
+	UPROPERTY()
+	AVREditorFloatingUI* QuickMenuUI;
+
+	/** Editor UI panels */
+	UPROPERTY()
+	TArray<AVREditorFloatingUI*> EditorUIPanels;
+
+	/** The Radial Menu UI */
+	UPROPERTY()
+	AVREditorFloatingUI* QuickRadialMenu;
+
+	/** The time since the radial menu was updated */
+	float RadialMenuHideDelayTime;
+
+	/** Quick menu widget class */
+	UClass* QuickMenuWidgetClass;
+
+	/** Quick radial menu widget class */
+	UClass* QuickRadialWidgetClass;
+
+	/** Tutorial menu widget class */
+	UClass* TutorialWidgetClass;
+	
+	//
+	// Dragging UI
+	//
+
+	/** Interactor that is dragging the UI */
+	class UVREditorInteractor* InteractorDraggingUI;
+
+	/** Offset transform from room-relative transform to the object where we picked it up at */
+	FTransform DraggingUIOffsetTransform;
+
+	/** The current UI that is being dragged */
+	UPROPERTY()
+	class AVREditorDockableWindow* DraggingUI;
+
+	/** The color picker dockable window */
+	UPROPERTY()
+	class AVREditorDockableWindow* ColorPickerUI;
+
+	//
+	// Tab Manager UI
+	//
+
+	/** Allows us to steal the global tabs and show them in the world. */
+	TSharedPtr<FProxyTabmanager> ProxyTabManager;
+
+	/** Set to true when we need to refocus the viewport. */
+	bool bRefocusViewport;
+
+
+	//
+	// Sounds
+	//
+
+	/** Start dragging UI sound */
+	UPROPERTY()
+	class USoundCue* StartDragUISound;
+
+	/** Stop dragging UI sound */
+	UPROPERTY()
+	class USoundCue* StopDragUISound;
+
+	/** Hide UI sound */
+	UPROPERTY()
+	class USoundCue* HideUISound;
+
+	/** Show UI sound */
+	UPROPERTY()
+	class USoundCue* ShowUISound;
+
+	/** If the current dragged dock passed a certain distance if dragged from a hand */
+	bool bDraggedDockFromHandPassedThreshold;
+
+	/** The last dragged hover location by the laser */
+	FVector LastDraggingHoverLocation;
+
+	/** Default transforms */
+	TArray<FTransform> DefaultWindowTransforms;
+
+	/** If this is the first time using TogglePanelsVisibility */
+	bool bSetDefaultLayout;
+};
+
diff --git a/Engine/Source/Editor/VREditor/VREditor.Build.cs b/Engine/Source/Editor/VREditor/VREditor.Build.cs
--- a/Engine/Source/Editor/VREditor/VREditor.Build.cs
+++ b/Engine/Source/Editor/VREditor/VREditor.Build.cs
@@ -6,7 +6,8 @@ namespace UnrealBuildTool.Rules
 	{
         public VREditor(TargetInfo Target)
 		{
-			PrivateDependencyModuleNames.AddRange(
+
+            PrivateDependencyModuleNames.AddRange(
                 new string[] {
                     "AppFramework",
 				    "Core",
@@ -35,6 +36,16 @@ namespace UnrealBuildTool.Rules
 				new string[] {
 				}
 			);
+
+            PrivateIncludePaths.AddRange(
+                new string[] {
+                    "Editor/VREditor/Gizmo",
+                    "Editor/VREditor/UI",
+                    "Editor/VREditor/Teleporter",
+                    "Editor/VREditor/Interactables"
+                }
+            );
+
 		}
 	}
 }
\ No newline at end of file
diff --git a/Engine/Source/Editor/VREditor/VREditorBaseUserWidget.cpp b/Engine/Source/Editor/VREditor/VREditorBaseUserWidget.cpp
deleted file mode 100644
diff --git a/Engine/Source/Editor/VREditor/VREditorBaseUserWidget.h b/Engine/Source/Editor/VREditor/VREditorBaseUserWidget.h
deleted file mode 100644
diff --git a/Engine/Source/Editor/VREditor/VREditorDockableWindow.cpp b/Engine/Source/Editor/VREditor/VREditorDockableWindow.cpp
deleted file mode 100644
diff --git a/Engine/Source/Editor/VREditor/VREditorDockableWindow.h b/Engine/Source/Editor/VREditor/VREditorDockableWindow.h
deleted file mode 100644
diff --git a/Engine/Source/Editor/VREditor/VREditorFloatingUI.cpp b/Engine/Source/Editor/VREditor/VREditorFloatingUI.cpp
deleted file mode 100644
diff --git a/Engine/Source/Editor/VREditor/VREditorFloatingUI.h b/Engine/Source/Editor/VREditor/VREditorFloatingUI.h
deleted file mode 100644
diff --git a/Engine/Source/Editor/VREditor/VREditorInteractor.cpp b/Engine/Source/Editor/VREditor/VREditorInteractor.cpp
--- a/Engine/Source/Editor/VREditor/VREditorInteractor.cpp
+++ b/Engine/Source/Editor/VREditor/VREditorInteractor.cpp
@@ -7,7 +7,7 @@
 #include "VREditorMode.h"
 #include "VREditorFloatingText.h"
 
-#include "Interactables/VREditorButton.h"
+#include "VREditorButton.h"
 #include "VREditorDockableWindow.h"
 
 #define LOCTEXT_NAMESPACE "VREditor"
diff --git a/Engine/Source/Editor/VREditor/VREditorQuickMenu.cpp b/Engine/Source/Editor/VREditor/VREditorQuickMenu.cpp
deleted file mode 100644
diff --git a/Engine/Source/Editor/VREditor/VREditorQuickMenu.h b/Engine/Source/Editor/VREditor/VREditorQuickMenu.h
deleted file mode 100644
diff --git a/Engine/Source/Editor/VREditor/VREditorRadialMenu.cpp b/Engine/Source/Editor/VREditor/VREditorRadialMenu.cpp
deleted file mode 100644
diff --git a/Engine/Source/Editor/VREditor/VREditorRadialMenu.h b/Engine/Source/Editor/VREditor/VREditorRadialMenu.h
deleted file mode 100644
diff --git a/Engine/Source/Editor/VREditor/VREditorRadialMenuItem.cpp b/Engine/Source/Editor/VREditor/VREditorRadialMenuItem.cpp
deleted file mode 100644
diff --git a/Engine/Source/Editor/VREditor/VREditorRadialMenuItem.h b/Engine/Source/Editor/VREditor/VREditorRadialMenuItem.h
deleted file mode 100644
diff --git a/Engine/Source/Editor/VREditor/VREditorTransformGizmo.cpp b/Engine/Source/Editor/VREditor/VREditorTransformGizmo.cpp
deleted file mode 100644
diff --git a/Engine/Source/Editor/VREditor/VREditorTransformGizmo.h b/Engine/Source/Editor/VREditor/VREditorTransformGizmo.h
deleted file mode 100644
diff --git a/Engine/Source/Editor/VREditor/VREditorUISystem.cpp b/Engine/Source/Editor/VREditor/VREditorUISystem.cpp
deleted file mode 100644
diff --git a/Engine/Source/Editor/VREditor/VREditorUISystem.h b/Engine/Source/Editor/VREditor/VREditorUISystem.h
deleted file mode 100644
