Added a checkbox for disabling startup VR Editor tutorials to the Experimental Editor Settings. On by default.

diff --git a/Engine/Config/BaseEditorPerProjectUserSettings.ini b/Engine/Config/BaseEditorPerProjectUserSettings.ini
--- a/Engine/Config/BaseEditorPerProjectUserSettings.ini
+++ b/Engine/Config/BaseEditorPerProjectUserSettings.ini
@@ -137,6 +137,8 @@ bEnableFindAndReplaceReferences =False
 bEnableVREditing=false
 ; If true, wearing a Vive or Oculus Rift headset will automatically enter VR Editing mode if Enable VR Editing is true. 
 bEnableAutoVREditMode=true
+; If true, the VR Tutorial window will always be shown at startup
+bAlwaysShowVRTutorial=true
 
 [/Script/UnrealEd.EditorLoadingSavingSettings]
 ; True if we should automatically load a default level at start up
diff --git a/Engine/Source/Editor/UnrealEd/Classes/Settings/EditorExperimentalSettings.h b/Engine/Source/Editor/UnrealEd/Classes/Settings/EditorExperimentalSettings.h
--- a/Engine/Source/Editor/UnrealEd/Classes/Settings/EditorExperimentalSettings.h
+++ b/Engine/Source/Editor/UnrealEd/Classes/Settings/EditorExperimentalSettings.h
@@ -126,6 +126,10 @@ public:
 	UPROPERTY(EditAnywhere, config, Category = VR, meta = (DisplayName = "Enable VR Mode Auto-Entry"))
 	bool bEnableAutoVREditMode;
 
+	/**If true, the tutorial window should always appear when the VR Editor is launched. */
+	UPROPERTY(EditAnywhere, config, Category = VR, meta = (DisplayName = "Always Show VR Tutorial at Start"))
+	bool bAlwaysShowVRTutorial;
+
 	/** Allows editing of potentially unsafe properties during PIE. Advanced use only - use with caution. */
 	UPROPERTY(EditAnywhere, config, Category = Tools, meta = (DisplayName = "Allow editing of potentially unsafe properties."))
 	bool bAllowPotentiallyUnsafePropertyEditing;
diff --git a/Engine/Source/Editor/VREditor/UI/VREditorUISystem.cpp b/Engine/Source/Editor/VREditor/UI/VREditorUISystem.cpp
--- a/Engine/Source/Editor/VREditor/UI/VREditorUISystem.cpp
+++ b/Engine/Source/Editor/VREditor/UI/VREditorUISystem.cpp
@@ -790,7 +790,8 @@ void UVREditorUISystem::CreateUIs()
 		{
 			const bool bShowAtStart = 
 				!bIsVREditorDemo && 
-				GetOwner().IsActuallyUsingVR();	// Don't show tutorial at startup when in "forced VR" mode, because you can't interact with UI using the mouse currently
+				GetOwner().IsActuallyUsingVR() &&	// Don't show tutorial at startup when in "forced VR" mode, because you can't interact with UI using the mouse currently
+				(GetDefault<UEditorExperimentalSettings>()->bAlwaysShowVRTutorial); // Only show tutorial at startup if the user has this enabled
 
 			AVREditorFloatingUI* TutorialUI = GetOwner().SpawnTransientSceneActor< AVREditorDockableWindow >( TEXT( "TutorialUI" ), false );
 			TutorialUI->SetUMGWidget( *this, TutorialWidgetClass, FIntPoint( VREd::TutorialUIResolutionX->GetFloat(), VREd::TutorialUIResolutionY->GetFloat() ), VREd::TutorialUISize->GetFloat(), bShowAtStart ? AVREditorFloatingUI::EDockedTo::Room : AVREditorFloatingUI::EDockedTo::Nothing );
