VREditor: Fix editor startup crash (UE-36783)
- Recent changes to VREditor caused "LevelEditor" module to become loaded earlier than ever before, which caused various registration problems leading to crashes
- Changed VREditor to not load "LevelEditor" until VR mode was actually enabled

diff --git a/Engine/Source/Editor/VREditor/VREditorModeManager.cpp b/Engine/Source/Editor/VREditor/VREditorModeManager.cpp
--- a/Engine/Source/Editor/VREditor/VREditorModeManager.cpp
+++ b/Engine/Source/Editor/VREditor/VREditorModeManager.cpp
@@ -20,19 +20,10 @@ FVREditorModeManager::FVREditorModeManager() :
 	HMDWornState( EHMDWornState::Unknown ),
 	TimeSinceHMDChecked( 0.0f )
 {
-	// Tell the level editor we want to be notified when selection changes
-	FLevelEditorModule& LevelEditor = FModuleManager::LoadModuleChecked<FLevelEditorModule>( "LevelEditor" );
-	LevelEditor.OnMapChanged().AddRaw( this, &FVREditorModeManager::OnMapChanged );
 }
 
 FVREditorModeManager::~FVREditorModeManager()
 {
-	FLevelEditorModule* LevelEditor = FModuleManager::GetModulePtr<FLevelEditorModule>( "LevelEditor" );
-	if( LevelEditor != nullptr )
-	{
-		LevelEditor->OnMapChanged().RemoveAll( this );
-	}
-
 	CurrentVREditorMode = nullptr;
 	PreviousVREditorMode = nullptr;
 }
@@ -170,6 +161,12 @@ void FVREditorModeManager::StartVREditorMode( const bool bForceWithoutHMD )
 	TSharedPtr<FEditorWorldWrapper> EditorWorld = GEditor->GetEditorWorldManager()->GetEditorWorldWrapper( GWorld );
 	UVREditorMode* ModeFromWorld = EditorWorld->GetVREditorMode();
 
+	// Tell the level editor we want to be notified when selection changes
+	{
+		FLevelEditorModule& LevelEditor = FModuleManager::LoadModuleChecked<FLevelEditorModule>( "LevelEditor" );
+		LevelEditor.OnMapChanged().AddRaw( this, &FVREditorModeManager::OnMapChanged );
+	}
+	
 	CurrentVREditorMode = ModeFromWorld;
 	CurrentVREditorMode->SetActuallyUsingVR( !bForceWithoutHMD );
 	CurrentVREditorMode->Enter();
@@ -180,6 +177,12 @@ void FVREditorModeManager::CloseVREditor()
 	//Store the current WorldToMeters before exiting the mode
 	LastWorldToMeters = GWorld->GetWorldSettings()->WorldToMeters;
 
+	FLevelEditorModule* LevelEditor = FModuleManager::GetModulePtr<FLevelEditorModule>( "LevelEditor" );
+	if( LevelEditor != nullptr )
+	{
+		LevelEditor->OnMapChanged().RemoveAll( this );
+	}
+
 	if( CurrentVREditorMode != nullptr )
 	{
 		CurrentVREditorMode->Exit();
