VREditor: Fixed a few static analysis warnings

diff --git a/Engine/Source/Editor/VREditor/VREditorAvatarActor.cpp b/Engine/Source/Editor/VREditor/VREditorAvatarActor.cpp
--- a/Engine/Source/Editor/VREditor/VREditorAvatarActor.cpp
+++ b/Engine/Source/Editor/VREditor/VREditorAvatarActor.cpp
@@ -241,7 +241,8 @@ void AVREditorAvatarActor::TickManually( const float DeltaTime )
 		RoomSpaceTransformWithWorldToMetersScaling.SetScale3D( RoomSpaceTransformWithWorldToMetersScaling.GetScale3D( ) * WorldScaleFactorVector );
 
 		// @todo vreditor urgent: Head disabled until we can fix late frame update issue
-		if ( false && VRMode->IsActuallyUsingVR() && GEngine->HMDDevice.IsValid( ) && GEngine->HMDDevice->IsHeadTrackingAllowed( ) )
+		check( HeadMeshComponent != nullptr );
+		if( false ) // && VRMode->IsActuallyUsingVR() && GEngine->HMDDevice.IsValid() && GEngine->HMDDevice->IsHeadTrackingAllowed() )
 		{
 			HeadMeshComponent->SetVisibility( true );
 
@@ -258,6 +259,7 @@ void AVREditorAvatarActor::TickManually( const float DeltaTime )
 	}
 
 	// Scale the grid so that it stays the same size in the tracking space
+	check( WorldMovementGridMeshComponent != nullptr );
 	WorldMovementGridMeshComponent->SetRelativeScale3D( WorldScaleFactorVector * VREd::GridScaleMultiplier->GetFloat() );
 	WorldMovementGridMeshComponent->SetRelativeLocation( WorldScaleFactorVector * FVector( 0.0f, 0.0f, VREd::GridHeightOffset->GetFloat() ) );
 
@@ -267,8 +269,8 @@ void AVREditorAvatarActor::TickManually( const float DeltaTime )
 		UVREditorInteractor* RightHandInteractor = VRMode->GetHandInteractor( EControllerHand::Right );
 
 		if ( LeftHandInteractor != nullptr && RightHandInteractor != nullptr &&
-			( LeftHandInteractor->GetDraggingMode( ) == EViewportInteractionDraggingMode::World && RightHandInteractor->GetDraggingMode( ) == EViewportInteractionDraggingMode::AssistingDrag ) ||
-			( LeftHandInteractor->GetDraggingMode( ) == EViewportInteractionDraggingMode::AssistingDrag && RightHandInteractor->GetDraggingMode( ) == EViewportInteractionDraggingMode::World ) )
+			 ( ( LeftHandInteractor->GetDraggingMode() == EViewportInteractionDraggingMode::World && RightHandInteractor->GetDraggingMode() == EViewportInteractionDraggingMode::AssistingDrag ) ||
+			   ( LeftHandInteractor->GetDraggingMode() == EViewportInteractionDraggingMode::AssistingDrag && RightHandInteractor->GetDraggingMode() == EViewportInteractionDraggingMode::World ) ) )
 		{
 			// Setting all components to be visible
 			CurrentScaleProgressMeshComponent->SetVisibility( true );
