New Slate-based dynamic menu, with delegates for overriding widget creating. New behavior is that you must click the thumbstick or touchpad to summon/dismiss the radial menu, select options with a direction on the thumbstick/touchpad, and then select with the trigger. Quick menu functionality is mostly the same, just the laser is now hidden.

diff --git a/Engine/Source/Editor/VREditor/SVRRadialPanel.cpp b/Engine/Source/Editor/VREditor/SVRRadialPanel.cpp
new file mode 100644
--- /dev/null
+++ b/Engine/Source/Editor/VREditor/SVRRadialPanel.cpp
@@ -0,0 +1,328 @@
+// Copyright 1998-2016 Epic Games, Inc. All Rights Reserved.
+
+#include "VREditorModule.h"
+#include "VREditorUISystem.h"
+#include "SVRRadialPanel.h"
+
+namespace VREd
+{
+	static FAutoConsoleVariable MinJoystickOffsetBeforeRadialMenu(TEXT("VREd.MinJoystickOffsetBeforeRadialMenu"), 0.15f, TEXT("Toggles inverting the touch pad vertical axis"));
+}
+
+/**
+ * Construct this widget
+ *
+ * @param	InArgs	The declaration data for this widget
+ */
+void SRadialBox::Construct( const SRadialBox::FArguments& InArgs )
+{
+	const int32 NumSlots = InArgs.Slots.Num();
+	for ( int32 SlotIndex = 0; SlotIndex < NumSlots; ++SlotIndex )
+	{
+		Children.Add( InArgs.Slots[SlotIndex] );
+	}
+
+	CurrentlyHoveredButton = nullptr;
+}
+
+
+
+
+template<EOrientation Orientation>
+static void ArrangeChildrenAlong( const TPanelChildren<SRadialPanel::FSlot>& Children, const FGeometry& AllottedGeometry, FArrangedChildren& ArrangedChildren )
+{
+	// Allotted space will be given to fixed-size children first.
+	// Remaining space will be proportionately divided between stretch children (SizeRule_Stretch)
+	// based on their stretch coefficient
+
+	if (Children.Num() > 0)
+	{
+		float StretchCoefficientTotal = 0;
+		float FixedTotal = 0;
+
+		// Compute the sum of stretch coefficients (SizeRule_Stretch) and space required by fixed-size widgets
+		// (SizeRule_Auto).
+		for( int32 ChildIndex=0; ChildIndex < Children.Num(); ++ChildIndex )
+		{
+			const SRadialPanel::FSlot& CurChild = Children[ChildIndex];
+
+			if ( CurChild.GetWidget()->GetVisibility() != EVisibility::Collapsed )
+			{
+				// All widgets contribute their margin to the fixed space requirement
+				FixedTotal += CurChild.SlotPadding.Get().GetTotalSpaceAlong<Orientation>();
+
+				if ( CurChild.SizeParam.SizeRule == FSizeParam::SizeRule_Stretch )
+				{
+					// for stretch children we save sum up the stretch coefficients
+					StretchCoefficientTotal += CurChild.SizeParam.Value.Get();
+				}
+				else
+				{
+					// Auto-sized children contribute their desired size to the fixed space requirement
+					float ChildSize =  CurChild.GetWidget()->GetDesiredSize().Y;
+
+					// Clamp to the max size if it was specified
+					float MaxSize = CurChild.MaxSize.Get();
+					FixedTotal += MaxSize > 0 ? FMath::Min( MaxSize, ChildSize ) : ChildSize ;
+				}
+			}
+
+		}
+
+		// The space available for SizeRule_Stretch widgets is any space that wasn't taken up by fixed-sized widgets.
+		const float NonFixedSpace = FMath::Max( 0.0f, AllottedGeometry.Size.Y - FixedTotal );
+
+		int PositionSoFar = 0;
+
+		// Now that we have the total fixed-space requirement and the total stretch coefficients we can
+		// arrange widgets top-to-bottom or left-to-right (depending on the orientation).
+		for( int32 ChildIndex=0; ChildIndex < Children.Num(); ++ChildIndex )
+		{
+			const SRadialPanel::FSlot& CurChild = Children[ChildIndex];
+			const EVisibility ChildVisibility = CurChild.GetWidget()->GetVisibility();
+
+			FVector2D ChildDesiredSize = CurChild.GetWidget()->GetDesiredSize();
+			
+			// Figure out the area allocated to the child in the direction of BoxPanel
+			// The area allocated to the slot is ChildSize + the associated margin.
+			float ChildSize = 0;
+			if ( ChildVisibility != EVisibility::Collapsed )
+			{
+				// The size of the widget depends on its size type
+				if ( CurChild.SizeParam.SizeRule == FSizeParam::SizeRule_Stretch )
+				{
+					// Stretch widgets get a fraction of the space remaining after all the fixed-space requirements are met
+					ChildSize = NonFixedSpace * CurChild.SizeParam.Value.Get() / StretchCoefficientTotal;
+				}
+				else
+				{
+					// Auto-sized widgets get their desired-size value
+					ChildSize =  ChildDesiredSize.Y;
+				}
+
+				// Clamp to the max size if it was specified
+				float MaxSize = CurChild.MaxSize.Get();
+				if( MaxSize > 0 )
+				{
+					ChildSize = FMath::Min( MaxSize, ChildSize );
+				}
+			}
+
+			const FMargin SlotPadding(CurChild.SlotPadding.Get());
+
+			FVector2D SlotSize =  FVector2D( AllottedGeometry.Size.X, ChildSize + SlotPadding.GetTotalSpaceAlong<Orient_Vertical>() );
+
+			// Figure out the size and local position of the child within the slot			
+			AlignmentArrangeResult XAlignmentResult = AlignChild<Orient_Horizontal>( SlotSize.X, CurChild, SlotPadding );					
+			AlignmentArrangeResult YAlignmentResult = AlignChild<Orient_Vertical>( SlotSize.Y, CurChild, SlotPadding );
+			
+			FVector2D MyDesiredSize(0, 0);
+			MyDesiredSize = ComputeDesiredSizeForBox(Children);
+			
+			// Set the x and y coordinates of the slot based on the current angle, a radius of MyDesiredSize.Y / 2.0, and a center of (0, MyDesiredSize.Y / 2.0)
+			// First element ends up at 90 degrees
+			float RadialX = MyDesiredSize.Y / 2.0f * FMath::Cos(PositionSoFar * ( (2.0f*PI) / Children.Num() ));
+			float RadialY = MyDesiredSize.Y / 2.0f + MyDesiredSize.Y / 2.0f * FMath::Sin(PositionSoFar * ( (2.0f*PI) / Children.Num() ));
+			 
+
+			// Add the information about this child to the output list (ArrangedChildren)
+			ArrangedChildren.AddWidget( ChildVisibility, AllottedGeometry.MakeChild(
+				// The child widget being arranged
+				CurChild.GetWidget(),
+				// Child's local position (i.e. position within parent)
+				 FVector2D( RadialX + XAlignmentResult.Offset, RadialY),
+				// Child's size
+				FVector2D( XAlignmentResult.Size, YAlignmentResult.Size )
+				));
+
+			if ( ChildVisibility != EVisibility::Collapsed )
+			{
+				// Offset the next child by the size of the current child and any post-child (bottom/right) margin
+				PositionSoFar++;
+			}
+		}
+	}
+}
+
+/**
+ * Panels arrange their children in a space described by the AllottedGeometry parameter. The results of the arrangement
+ * should be returned by appending a FArrangedWidget pair for every child widget. See StackPanel for an example
+ *
+ * @param AllottedGeometry    The geometry allotted for this widget by its parent.
+ * @param ArrangedChildren    The array to which to add the WidgetGeometries that represent the arranged children.
+ */
+void SRadialPanel::OnArrangeChildren( const FGeometry& AllottedGeometry, FArrangedChildren& ArrangedChildren ) const
+{
+	
+	ArrangeChildrenAlong<Orient_Vertical>(this->Children, AllottedGeometry, ArrangedChildren );
+	
+}
+
+/**
+ * Helper to ComputeDesiredSize().
+ *
+ * @param Orientation   Template parameters that controls the orientation in which the children are layed out
+ * @param Children      The children whose size we want to assess in a horizontal or vertical arrangement.
+ *
+ * @return The size desired by the children given an orientation.
+ */
+static FVector2D ComputeDesiredSizeForBox( const TPanelChildren<SRadialPanel::FSlot>& Children ) 
+{
+	// The desired size of this panel is the total size desired by its children plus any margins specified in this panel.
+	// The layout along the panel's axis is describe dy the SizeParam, while the perpendicular layout is described by the
+	// alignment property.
+	FVector2D MyDesiredSize(0,0);
+	for( int32 ChildIndex=0; ChildIndex < Children.Num(); ++ChildIndex )
+	{
+		const SRadialPanel::FSlot& CurChild = Children[ChildIndex];
+		const FVector2D& CurChildDesiredSize = CurChild.GetWidget()->GetDesiredSize();
+		if ( CurChild.GetWidget()->GetVisibility() != EVisibility::Collapsed )
+		{
+			
+			// For a vertical panel, we want to find the maximum desired width (including margin).
+			// That will be the desired width of the whole panel.
+			MyDesiredSize.X = FMath::Max(MyDesiredSize.X, (2.0f * CurChildDesiredSize.X + CurChild.SlotPadding.Get().GetTotalSpaceAlong<Orient_Horizontal>()));
+
+			// Clamp to the max size if it was specified
+			float FinalChildDesiredSize = CurChildDesiredSize.Y;
+			float MaxSize = CurChild.MaxSize.Get();
+			if( MaxSize > 0 )
+			{
+				FinalChildDesiredSize = FMath::Min( MaxSize, FinalChildDesiredSize );
+			}
+
+			MyDesiredSize.Y += FinalChildDesiredSize + CurChild.SlotPadding.Get().GetTotalSpaceAlong<Orient_Vertical>();				
+			
+		}
+	}
+
+	return MyDesiredSize;
+}
+
+/**
+ * A Panel's desired size in the space required to arrange of its children on the screen while respecting all of
+ * the children's desired sizes and any layout-related options specified by the user. See StackPanel for an example.
+ *
+ * @return The desired size.
+ */
+FVector2D SRadialPanel::ComputeDesiredSize( float ) const
+{
+	return ComputeDesiredSizeForBox(this->Children);
+}
+
+/** @return  The children of a panel in a slot-agnostic way. */
+FChildren* SRadialPanel::GetChildren()
+{
+	return &Children;
+}
+
+int32 SRadialPanel::RemoveSlot( const TSharedRef<SWidget>& SlotWidget )
+{
+	for (int32 SlotIdx = 0; SlotIdx < Children.Num(); ++SlotIdx)
+	{
+		if ( SlotWidget == Children[SlotIdx].GetWidget() )
+		{
+			Children.RemoveAt(SlotIdx);
+			return SlotIdx;
+		}
+	}
+
+	return -1;
+}
+
+void SRadialPanel::ClearChildren()
+{
+	Children.Empty();
+}
+
+SRadialPanel::SRadialPanel( )
+: Children()
+{
+
+}
+
+void SRadialBox::HighlightSlot(FVector2D& TrackpadPosition)
+{
+	if (FMath::Abs(TrackpadPosition.Size()) < VREd::MinJoystickOffsetBeforeRadialMenu->GetFloat())
+	{
+		if (CurrentlyHoveredButton.Get() != nullptr)
+		{
+			const FPointerEvent& SimulatedPointer = FPointerEvent::FPointerEvent();
+			CurrentlyHoveredButton->OnMouseLeave(SimulatedPointer);
+			CurrentlyHoveredButton = nullptr;
+		}
+		
+		return;
+	}
+
+	const float AnglePerItem = 360.0f / NumEntries;
+	float Angle = FRotator::NormalizeAxis(FMath::RadiansToDegrees(FMath::Atan2(TrackpadPosition.X, TrackpadPosition.Y)));
+
+	Angle += AnglePerItem / 2.0f;
+	// first element of the menu is at 90
+	Angle += -90.0f;
+	if (Angle < 0)
+	{
+		Angle = Angle + 360.0f;
+	}
+
+	const int32 Index = (Angle / AnglePerItem);
+	
+	TSharedRef<SWidget> CurrentChild = GetChildren()->GetChildAt(Index);
+	TSharedRef<SWidget> TestWidget = UVREditorUISystem::FindWidgetOfType(CurrentChild, FName(TEXT("SButton")));
+
+	if (TestWidget != SNullWidget::NullWidget)
+	{
+		CurrentlyHoveredButton = StaticCastSharedRef<SButton>(TestWidget);
+		const FPointerEvent& SimulatedPointer = FPointerEvent::FPointerEvent();
+		const FGeometry& ChildGeometry = FGeometry();
+		
+		// Simulate mouse entering event for the button if it was not previously hovered
+		if(!(CurrentlyHoveredButton->IsHovered()))
+		{
+			CurrentlyHoveredButton->OnMouseEnter(ChildGeometry, SimulatedPointer);
+		}
+
+		// Simulate mouse leaving events for any buttons that were previously hovered
+		for (int32 ButtonCount = 0; ButtonCount < (NumEntries); ButtonCount++)
+		{
+			if (ButtonCount != Index)
+			{
+				CurrentChild = GetChildren()->GetChildAt(ButtonCount);
+				TestWidget = UVREditorUISystem::FindWidgetOfType(CurrentChild, FName(TEXT("SButton")));
+				TSharedRef<SButton> TestButton = StaticCastSharedRef<SButton>(TestWidget);
+				if(TestButton->IsHovered())
+				{
+					TestButton->OnMouseLeave(SimulatedPointer);
+				}
+			}
+
+		}
+	}
+	
+}
+
+void SRadialBox::SimulateLeftClickDown()
+{
+	if(CurrentlyHoveredButton.Get() != nullptr)
+	{
+		const FPointerEvent& SimulatedPointer = FPointerEvent(uint32(0), uint32(0), FVector2D::ZeroVector, FVector2D::ZeroVector, true);
+		const FGeometry& ChildGeometry = FGeometry();
+		CurrentlyHoveredButton->OnMouseButtonDown(ChildGeometry, SimulatedPointer);
+	}	
+}
+
+void SRadialBox::SimulateLeftClickUp()
+{
+	if (CurrentlyHoveredButton.Get() != nullptr)
+	{
+		const FPointerEvent& SimulatedPointer = FPointerEvent(uint32(0), uint32(0), FVector2D::ZeroVector, FVector2D::ZeroVector, true);
+		const FGeometry& ChildGeometry = FGeometry();
+		CurrentlyHoveredButton->OnMouseButtonUp(ChildGeometry, SimulatedPointer);
+	}
+}
+
+const TSharedPtr<SButton>& SRadialBox::GetCurrentlyHoveredButton()
+{
+	return CurrentlyHoveredButton;
+}
diff --git a/Engine/Source/Editor/VREditor/SVRRadialPanel.h b/Engine/Source/Editor/VREditor/SVRRadialPanel.h
new file mode 100644
--- /dev/null
+++ b/Engine/Source/Editor/VREditor/SVRRadialPanel.h
@@ -0,0 +1,251 @@
+// Copyright 1998-2016 Epic Games, Inc. All Rights Reserved.
+
+#pragma once
+
+#include "SlateCore.h"
+
+
+
+/**
+ * A RadialPanel contains one child and describes how that child should be arranged on the screen.
+ */
+class SRadialPanel
+	: public SPanel
+{
+public:
+
+	/**
+	 * A RadialPanel contains one RadialPanel child and describes how that
+	 * child should be arranged on the screen.
+	 */
+	class FSlot : public TSlotBase<FSlot>
+	{
+	public:		
+		/** Horizontal and Vertical Boxes inherit from FSlot */
+		virtual ~FSlot(){}
+		/**
+		 * How much space this slot should occupy along panel's direction.
+		 *   When SizeRule is SizeRule_Auto, the widget's DesiredSize will be used as the space required.
+		 *   When SizeRule is SizeRule_Stretch, the available space will be distributed proportionately between
+		 *   peer Widgets depending on the Value property. Available space is space remaining after all the
+		 *   peers' SizeRule_Auto requirements have been satisfied.
+		 */
+		FSizeParam SizeParam;
+		
+		/** Horizontal positioning of child within the allocated slot */
+		EHorizontalAlignment HAlignment;
+		
+		/** Vertical positioning of child within the allocated slot */
+		EVerticalAlignment VAlignment;
+		
+		/** The padding to add around the child. */
+		TAttribute<FMargin> SlotPadding;
+		
+		/** The max size that this slot can be (0 if no max) */
+		TAttribute<float> MaxSize;
+			
+	protected:
+		/** Default values for a slot. */
+		FSlot()
+			: TSlotBase<FSlot>()
+			, SizeParam( FStretch(1) )
+			, HAlignment( HAlign_Fill )
+			, VAlignment( VAlign_Fill )
+			, SlotPadding( FMargin(0) )
+			, MaxSize( 0.0f )
+		{ }
+	};
+
+public:
+
+	/** Removes a slot from this box panel which contains the specified SWidget
+	 *
+	 * @param SlotWidget The widget to match when searching through the slots
+	 * @returns The index in the children array where the slot was removed and -1 if no slot was found matching the widget
+	 */
+	int32 RemoveSlot( const TSharedRef<SWidget>& SlotWidget );
+
+	/**
+	 * Removes all children from the box.
+	 */
+	void ClearChildren();
+
+public:
+
+	// Begin SWidget overrides.
+
+	virtual void OnArrangeChildren( const FGeometry& AllottedGeometry, FArrangedChildren& ArrangedChildren ) const;
+	
+	virtual FVector2D ComputeDesiredSize(float) const;
+
+	virtual FChildren* GetChildren();
+
+	// End SWidget overrides.
+
+protected:
+
+	SRadialPanel();
+
+	/** The Radial Panel's children. */
+	TPanelChildren<FSlot> Children;
+};
+
+
+/** A Radial Box Panel. See SRadialPanel for more info. */
+class SRadialBox : public SRadialPanel
+{
+public:
+	class FSlot : public SRadialPanel::FSlot
+	{
+		public:
+		FSlot()
+		: SRadialPanel::FSlot()
+		{
+		}
+
+		FSlot& AutoWidth()
+		{
+			SizeParam = FAuto();
+			return *this;
+		}
+
+		FSlot& MaxWidth( const TAttribute< float >& InMaxWidth )
+		{
+			MaxSize = InMaxWidth;
+			return *this;
+		}
+
+		FSlot& FillWidth( const TAttribute< float >& StretchCoefficient )
+		{
+			SizeParam = FStretch( StretchCoefficient );
+			return *this;
+		}
+		
+		FSlot& Padding( float Uniform )
+		{
+			SlotPadding = FMargin(Uniform);
+			return *this;
+		}
+
+		FSlot& Padding( float Horizontal, float Vertical )
+		{
+			SlotPadding = FMargin(Horizontal, Vertical);
+			return *this;
+		}
+
+		FSlot& Padding( float Left, float Top, float Right, float Bottom )
+		{
+			SlotPadding = FMargin(Left, Top, Right, Bottom);
+			return *this;
+		}
+		
+		FSlot& Padding( TAttribute<FMargin> InPadding )
+		{
+			SlotPadding = InPadding;
+			return *this;
+		}
+
+		FSlot& HAlign( EHorizontalAlignment InHAlignment )
+		{
+			HAlignment = InHAlignment;
+			return *this;
+		}
+
+		FSlot& VAlign( EVerticalAlignment InVAlignment )
+		{
+			VAlignment = InVAlignment;
+			return *this;
+		}
+
+		FSlot& operator[]( TSharedRef<SWidget> InWidget )
+		{
+			SRadialPanel::FSlot::operator[](InWidget);
+			return *this;
+		}
+
+		FSlot& Expose( FSlot*& OutVarToInit )
+		{
+			OutVarToInit = this;
+			return *this;
+		}
+	};
+
+	static FSlot& Slot()
+	{
+		return *(new FSlot());
+	}
+
+	SLATE_BEGIN_ARGS( SRadialBox )
+	{
+		_Visibility = EVisibility::SelfHitTestInvisible;
+	}
+
+		SLATE_SUPPORTS_SLOT(SRadialBox::FSlot)
+
+	SLATE_END_ARGS()
+
+	FSlot& AddSlot()
+	{
+		SRadialBox::FSlot& NewSlot = *new SRadialBox::FSlot();
+		this->Children.Add( &NewSlot );
+
+		Invalidate(EInvalidateWidget::Layout);
+
+		return NewSlot;
+	}
+
+	FSlot& InsertSlot(int32 Index = INDEX_NONE)
+	{
+		if (Index == INDEX_NONE)
+		{
+			return AddSlot();
+		}
+		SRadialBox::FSlot& NewSlot = *new SRadialBox::FSlot();
+		this->Children.Insert(&NewSlot, Index);
+
+		Invalidate(EInvalidateWidget::Layout);
+
+		return NewSlot;
+	}
+
+	int32 NumSlots() const
+	{
+		return this->Children.Num();
+	}
+
+	int32 NumEntries;
+
+	FORCENOINLINE SRadialBox()
+	: SRadialPanel()
+	{
+		bCanTick = false;
+		bCanSupportFocus = false;
+	}
+
+	/**
+	 * Construct this widget
+	 *
+	 * @param	InArgs	The declaration data for this widget
+	 */
+	void Construct( const FArguments& InArgs );
+
+
+	/**
+	* Highlight the widget in a slot based on a given trackpad position
+	*
+	* @param	TrackpadPosition	The current trackpad or thumbstick position
+	*/
+	void HighlightSlot(FVector2D& TrackpadPosition);
+
+	/** Simulate a left-mouse click (down) on the currently hovered button */
+	void SimulateLeftClickDown();
+
+	/** Simulate a left-mouse click (up) on the currently hovered button */
+	void SimulateLeftClickUp();
+
+	const TSharedPtr<SButton>& GetCurrentlyHoveredButton();
+
+	private:
+	/** Stores the currently hovered button */
+	TSharedPtr<SButton> CurrentlyHoveredButton;
+};
diff --git a/Engine/Source/Editor/VREditor/UI/VREditorQuickMenu.cpp b/Engine/Source/Editor/VREditor/UI/VREditorQuickMenu.cpp
deleted file mode 100644
diff --git a/Engine/Source/Editor/VREditor/UI/VREditorQuickMenu.h b/Engine/Source/Editor/VREditor/UI/VREditorQuickMenu.h
deleted file mode 100644
diff --git a/Engine/Source/Editor/VREditor/UI/VREditorRadialMenu.cpp b/Engine/Source/Editor/VREditor/UI/VREditorRadialMenu.cpp
deleted file mode 100644
diff --git a/Engine/Source/Editor/VREditor/UI/VREditorRadialMenu.h b/Engine/Source/Editor/VREditor/UI/VREditorRadialMenu.h
deleted file mode 100644
diff --git a/Engine/Source/Editor/VREditor/UI/VREditorRadialMenuItem.cpp b/Engine/Source/Editor/VREditor/UI/VREditorRadialMenuItem.cpp
deleted file mode 100644
diff --git a/Engine/Source/Editor/VREditor/UI/VREditorRadialMenuItem.h b/Engine/Source/Editor/VREditor/UI/VREditorRadialMenuItem.h
deleted file mode 100644
diff --git a/Engine/Source/Editor/VREditor/UI/VREditorUISystem.cpp b/Engine/Source/Editor/VREditor/UI/VREditorUISystem.cpp
--- a/Engine/Source/Editor/VREditor/UI/VREditorUISystem.cpp
+++ b/Engine/Source/Editor/VREditor/UI/VREditorUISystem.cpp
@@ -3,11 +3,9 @@
 #include "VREditorModule.h"
 #include "VREditorUISystem.h"
 #include "VREditorMode.h"
+#include "VREditorBaseUserWidget.h"
 #include "VREditorFloatingUI.h"
 #include "VREditorDockableWindow.h"
-#include "VREditorQuickMenu.h"
-#include "VREditorRadialMenu.h"
-#include "VREditorRadialMenuItem.h"
 #include "IHeadMountedDisplay.h"
 #include "ViewportWorldInteraction.h"
 #include "ViewportInteractor.h"
@@ -17,6 +15,7 @@
 // UI
 #include "WidgetComponent.h"
 #include "VREditorWidgetComponent.h"
+#include "VREditorStyle.h"
 
 // Content Browser
 #include "Editor/ContentBrowser/Public/ContentBrowserModule.h"
@@ -38,6 +37,14 @@
 #include "SWidget.h"
 #include "SDockTab.h"
 #include "SColorPicker.h"
+#include "MultiBox.h"
+#include "LevelEditorActions.h"
+#include "VREditorActions.h"
+#include "GenericCommands.h"
+#include "SVRRadialPanel.h"
+
+
+#define LOCTEXT_NAMESPACE "VREditor"
 
 namespace VREd
 {
@@ -45,8 +52,8 @@ namespace VREd
 	static FAutoConsoleVariable ContentBrowserUIResolutionY( TEXT( "VREd.ContentBrowserUIResolutionY" ), 1200, TEXT( "Vertical resolution to use for content browser UI render targets" ) );
 	static FAutoConsoleVariable DefaultEditorUIResolutionX( TEXT( "VREd.DefaultEditorUIResolutionX" ), 1024, TEXT( "Horizontal resolution to use for VR editor UI render targets" ) );
 	static FAutoConsoleVariable DefaultEditorUIResolutionY( TEXT( "VREd.DefaultEditorUIResolutionY" ), 1024, TEXT( "Vertical resolution to use for VR editor UI render targets" ) );
-	static FAutoConsoleVariable QuickMenuUIResolutionX( TEXT( "VREd.QuickMenuUIResolutionX" ), 1024, TEXT( "Horizontal resolution to use for Quick Menu VR UI render targets" ) );
-	static FAutoConsoleVariable QuickMenuUIResolutionY( TEXT( "VREd.QuickMenuUIResolutionY" ), 900, TEXT( "Vertical resolution to use for Quick Menu VR UI render targets" ) );
+	static FAutoConsoleVariable QuickMenuUIResolutionX( TEXT( "VREd.QuickMenuUIResolutionX" ), 1440, TEXT( "Horizontal resolution to use for Quick Menu VR UI render targets" ) );
+	static FAutoConsoleVariable QuickMenuUIResolutionY( TEXT( "VREd.QuickMenuUIResolutionY" ), 1440, TEXT( "Vertical resolution to use for Quick Menu VR UI render targets" ) );
 	static FAutoConsoleVariable ContentBrowserUISize( TEXT( "VREd.ContentBrowserUISize" ), 60.0f, TEXT( "How big content browser UIs should be" ) );
 	static FAutoConsoleVariable EditorUISize( TEXT( "VREd.EditorUISize" ), 50.0f, TEXT( "How big editor UIs should be" ) );
 	static FAutoConsoleVariable ContentBrowserUIScale( TEXT( "VREd.ContentBrowserUIScale" ), 2.0f, TEXT( "How much to scale up (or down) the content browser for VR" ) );
@@ -95,8 +102,6 @@ UVREditorUISystem::UVREditorUISystem( const FObjectInitializer& Initializer ) :
 	QuickMenuUI( nullptr ),
 	QuickRadialMenu( nullptr ),
 	RadialMenuHideDelayTime( 0.0f ),
-	QuickMenuWidgetClass( nullptr ),
-	QuickRadialWidgetClass( nullptr ),
 	TutorialWidgetClass( nullptr ),
 	InteractorDraggingUI( nullptr ),
 	DraggingUIOffsetTransform( FTransform::Identity ),
@@ -126,12 +131,6 @@ void UVREditorUISystem::Init()
 	GetOwner().GetWorldInteraction().OnViewportInteractionInputAction().AddUObject( this, &UVREditorUISystem::OnVRAction );
 	GetOwner().GetWorldInteraction().OnViewportInteractionHoverUpdate().AddUObject( this, &UVREditorUISystem::OnVRHoverUpdate );
 
-	QuickMenuWidgetClass = LoadClass<UVREditorQuickMenu>( nullptr, TEXT( "/Engine/VREditor/UI/VRQuickMenu.VRQuickMenu_C" ) );
-	check( QuickMenuWidgetClass != nullptr );
-
-	QuickRadialWidgetClass = LoadClass<UVREditorRadialMenu>( nullptr, TEXT( "/Engine/VREditor/UI/VRRadialQuickMenu.VRRadialQuickMenu_C" ) );
-	check( QuickRadialWidgetClass != nullptr );
-
 	TutorialWidgetClass = LoadClass<UVREditorBaseUserWidget>( nullptr, TEXT( "/Engine/VREditor/Tutorial/UI_VR_Tutorial_00.UI_VR_Tutorial_00_C" ) );
 	check( TutorialWidgetClass != nullptr );
 
@@ -148,11 +147,9 @@ void UVREditorUISystem::Init()
 	ShowUISound = LoadObject<USoundCue>( nullptr, TEXT( "/Engine/VREditor/Sounds/VR_open_Cue" ) );
 	check( ShowUISound != nullptr );
 
-	
 	// Create all of our UI panels
 	CreateUIs();
 
-
 	// Bind the color picker creation & destruction overrides
 	SColorPicker::OnColorPickerNonModalCreateOverride.BindUObject(this, &UVREditorUISystem::CreateVRColorPicker);
 	SColorPicker::OnColorPickerDestroyOverride.BindUObject(this, &UVREditorUISystem::DestroyVRColorPicker);
@@ -175,8 +172,6 @@ void UVREditorUISystem::Shutdown()
 
 	DefaultWindowTransforms.Empty();
 
-	QuickMenuWidgetClass = nullptr;
-	QuickRadialWidgetClass = nullptr;
 	TutorialWidgetClass = nullptr;
 	StartDragUISound = nullptr;
 	StopDragUISound = nullptr;
@@ -200,122 +195,145 @@ void UVREditorUISystem::OnVRAction( FEditorViewportClient& ViewportClient, UView
 			if ( Action.ActionType == VRActionTypes::ConfirmRadialSelection )
 			{
 				const EViewportInteractionDraggingMode DraggingMode = Interactor->GetDraggingMode();
+				// Toggle the radial menu by pressing the touchpad or thumb stick
 				if ( Action.Event == IE_Pressed && 
 					 ( DraggingMode != EViewportInteractionDraggingMode::ActorsFreely && DraggingMode != EViewportInteractionDraggingMode::World && 
-					   ( DraggingMode != EViewportInteractionDraggingMode::AssistingDrag && Interactor->GetOtherInteractor() != nullptr && Interactor->GetOtherInteractor()->GetDraggingMode() != EViewportInteractionDraggingMode::World ) ) )
+					   ( DraggingMode != EViewportInteractionDraggingMode::AssistingDrag &&
+						 Interactor->GetOtherInteractor() != nullptr ) ) )
 				{
-					UVREditorRadialMenu* RadialMenu = QuickRadialMenu->GetUserWidget<UVREditorRadialMenu>();
-					if ( IsShowingRadialMenu( VREditorInteractor ) )
+					if ( !IsShowingRadialMenu(VREditorInteractor) )
 					{
-						RadialMenu->Update( VREditorInteractor );
+						TryToSpawnRadialMenu(VREditorInteractor);
+					}
+					else
+					{
+						HideRadialMenu(VREditorInteractor);
 					}
-
-					RadialMenu->SelectCurrentItem();
 					bWasHandled = true;
 				}
 			}
-			else if ( Action.ActionType == ViewportWorldActionTypes::SelectAndMove_LightlyPressed )
+			else 
+			if ( Action.ActionType == ViewportWorldActionTypes::SelectAndMove_LightlyPressed )
 			{
-				FVector LaserPointerStart, LaserPointerEnd;
-				if ( VREditorInteractor->GetLaserPointer( LaserPointerStart, LaserPointerEnd ) )
+				// If the radial menu is showing, select the currently highlighted button by pressing the trigger
+				if ( IsShowingRadialMenu( VREditorInteractor ) )
 				{
-					FHitResult HitResult = VREditorInteractor->GetHitResultFromLaserPointer();
-					if ( HitResult.Actor.IsValid() )
+					TSharedPtr<SRadialBox> RadialBox = StaticCastSharedPtr<SRadialBox>(RadialWidget);
+					if(Action.Event == IE_Pressed)
 					{
-						// Only allow clicks to our own widget components
-						UWidgetComponent* WidgetComponent = Cast<UWidgetComponent>( HitResult.GetComponent() );
-						if ( WidgetComponent != nullptr && IsWidgetAnEditorUIWidget( WidgetComponent ) )
-						{
-							// Always mark the event as handled so that the editor doesn't try to select the widget component
-							bWasHandled = true;
+						RadialBox->SimulateLeftClickDown();
+					}
+					if (Action.Event == IE_Released)
+					{
+						RadialBox->SimulateLeftClickUp();
+						// Hide the radial menu after selection
+						HideRadialMenu(VREditorInteractor);
+					}
+					bWasHandled = true;
+				}
 
-							if ( Action.Event != IE_Repeat )
+				else
+				{
+					FVector LaserPointerStart, LaserPointerEnd;
+					if ( VREditorInteractor->GetLaserPointer( LaserPointerStart, LaserPointerEnd ) )
+					{
+						FHitResult HitResult = VREditorInteractor->GetHitResultFromLaserPointer();
+						if ( HitResult.Actor.IsValid() )
+						{
+							// Only allow clicks to our own widget components
+							UWidgetComponent* WidgetComponent = Cast<UWidgetComponent>( HitResult.GetComponent() );
+							if ( WidgetComponent != nullptr && IsWidgetAnEditorUIWidget( WidgetComponent ) )
 							{
-								// If the Modifier button is held down, treat this like a right click instead of a left click
-								const bool bIsRightClicking =
-									( Action.Event == IE_Pressed && VREditorInteractor->IsModifierPressed() ) ||
-									( Action.Event == IE_Released && VREditorInteractor->IsRightClickingOnUI() );
-
-								FVector2D LastLocalHitLocation = WidgetComponent->GetLastLocalHitLocation();
+								// Always mark the event as handled so that the editor doesn't try to select the widget component
+								bWasHandled = true;
 
-								FVector2D LocalHitLocation;
-								WidgetComponent->GetLocalHitLocation( HitResult.ImpactPoint, LocalHitLocation );
-
-								// If we weren't already hovering over this widget, then we'll reset the last hit location
-								if ( WidgetComponent != VREditorInteractor->GetHoveringOverWidgetComponent() )
+								if ( Action.Event != IE_Repeat )
 								{
-									LastLocalHitLocation = LocalHitLocation;
+									// If the Modifier button is held down, treat this like a right click instead of a left click
+									const bool bIsRightClicking =
+										( Action.Event == IE_Pressed && VREditorInteractor->IsModifierPressed() ) ||
+										( Action.Event == IE_Released && VREditorInteractor->IsRightClickingOnUI() );
 
-									if ( UVREditorWidgetComponent* VRWidgetComponent = Cast<UVREditorWidgetComponent>( VREditorInteractor->GetHoveringOverWidgetComponent() ) )
-									{
-										VRWidgetComponent->SetIsHovering( false );
-									}
-								}
+									FVector2D LastLocalHitLocation = WidgetComponent->GetLastLocalHitLocation();
 
-								FWidgetPath WidgetPathUnderFinger = FWidgetPath( WidgetComponent->GetHitWidgetPath( HitResult.ImpactPoint, /*bIgnoreEnabledStatus*/ false ) );
-								if ( WidgetPathUnderFinger.IsValid() )
-								{
-									TSet<FKey> PressedButtons;
-									if ( Action.Event == IE_Pressed )
-									{
-										PressedButtons.Add( bIsRightClicking ? EKeys::RightMouseButton : EKeys::LeftMouseButton );
-									}
-
-									FPointerEvent PointerEvent(
-										1 + (uint8)VREditorInteractor->GetControllerSide(),
-										LocalHitLocation,
-										LastLocalHitLocation,
-										PressedButtons,
-										bIsRightClicking ? EKeys::RightMouseButton : EKeys::LeftMouseButton,
-										0.0f,	// Wheel delta
-										FModifierKeysState() );
+									FVector2D LocalHitLocation;
+									WidgetComponent->GetLocalHitLocation( HitResult.ImpactPoint, LocalHitLocation );
 
-									VREditorInteractor->SetHoveringOverWidgetComponent( WidgetComponent );
-
-									if ( UVREditorWidgetComponent* VRWidgetComponent = Cast<UVREditorWidgetComponent>( VREditorInteractor->GetHoveringOverWidgetComponent() ) )
+									// If we weren't already hovering over this widget, then we'll reset the last hit location
+									if ( WidgetComponent != VREditorInteractor->GetHoveringOverWidgetComponent() )
 									{
-										VRWidgetComponent->SetIsHovering( true );
+										LastLocalHitLocation = LocalHitLocation;
+
+										if ( UVREditorWidgetComponent* VRWidgetComponent = Cast<UVREditorWidgetComponent>( VREditorInteractor->GetHoveringOverWidgetComponent() ) )
+										{
+											VRWidgetComponent->SetIsHovering( false );
+										}
 									}
 
-									FReply Reply = FReply::Unhandled();
-									if ( Action.Event == IE_Pressed )
+									FWidgetPath WidgetPathUnderFinger = FWidgetPath( WidgetComponent->GetHitWidgetPath( HitResult.ImpactPoint, /*bIgnoreEnabledStatus*/ false ) );
+									if ( WidgetPathUnderFinger.IsValid() )
 									{
-										const double CurrentTime = FPlatformTime::Seconds();
-										if (CurrentTime - VREditorInteractor->GetLastClickReleaseTime() <= VREd::DoubleClickTime->GetFloat())
+										TSet<FKey> PressedButtons;
+										if ( Action.Event == IE_Pressed )
 										{
-											// Trigger a double click event!
-											Reply = FSlateApplication::Get().RoutePointerDoubleClickEvent(WidgetPathUnderFinger, PointerEvent);
+											PressedButtons.Add( bIsRightClicking ? EKeys::RightMouseButton : EKeys::LeftMouseButton );
 										}
-										else
+
+										FPointerEvent PointerEvent(
+											1 + (uint8)VREditorInteractor->GetControllerSide(),
+											LocalHitLocation,
+											LastLocalHitLocation,
+											PressedButtons,
+											bIsRightClicking ? EKeys::RightMouseButton : EKeys::LeftMouseButton,
+											0.0f,	// Wheel delta
+											FModifierKeysState() );
+
+										VREditorInteractor->SetHoveringOverWidgetComponent( WidgetComponent );
+
+										if ( UVREditorWidgetComponent* VRWidgetComponent = Cast<UVREditorWidgetComponent>( VREditorInteractor->GetHoveringOverWidgetComponent() ) )
 										{
-											Reply = FSlateApplication::Get().RoutePointerDownEvent(WidgetPathUnderFinger, PointerEvent);
+											VRWidgetComponent->SetIsHovering( true );
 										}
 
-										// In case of selecting a level in the content browser the VREditormode is closed, this makes sure nothing happens after that.
-										if( !IVREditorModule::Get().IsVREditorModeActive() )
+										FReply Reply = FReply::Unhandled();
+										if ( Action.Event == IE_Pressed )
 										{
-											bWasHandled = true;
-											return;
-										}
+											const double CurrentTime = FPlatformTime::Seconds();
+											if (CurrentTime - VREditorInteractor->GetLastClickReleaseTime() <= VREd::DoubleClickTime->GetFloat())
+											{
+												// Trigger a double click event!
+												Reply = FSlateApplication::Get().RoutePointerDoubleClickEvent(WidgetPathUnderFinger, PointerEvent);
+											}
+											else
+											{
+												Reply = FSlateApplication::Get().RoutePointerDownEvent(WidgetPathUnderFinger, PointerEvent);
+											}
+
+											// In case of selecting a level in the content browser the VREditormode is closed, this makes sure nothing happens after that.
+											if( !IVREditorModule::Get().IsVREditorModeActive() )
+											{
+												bWasHandled = true;
+												return;
+											}
 										
-										VREditorInteractor->SetIsClickingOnUI( true );
-										VREditorInteractor->SetIsRightClickingOnUI( bIsRightClicking );
-										VREditorInteractor->SetLastClickReleaseTime( CurrentTime );
-										bOutIsInputCaptured = true;
+											VREditorInteractor->SetIsClickingOnUI( true );
+											VREditorInteractor->SetIsRightClickingOnUI( bIsRightClicking );
+											VREditorInteractor->SetLastClickReleaseTime( CurrentTime );
+											bOutIsInputCaptured = true;
 
-										// Play a haptic effect on press
-										VREditorInteractor->PlayHapticEffect( VREd::UIPressHapticFeedbackStrength->GetFloat() );
-									}
-									else if ( Action.Event == IE_Released )
-									{
-										Reply = FSlateApplication::Get().RoutePointerUpEvent( WidgetPathUnderFinger, PointerEvent );
+											// Play a haptic effect on press
+											VREditorInteractor->PlayHapticEffect( VREd::UIPressHapticFeedbackStrength->GetFloat() );
+										}
+										else if ( Action.Event == IE_Released )
+										{
+											Reply = FSlateApplication::Get().RoutePointerUpEvent( WidgetPathUnderFinger, PointerEvent );
+										}
 									}
 								}
 							}
 						}
 					}
 				}
-
 				if ( Action.Event == IE_Released )
 				{
 					bool bWasRightClicking = false;
@@ -368,6 +386,16 @@ void UVREditorUISystem::OnVRHoverUpdate( FEditorViewportClient& ViewportClient,
 				FHitResult HitResult = Interactor->GetHitResultFromLaserPointer();
 				if( HitResult.Actor.IsValid() )
 				{
+					// If a button is currently highlighted on the radial menu, don't highlight another button with the laser
+					if ((HitResult.Actor == QuickRadialMenu) && (RadialWidget.Get() != nullptr))
+					{
+						const TSharedPtr<SRadialBox>& RadialBox = StaticCastSharedPtr<SRadialBox>(RadialWidget);
+						if (RadialBox->GetCurrentlyHoveredButton().Get() != nullptr)
+						{
+							return;
+						}
+					}
+
 					// Only allow clicks to our own widget components
 					UWidgetComponent* WidgetComponent = Cast<UWidgetComponent>( HitResult.GetComponent() );
 					if( WidgetComponent != nullptr && IsWidgetAnEditorUIWidget( WidgetComponent ) )
@@ -558,14 +586,14 @@ void UVREditorUISystem::Tick( FEditorViewportClient* ViewportClient, const float
 			{
 				// @todo vreditor tweak: Weird to hard code this here. Probably should be an accessor on the hand itself, and based on the actual device type
 				const FTransform UICapsuleTransform = OtherInteractor->GetTransform();
-				const FVector UICapsuleStart = FVector( -9.0f, 0.0f, 0.0f ) * WorldScaleFactor;
+				const FVector UICapsuleStart = FVector( -3.0f, 0.0f, 0.0f ) * WorldScaleFactor;
 				const FVector UICapsuleEnd = FVector( -18.0f, 0.0f, 0.0f ) * WorldScaleFactor;
-				const float UICapsuleLocalRadius = 6.0f * WorldScaleFactor;
+				const float UICapsuleLocalRadius = 8.0f * WorldScaleFactor;
 				const float MinDistanceToUICapsule = 8.0f * WorldScaleFactor; // @todo vreditor tweak
 				const FVector UIForwardVector = FVector::UpVector;
 				const float MinDotForAimingAtOtherHand = 0.25f;	// @todo vreditor tweak
 
-				if( GetOwner().IsHandAimingTowardsCapsule( Interactor, UICapsuleTransform, UICapsuleStart, UICapsuleEnd, UICapsuleLocalRadius, MinDistanceToUICapsule, UIForwardVector, MinDotForAimingAtOtherHand ) )
+				if (GetOwner().IsHandAimingTowardsCapsule(Interactor, UICapsuleTransform, UICapsuleStart, UICapsuleEnd, UICapsuleLocalRadius, MinDistanceToUICapsule, UIForwardVector, MinDotForAimingAtOtherHand) && (QuickRadialMenu->bHidden))
 				{
 					bShowQuickMenuOnArm = true;
 				}
@@ -583,9 +611,26 @@ void UVREditorUISystem::Tick( FEditorViewportClient* ViewportClient, const float
 		{
 			// Despawn
 			HandInteractorWithQuickMenu->SetHasUIOnForearm( false );
+			HandInteractorWithQuickMenu->SetHasUIInFront( false );
 			QuickMenuUI->ShowUI( false );
+			// Reset all the currently animating buttons to their minimum scale and stop animating them
+			for (FVRButton& VRButton : VRButtons)
+			{
+				if(VRButton.AnimationDirection != EVREditorAnimationState::None)
+				{
+					
+					TSharedRef<SWidget> PlaceholderWidget = (VRButton.BaseSlot->DetachWidget()).ToSharedRef();
+					TSharedRef<SWidget> WidgetToTransfer = (VRButton.HoverSlot->DetachWidget()).ToSharedRef();
+					VRButton.BaseSlot->AttachWidget(WidgetToTransfer);
+					VRButton.HoverSlot->AttachWidget(PlaceholderWidget);
+					VRButton.CurrentScale = VRButton.MinScale;
+					VRButton.ButtonBorder->SetRenderTransform(FSlateRenderTransform::FTransform2D(VRButton.CurrentScale));
+					VRButton.AnimationDirection = EVREditorAnimationState::None;
+		}
+			}
 		}
 
+
 		if( HandInteractorThatNeedsQuickMenu != nullptr && !QuickMenuUI->IsUIVisible() )
 		{
 			UVREditorMotionControllerInteractor* MotionControllerHandInteractorThatNeedsQuickMenu = Cast<UVREditorMotionControllerInteractor>( HandInteractorThatNeedsQuickMenu );
@@ -594,17 +639,49 @@ void UVREditorUISystem::Tick( FEditorViewportClient* ViewportClient, const float
 				const AVREditorFloatingUI::EDockedTo DockTo = ( MotionControllerHandInteractorThatNeedsQuickMenu->GetControllerSide() == EControllerHand::Left ) ? AVREditorFloatingUI::EDockedTo::LeftArm : AVREditorFloatingUI::EDockedTo::RightArm;
 				QuickMenuUI->SetDockedTo( DockTo );
 				QuickMenuUI->ShowUI( true );
+				MotionControllerHandInteractorThatNeedsQuickMenu->SetHasUIInFront( true );
 				MotionControllerHandInteractorThatNeedsQuickMenu->SetHasUIOnForearm( true );
 			}
 		}
 	}
-
-	// Close the radial menu if it was not updated for a while
-	RadialMenuHideDelayTime += DeltaTime;
-	if( RadialMenuHideDelayTime >= VREd::RadialMenuFadeDelay->GetFloat() )
+	
+	// Iterate through all quick menu and radial menu buttons and animate any that need it
+	for (FVRButton& VRButton : VRButtons)
 	{
-		UVREditorInteractor* Interactor = GetOwner().GetHandInteractor( QuickRadialMenu->GetDockedTo() == AVREditorFloatingUI::EDockedTo::LeftHand ? EControllerHand::Left : EControllerHand::Right );
-		HideRadialMenu( Interactor );
+		switch (VRButton.AnimationDirection)
+		{
+			case EVREditorAnimationState::Forward:
+			{
+				if (VRButton.CurrentScale < VRButton.MaxScale)
+				{
+					VRButton.CurrentScale += DeltaTime*VRButton.ScaleRate;
+				}
+				else
+				{
+					VRButton.CurrentScale = VRButton.MaxScale;
+					VRButton.AnimationDirection = EVREditorAnimationState::None;
+				}
+				break;
+			}
+			case EVREditorAnimationState::Backward:
+			{
+				if (VRButton.CurrentScale > VRButton.MinScale)
+				{
+					VRButton.CurrentScale -= DeltaTime*VRButton.ScaleRate;
+				}
+				else
+				{
+					VRButton.CurrentScale = VRButton.MinScale;
+					VRButton.AnimationDirection = EVREditorAnimationState::None;	
+				}
+				break;
+			}
+			case EVREditorAnimationState::None:
+			{
+				break;
+			}
+		}
+		VRButton.ButtonBorder->SetRenderTransform(FSlateRenderTransform::FTransform2D(VRButton.CurrentScale));	
 	}
 	
 	// Tick all of our floating UIs
@@ -627,37 +704,38 @@ void UVREditorUISystem::CreateUIs()
 	const FIntPoint DefaultResolution( VREd::DefaultEditorUIResolutionX->GetInt(), VREd::DefaultEditorUIResolutionY->GetInt() );
 
 	const bool bIsVREditorDemo = FParse::Param( FCommandLine::Get(), TEXT( "VREditorDemo" ) );	// @todo vreditor: Remove this when no longer needed
-																								
-	// @todo vreditor: Tweak
-	if ( QuickMenuWidgetClass != nullptr )
 	{
-		const FIntPoint Resolution( VREd::QuickMenuUIResolutionX->GetInt(), VREd::QuickMenuUIResolutionY->GetInt() );
-
+		const FIntPoint Resolution(VREd::QuickMenuUIResolutionX->GetInt(), VREd::QuickMenuUIResolutionY->GetInt());
 		const bool bWithSceneComponent = false;
-		QuickMenuUI = GetOwner().SpawnTransientSceneActor< AVREditorFloatingUI >(TEXT("QuickMenu"), bWithSceneComponent);
-		QuickMenuUI->SetUMGWidget( *this, QuickMenuWidgetClass, Resolution, 30.0f, AVREditorFloatingUI::EDockedTo::Nothing );
-		QuickMenuUI->ShowUI( false );
-		QuickMenuUI->SetRelativeOffset( FVector( -11.0f, 0.0f, 3.0f ) );
-		FloatingUIs.Add( QuickMenuUI );
-	}
+		// @todo vreditor: Tweak
+		{
+			QuickMenuUI = GetOwner().SpawnTransientSceneActor< AVREditorFloatingUI >(TEXT("QuickMenu"), bWithSceneComponent);
+			QuickMenuUI->SetSlateWidget(*this, BuildQuickMenuWidget(), Resolution, 40.0f, AVREditorFloatingUI::EDockedTo::Nothing);
+			QuickMenuUI->ShowUI(false);
+			QuickMenuUI->SetRelativeOffset(FVector(-10.0f, 0.0f, 5.0f));
+			QuickMenuUI->GetWidgetComponent()->SetOpacityFromTexture(1.0f);
+			QuickMenuUI->GetWidgetComponent()->SetBackgroundColor(FLinearColor::Transparent);
+			QuickMenuUI->GetWidgetComponent()->SetBlendMode(EWidgetBlendMode::Masked);
+			FloatingUIs.Add(QuickMenuUI);
 
-	// Create the radial UI
-	if ( QuickRadialWidgetClass != nullptr )
-	{
-		TSharedPtr<SWidget> WidgetToDraw;
-		QuickRadialMenu = GetOwner().SpawnTransientSceneActor< AVREditorFloatingUI >(TEXT("QuickRadialmenu"), false);
-		QuickRadialMenu->SetUMGWidget( *this, QuickRadialWidgetClass, DefaultResolution, 40.0f, AVREditorFloatingUI::EDockedTo::Nothing );
-		QuickRadialMenu->ShowUI( false );
-		QuickRadialMenu->SetActorEnableCollision( false );
-		QuickRadialMenu->SetRelativeOffset( FVector( 10.0f, 0.0f, 10.f ) );
-		QuickRadialMenu->SetCollisionOnShow( false );
-		FloatingUIs.Add(QuickRadialMenu);
+		}
+		// Create the radial UI
+		{
+			QuickRadialMenu = GetOwner().SpawnTransientSceneActor< AVREditorFloatingUI >(TEXT("QuickRadialmenu"), bWithSceneComponent);
+			QuickRadialMenu->SetSlateWidget(*this, BuildRadialMenuWidget(), DefaultResolution, 40.0f, AVREditorFloatingUI::EDockedTo::Nothing);
+			QuickRadialMenu->ShowUI(false);
+			QuickRadialMenu->SetActorEnableCollision(false);
+			QuickRadialMenu->SetRelativeOffset(FVector(0.0f, 0.0f, 5.0f));
+			QuickRadialMenu->GetWidgetComponent()->SetOpacityFromTexture(1.0f);
+			QuickRadialMenu->GetWidgetComponent()->SetBackgroundColor(FLinearColor::Transparent);
+			QuickRadialMenu->GetWidgetComponent()->SetBlendMode(EWidgetBlendMode::Masked);
+			FloatingUIs.Add(QuickRadialMenu);
+		}
 	}
-
 	// Make some editor UIs!
 	{
 		{
-			const FIntPoint Resolution( VREd::ContentBrowserUIResolutionX->GetInt(), VREd::ContentBrowserUIResolutionY->GetInt() );
+			const FIntPoint Resolution = FIntPoint( VREd::ContentBrowserUIResolutionX->GetInt(), VREd::ContentBrowserUIResolutionY->GetInt() );
 
 			IContentBrowserSingleton& ContentBrowserSingleton = FModuleManager::Get().LoadModuleChecked<FContentBrowserModule>( "ContentBrowser" ).Get();;
 
@@ -795,6 +873,7 @@ void UVREditorUISystem::CreateUIs()
 
 			AVREditorFloatingUI* TutorialUI = GetOwner().SpawnTransientSceneActor< AVREditorDockableWindow >( TEXT( "TutorialUI" ), false );
 			TutorialUI->SetUMGWidget( *this, TutorialWidgetClass, FIntPoint( VREd::TutorialUIResolutionX->GetFloat(), VREd::TutorialUIResolutionY->GetFloat() ), VREd::TutorialUISize->GetFloat(), bShowAtStart ? AVREditorFloatingUI::EDockedTo::Room : AVREditorFloatingUI::EDockedTo::Nothing );
+
 			if( bShowAtStart )
 			{
 				TutorialUI->SetRelativeOffset( FVector( VREd::TutorialUILocationX->GetFloat(), VREd::TutorialUILocationY->GetFloat(), VREd::TutorialUILocationZ->GetFloat() ) );
@@ -809,7 +888,7 @@ void UVREditorUISystem::CreateUIs()
 		}
 
 		{
-			const FIntPoint Resolution(VREd::AssetEditorUIResolutionX->GetInt(), VREd::AssetEditorUIResolutionY->GetInt());
+			const FIntPoint Resolution = FIntPoint(VREd::AssetEditorUIResolutionX->GetInt(), VREd::AssetEditorUIResolutionY->GetInt());
 
 			const bool bWithSceneComponent = false;
 			AVREditorDockableWindow* TabManagerUI = GetOwner().SpawnTransientSceneActor< AVREditorDockableWindow >( TEXT( "AssetEditor" ), bWithSceneComponent );
@@ -1009,10 +1088,6 @@ void UVREditorUISystem::ShowEditorUIPanel( AVREditorFloatingUI* Panel, UVREditor
 			UGameplayStatics::PlaySound2D( GetOwner().GetWorld(), bShouldShow ? ShowUISound : HideUISound );
 		}
 
-		if ( bRefreshQuickMenu && QuickMenuUI )
-		{
-			QuickMenuUI->GetUserWidget<UVREditorQuickMenu>()->RefreshUI();
-		}
 	}
 }
 
@@ -1025,19 +1100,7 @@ bool UVREditorUISystem::IsShowingRadialMenu( UVREditorInteractor* Interactor ) c
 }
 
 
-void UVREditorUISystem::UpdateRadialMenu( UVREditorInteractor* Interactor )
-{
-	if( QuickRadialMenu->bHidden )
-	{
-		QuickRadialMenu->ShowUI( true, false );
-	}
 
-	if( !QuickRadialMenu->bHidden )
-	{
-		RadialMenuHideDelayTime = 0.0f;
-		QuickRadialMenu->GetUserWidget<UVREditorRadialMenu>()->Update( Interactor );
-	}
-}
 
 
 void UVREditorUISystem::TryToSpawnRadialMenu( UVREditorInteractor* Interactor )
@@ -1066,19 +1129,11 @@ void UVREditorUISystem::TryToSpawnRadialMenu( UVREditorInteractor* Interactor )
 				( InteractorDraggingUI == nullptr || InteractorDraggingUI != Interactor ) &&
 				!Interactor->IsHoveringOverUI();	// Don't show radial menu when aiming at a UI  (too much clutter)
 
-			UVREditorRadialMenu* RadialMenu = QuickRadialMenu->GetUserWidget<UVREditorRadialMenu>();
-			// We need to update the trackpad position in the radialmenu before checking if it can be used
-			RadialMenu->Update( Interactor );
-			if( RadialMenu && RadialMenu->IsInMenuRadius() )
-			{
-				bNeedsSpawn = false;
-			}
-
-			if( bNeedsSpawn )
+			if (bNeedsSpawn)
 			{
 				const AVREditorFloatingUI::EDockedTo DockedTo = MotionControllerInteractor->GetControllerSide() == EControllerHand::Left ? AVREditorFloatingUI::EDockedTo::LeftHand : AVREditorFloatingUI::EDockedTo::RightHand;
-				QuickRadialMenu->SetDockedTo( DockedTo );
-				QuickRadialMenu->ShowUI( true );
+				QuickRadialMenu->SetDockedTo(DockedTo);
+				QuickRadialMenu->ShowUI(true);
 			}
 		}
 	}
@@ -1325,11 +1380,7 @@ void UVREditorUISystem::TogglePanelsVisibility()
 		VREditorInteractor->SetHasUIInFront( false );
 	}
 
-	if (QuickMenuUI)
-	{
-		QuickMenuUI->GetUserWidget<UVREditorQuickMenu>()->RefreshUI();
 	}
-}
 
 float UVREditorUISystem::GetMaxDockWindowSize() const
 {
@@ -1520,3 +1571,678 @@ void UVREditorUISystem::DestroyVRColorPicker()
 	const bool bRefreshQuickMenu = true;
 	ShowEditorUIPanel(EEditorUIPanel::ColorPicker, VREditorInteractor, bShouldShow, bShowOnHand, bRefreshQuickMenu);
 }
+
+void UVREditorUISystem::MakeUniformGridMenu(const TSharedRef<FMultiBox>& MultiBox, const TSharedRef<SMultiBoxWidget>& MultiBoxWidget, int32 Columns)
+{
+	// Grab the list of blocks, early out if there's nothing to fill the widget with
+	const TArray< TSharedRef<const FMultiBlock > >& Blocks = MultiBox->GetBlocks();
+	const uint32 NumItems = Blocks.Num();
+	if( NumItems == 0 )
+	{
+		return;
+	}
+
+	TSharedRef<SUniformGridPanel> GridPanel = SNew( SUniformGridPanel )
+		.SlotPadding(3.0f);
+	MultiBox.Get().SetStyle(&FVREditorStyle::Get(), FVREditorStyle::GetStyleSetName());
+	int32 Index = 0;
+	int32 Row = 0;
+	for( const TSharedRef<const FMultiBlock>& MultiBlock : Blocks )
+	{
+		const TSharedRef<const FMultiBlock>& Block = MultiBlock;
+		if(Block->GetType() == EMultiBlockType::MenuEntry)
+		{
+			TSharedRef<SWidget> BlockWidget = Block->MakeWidget(MultiBoxWidget, EMultiBlockLocation::Middle, true)->AsWidget();
+		
+			int32 Column = FMath::Fmod(Index, Columns);
+			Row = FMath::DivideAndRoundDown(Index, Columns);
+			TSharedRef<SOverlay> TestOverlay = SNew(SOverlay);
+			GridPanel->AddSlot(Column, Row)
+				.VAlign(VAlign_Fill)
+				.HAlign(HAlign_Fill)
+				[
+					SAssignNew(TestOverlay, SOverlay)
+				];
+
+			AddHoverableButton(BlockWidget, TEXT("SMenuEntryButton"), TestOverlay);
+			SetButtonTextWrap(BlockWidget, 75.0);
+			Index++;
+		}
+	}
+
+	MultiBoxWidget->SetContent(
+		SNew(SScrollBox)
+		+ SScrollBox::Slot()
+		.Padding(0, 0, 0, 0)
+		.VAlign(VAlign_Fill)
+		.HAlign(HAlign_Fill)
+		[
+			GridPanel
+		]
+	);
+}
+
+TSharedRef<SWidget> UVREditorUISystem::AddHoverableButton(TSharedRef<SWidget>& BlockWidget, FName ButtonType, TSharedRef<SOverlay>& TestOverlay)
+{
+	TSharedPtr<SBorder> ButtonBorder;
+	SOverlay::FOverlaySlot* BaseSlot;
+	SOverlay::FOverlaySlot* HoverSlot;
+	TestOverlay->AddSlot(0)
+	.Expose(BaseSlot)
+		[
+			SAssignNew(ButtonBorder, SBorder)
+			.BorderImage(FEditorStyle::GetBrush("ToolPanel.GroupBorder"))
+		[
+			BlockWidget
+		]
+		];
+	// Placeholder slot to put the button into when it is hovered
+	TestOverlay->AddSlot(1500)
+	.Expose(HoverSlot)
+		[
+			SNew(SOverlay)
+			.Visibility(EVisibility::SelfHitTestInvisible)
+		];
+
+	TSharedRef<SWidget> TestWidget = FindWidgetOfType(BlockWidget, ButtonType);
+	if (TestWidget != SNullWidget::NullWidget)
+	{
+		TSharedRef<SButton> Button = StaticCastSharedRef<SButton>(TestWidget);
+		Button->SetRenderTransformPivot(FVector2D(0.5, 0.5));
+		Button->SetOnHovered(FSimpleDelegate::CreateUObject(this, &UVREditorUISystem::OnHoverBeginEffect, Button));
+		Button->SetOnUnhovered(FSimpleDelegate::CreateUObject(this, &UVREditorUISystem::OnHoverEndEffect, Button));
+		ButtonBorder->SetRenderTransformPivot(FVector2D(0.5, 0.5));
+		FVRButton ButtonToAdd = FVRButton(Button, BaseSlot, HoverSlot, ButtonBorder);
+		VRButtons.Add(ButtonToAdd);
+	}
+	return BlockWidget;
+}
+
+TSharedRef<SWidget> UVREditorUISystem::SetButtonTextWrap(TSharedRef<SWidget>& BlockWidget, float WrapSize)
+{
+	TSharedRef<SWidget> TestWidget = FindWidgetOfType(BlockWidget, FName(TEXT("STextBlock")));
+	if (TestWidget != SNullWidget::NullWidget)
+	{
+		TSharedRef<STextBlock> TextBlock = StaticCastSharedRef<STextBlock>(TestWidget);
+		TextBlock->SetWrapTextAt(WrapSize);
+		TextBlock->SetJustification(ETextJustify::Center);
+	}			
+	return BlockWidget;
+}
+
+void UVREditorUISystem::MakeRadialBoxMenu(const TSharedRef<FMultiBox>& MultiBox, const TSharedRef<SMultiBoxWidget>& MultiBoxWidget)
+{
+	// Grab the list of blocks, early out if there's nothing to fill the widget with
+	const TArray< TSharedRef< const FMultiBlock > >& Blocks = MultiBox->GetBlocks();
+	const uint32 NumItems = Blocks.Num();
+	if (NumItems == 0)
+	{
+		return;
+	}
+	MultiBox.Get().SetStyle(&FVREditorStyle::Get(), FVREditorStyle::GetStyleSetName());
+	TSharedRef<SRadialBox> GridPanel = SNew(SRadialBox);
+	GridPanel->NumEntries = Blocks.Num();
+	for (const TSharedRef<const FMultiBlock>& MultiBlock : Blocks)
+	{
+		const TSharedRef<const FMultiBlock>& Block = MultiBlock;
+		TSharedRef<SWidget> BlockWidget = Block->MakeWidget(MultiBoxWidget, EMultiBlockLocation::Middle, true)->AsWidget();
+
+		
+		TSharedRef<SOverlay> TestOverlay = SNew(SOverlay);
+		
+		GridPanel->AddSlot()
+		.VAlign(VAlign_Center)
+		.HAlign(HAlign_Center)
+		[
+			SAssignNew(TestOverlay, SOverlay)
+		];
+		
+		BlockWidget = AddHoverableButton(BlockWidget, TEXT("SButton"), TestOverlay);
+		BlockWidget = SetButtonTextWrap(BlockWidget, 75.0);
+
+	}
+	RadialWidget = GridPanel;
+	StaticCastSharedPtr<SRadialBox>(RadialWidget)->NumEntries = Blocks.Num();
+	MultiBoxWidget->SetContent(
+			SNew(SDPIScaler)
+			.DPIScale(3)
+			[
+				SNew(SVerticalBox)
+				+ SVerticalBox::Slot()
+				.Padding(10, 10, 10, 10)
+				.VAlign(VAlign_Fill)
+				.HAlign(HAlign_Fill)
+				[
+					GridPanel
+				]
+			]
+		);
+}
+
+
+TSharedRef<SWidget> UVREditorUISystem::BuildQuickMenuWidget()
+{
+	FVREditorActionCallbacks::UpdateGizmoCoordinateSystemText(VRMode);
+	FVREditorActionCallbacks::UpdateGizmoModeText(VRMode);
+
+	FMultiBox::FOnMakeMultiBoxBuilderOverride VREditorMenuBuilderOverride;
+	VREditorMenuBuilderOverride.BindUObject(this, &UVREditorUISystem::MakeUniformGridMenu, 3);
+	const TSharedRef<FExtender> MenuExtender = MakeShareable(new FExtender());
+	TSharedPtr<FUICommandList> CommandList = MakeShareable(new FUICommandList());
+	FMenuBuilder SnapMenuBuilder(false, CommandList, MenuExtender);
+
+	// Snap menu
+	SnapMenuBuilder.AddMenuEntry(
+		LOCTEXT("ToggleTranslationSnap", "Translate"),
+		LOCTEXT("ToggleTranslationSnapTooltip", "Toggle Translation Snap"),
+		FSlateIcon(FEditorStyle::GetStyleSetName(), "EditorViewport.TranslateMode.Small"),
+		FUIAction
+		(
+			FExecuteAction::CreateStatic(&FLevelEditorActionCallbacks::LocationGridSnap_Clicked),
+			FCanExecuteAction::CreateStatic(&FLevelEditorActionCallbacks::DefaultCanExecuteAction),
+			FGetActionCheckState::CreateStatic(&FVREditorActionCallbacks::GetTranslationSnapState)
+		),
+		NAME_None,
+		EUserInterfaceActionType::Check
+		);
+	SnapMenuBuilder.AddMenuEntry(
+		LOCTEXT("ToggleRotationSnap", "Rotate"),
+		LOCTEXT("ToggleRotationSnapTooltip", "Toggle Rotation Snap"),
+		FSlateIcon(FEditorStyle::GetStyleSetName(), "EditorViewport.RotateMode.Small"),
+		FUIAction
+		(
+			FExecuteAction::CreateStatic(&FLevelEditorActionCallbacks::RotationGridSnap_Clicked),
+			FCanExecuteAction::CreateStatic(&FLevelEditorActionCallbacks::DefaultCanExecuteAction),
+			FGetActionCheckState::CreateStatic(&FVREditorActionCallbacks::GetRotationSnapState)
+			),
+		NAME_None,
+		EUserInterfaceActionType::Check
+		);
+	SnapMenuBuilder.AddMenuEntry(
+		LOCTEXT("ToggleScaleSnap", "Scale"),
+		LOCTEXT("ToggleScaleSnapTooltip", "Toggle Scale Snap"),
+		FSlateIcon(FEditorStyle::GetStyleSetName(), "EditorViewport.ScaleMode.Small"),
+		FUIAction
+		(
+			FExecuteAction::CreateStatic(&FLevelEditorActionCallbacks::ScaleGridSnap_Clicked),
+			FCanExecuteAction::CreateStatic(&FLevelEditorActionCallbacks::DefaultCanExecuteAction),
+			FGetActionCheckState::CreateStatic(&FVREditorActionCallbacks::GetScaleSnapState)
+			),
+		NAME_None,
+		EUserInterfaceActionType::Check
+		);
+	TAttribute<FText> DynamicTranslationSizeLabel;
+	DynamicTranslationSizeLabel.BindStatic(&FVREditorActionCallbacks::GetTranslationSnapSizeText);
+	SnapMenuBuilder.AddMenuEntry(
+		DynamicTranslationSizeLabel,
+		LOCTEXT("ToggleTranslationSnapSizeTooltip", "Toggle Translation Snap Size"),
+		FSlateIcon(),
+		FUIAction
+		(
+			FExecuteAction::CreateStatic(&FVREditorActionCallbacks::OnTranslationSnapSizeButtonClicked),
+			FCanExecuteAction::CreateStatic(&FLevelEditorActionCallbacks::DefaultCanExecuteAction)
+			)
+		);
+	TAttribute<FText> DynamicRotationSizeLabel;
+	DynamicRotationSizeLabel.BindStatic(&FVREditorActionCallbacks::GetRotationSnapSizeText);
+	SnapMenuBuilder.AddMenuEntry(
+		DynamicRotationSizeLabel,
+		LOCTEXT("ToggleRotationSnapSizeTooltip", "Toggle Rotation Snap Size"),
+		FSlateIcon(),
+		FUIAction
+		(
+			FExecuteAction::CreateStatic(&FVREditorActionCallbacks::OnRotationSnapSizeButtonClicked),
+			FCanExecuteAction::CreateStatic(&FLevelEditorActionCallbacks::DefaultCanExecuteAction)
+			)
+		);
+	TAttribute<FText> DynamicScaleSizeLabel;
+	DynamicScaleSizeLabel.BindStatic(&FVREditorActionCallbacks::GetScaleSnapSizeText);
+	SnapMenuBuilder.AddMenuEntry(
+		DynamicScaleSizeLabel,
+		LOCTEXT("ToggleScaleSnapSizeTooltip", "Toggle Scale Snap Size"),
+		FSlateIcon(),
+		FUIAction
+		(
+			FExecuteAction::CreateStatic(&FVREditorActionCallbacks::OnScaleSnapSizeButtonClicked),
+			FCanExecuteAction::CreateStatic(&FLevelEditorActionCallbacks::DefaultCanExecuteAction)
+			)
+		);
+
+	TSharedRef<SWidget> SnapMenuWidget = SnapMenuBuilder.MakeWidget(&VREditorMenuBuilderOverride);
+
+	// Next Menu section: Interface
+	CommandList = MakeShareable(new FUICommandList());
+	FMenuBuilder InterfaceMenuBuilder(false, CommandList, MenuExtender);
+
+	InterfaceMenuBuilder.AddMenuEntry(
+		LOCTEXT("ActorDetails", "Details"),
+		LOCTEXT("ActorDetailsTooltip", "Details"),
+		FSlateIcon(FEditorStyle::GetStyleSetName(), "LevelEditor.Tabs.Details"),
+		FUIAction
+		(
+			FExecuteAction::CreateStatic(&FVREditorActionCallbacks::OnUIToggleButtonClicked, VRMode, UVREditorUISystem::EEditorUIPanel::ActorDetails),
+			FCanExecuteAction::CreateStatic(&FLevelEditorActionCallbacks::DefaultCanExecuteAction),
+			FGetActionCheckState::CreateStatic(&FVREditorActionCallbacks::GetUIToggledState, VRMode, UVREditorUISystem::EEditorUIPanel::ActorDetails)
+			),
+		NAME_None,
+		EUserInterfaceActionType::Check
+		);
+	InterfaceMenuBuilder.AddMenuEntry(
+		LOCTEXT("ContentBrowser", "Content Browser"),
+		LOCTEXT("ContentBrowserTooltip", "Content Browser"),
+		FSlateIcon(FEditorStyle::GetStyleSetName(), "LevelEditor.Tabs.ContentBrowser"),
+		FUIAction
+		(
+			FExecuteAction::CreateStatic(&FVREditorActionCallbacks::OnUIToggleButtonClicked, VRMode, UVREditorUISystem::EEditorUIPanel::ContentBrowser),
+			FCanExecuteAction::CreateStatic(&FLevelEditorActionCallbacks::DefaultCanExecuteAction),
+			FGetActionCheckState::CreateStatic(&FVREditorActionCallbacks::GetUIToggledState, VRMode, UVREditorUISystem::EEditorUIPanel::ContentBrowser)
+			),
+		NAME_None,
+		EUserInterfaceActionType::Check
+		);
+	InterfaceMenuBuilder.AddMenuEntry(
+		LOCTEXT("Modes", "Modes"),
+		LOCTEXT("ModesTooltip", "Modes"),
+		FSlateIcon(FEditorStyle::GetStyleSetName(), "LevelEditor.Tabs.Modes"),
+		FUIAction
+		(
+			FExecuteAction::CreateStatic(&FVREditorActionCallbacks::OnUIToggleButtonClicked, VRMode, UVREditorUISystem::EEditorUIPanel::Modes),
+			FCanExecuteAction::CreateStatic(&FLevelEditorActionCallbacks::DefaultCanExecuteAction),
+			FGetActionCheckState::CreateStatic(&FVREditorActionCallbacks::GetUIToggledState, VRMode, UVREditorUISystem::EEditorUIPanel::Modes)
+			),
+		NAME_None,
+		EUserInterfaceActionType::Check
+		);
+	InterfaceMenuBuilder.AddMenuEntry(
+		LOCTEXT("WorldOutliner", "World Outliner"),
+		LOCTEXT("WorldOutlinerTooltip", "World Outliner"),
+		FSlateIcon(FEditorStyle::GetStyleSetName(), "LevelEditor.Tabs.Outliner"),
+		FUIAction
+		(
+			FExecuteAction::CreateStatic(&FVREditorActionCallbacks::OnUIToggleButtonClicked, VRMode, UVREditorUISystem::EEditorUIPanel::WorldOutliner),
+			FCanExecuteAction::CreateStatic(&FLevelEditorActionCallbacks::DefaultCanExecuteAction),
+			FGetActionCheckState::CreateStatic(&FVREditorActionCallbacks::GetUIToggledState, VRMode, UVREditorUISystem::EEditorUIPanel::WorldOutliner)
+			),
+		NAME_None,
+		EUserInterfaceActionType::Check
+		);
+	InterfaceMenuBuilder.AddMenuEntry(
+		LOCTEXT("AssetEditor", "Asset Editor"),
+		LOCTEXT("AssetEditorTooltip", "Asset Editor"),
+		FSlateIcon(),
+		FUIAction
+		(
+			FExecuteAction::CreateStatic(&FVREditorActionCallbacks::OnUIToggleButtonClicked, VRMode, UVREditorUISystem::EEditorUIPanel::AssetEditor),
+			FCanExecuteAction::CreateStatic(&FLevelEditorActionCallbacks::DefaultCanExecuteAction),
+			FGetActionCheckState::CreateStatic(&FVREditorActionCallbacks::GetUIToggledState, VRMode, UVREditorUISystem::EEditorUIPanel::AssetEditor)
+			),
+		NAME_None,
+		EUserInterfaceActionType::Check
+		);
+	InterfaceMenuBuilder.AddMenuEntry(
+		LOCTEXT("Tutorial", "Tutorial"),
+		LOCTEXT("TutorialTooltip", "Tutorial"),
+		FSlateIcon(FEditorStyle::GetStyleSetName(), "Tutorials.Browser.DefaultTutorialIcon"),
+		FUIAction
+		(
+			FExecuteAction::CreateStatic(&FVREditorActionCallbacks::OnUIToggleButtonClicked, VRMode, UVREditorUISystem::EEditorUIPanel::Tutorial),
+			FCanExecuteAction::CreateStatic(&FLevelEditorActionCallbacks::DefaultCanExecuteAction),
+			FGetActionCheckState::CreateStatic(&FVREditorActionCallbacks::GetUIToggledState, VRMode, UVREditorUISystem::EEditorUIPanel::Tutorial)
+			),
+		NAME_None,
+		EUserInterfaceActionType::Check
+		);
+	InterfaceMenuBuilder.AddMenuEntry(
+		LOCTEXT("WorldSettings", "World Settings"),
+		LOCTEXT("WorldSettingsTooltip", "World Settings"),
+		FSlateIcon(FEditorStyle::GetStyleSetName(), "LevelEditor.WorldProperties.Tab"),
+		FUIAction
+		(
+			FExecuteAction::CreateStatic(&FVREditorActionCallbacks::OnUIToggleButtonClicked, VRMode, UVREditorUISystem::EEditorUIPanel::WorldSettings),
+			FCanExecuteAction::CreateStatic(&FLevelEditorActionCallbacks::DefaultCanExecuteAction),
+			FGetActionCheckState::CreateStatic(&FVREditorActionCallbacks::GetUIToggledState, VRMode, UVREditorUISystem::EEditorUIPanel::WorldSettings)
+			),
+		NAME_None,
+		EUserInterfaceActionType::Check
+		);
+	InterfaceMenuBuilder.AddMenuEntry(
+		LOCTEXT("ToggleAllUI", "Toggle All UI"),
+		LOCTEXT("ToggleAllUITooltip", "ToggleAllUI"),
+		FSlateIcon(),
+		FUIAction
+		(
+			FExecuteAction::CreateUObject(this, &UVREditorUISystem::TogglePanelsVisibility),
+			FCanExecuteAction::CreateStatic(&FLevelEditorActionCallbacks::DefaultCanExecuteAction)
+			)
+		);
+	
+
+	TSharedRef<SWidget> InterfaceMenuWidget = InterfaceMenuBuilder.MakeWidget(&VREditorMenuBuilderOverride);
+
+	// Next Menu section: Gizmo
+	CommandList = MakeShareable(new FUICommandList());
+	FMenuBuilder GizmoMenuBuilder(false, CommandList, MenuExtender);
+
+	TAttribute<FText> DynamicGizmoCoordinateSystemLabel;
+	DynamicGizmoCoordinateSystemLabel.BindStatic(&FVREditorActionCallbacks::GetGizmoCoordinateSystemText);
+	GizmoMenuBuilder.AddMenuEntry(
+		DynamicGizmoCoordinateSystemLabel,
+		LOCTEXT("GizmoCoordinateTooltip", "Toggle Gizmo Coordinates"),
+		FSlateIcon(),
+		FUIAction
+		(
+			FExecuteAction::CreateStatic(&FVREditorActionCallbacks::OnGizmoCoordinateSystemButtonClicked, VRMode),
+			FCanExecuteAction::CreateStatic(&FLevelEditorActionCallbacks::DefaultCanExecuteAction)
+			)
+		);
+	TAttribute<FText> DynamicGizmoModeLabel;
+	DynamicGizmoModeLabel.BindStatic(&FVREditorActionCallbacks::GetGizmoModeText);
+	GizmoMenuBuilder.AddMenuEntry(
+		DynamicGizmoModeLabel,
+		LOCTEXT("GizmoModeTooltip", "Toggle Gizmo Mode"),
+		FSlateIcon(),
+		FUIAction
+		(
+			FExecuteAction::CreateStatic(&FVREditorActionCallbacks::OnGizmoModeButtonClicked, VRMode),
+			FCanExecuteAction::CreateStatic(&FLevelEditorActionCallbacks::DefaultCanExecuteAction)
+			)
+		);
+	GizmoMenuBuilder.AddMenuEntry(
+		LOCTEXT("CycleGizmo", "Cycle Gizmo"),
+		LOCTEXT("CycleGizmoTooltip", "Cycle Gizmo"),
+		FSlateIcon(),
+		FUIAction
+		(
+			FExecuteAction::CreateStatic(&FVREditorActionCallbacks::OnGizmoCycle, VRMode),
+			FCanExecuteAction::CreateStatic(&FLevelEditorActionCallbacks::DefaultCanExecuteAction)
+			)
+		);
+
+	TSharedRef<SWidget> GizmoMenuWidget = GizmoMenuBuilder.MakeWidget(&VREditorMenuBuilderOverride);
+
+	// Next Menu section: Extras
+	CommandList = MakeShareable(new FUICommandList());
+	FMenuBuilder ExtrasMenuBuilder(false, CommandList, MenuExtender);
+
+
+	ExtrasMenuBuilder.AddMenuEntry(
+		LOCTEXT("Screenshot", "Screenshot"),
+		LOCTEXT("ScreenshotTooltip", "Screenshot"),
+		FSlateIcon(),
+		FUIAction
+		(
+			FExecuteAction::CreateStatic(&FVREditorActionCallbacks::OnScreenshotButtonClicked, VRMode),
+			FCanExecuteAction::CreateStatic(&FLevelEditorActionCallbacks::DefaultCanExecuteAction)
+			)
+		);
+	ExtrasMenuBuilder.AddMenuEntry(
+		LOCTEXT("Flashlight", "Flashlight"),
+		LOCTEXT("FlashlightTooltip", "Flashlight"),
+		FSlateIcon(),
+		FUIAction
+		(
+			FExecuteAction::CreateStatic(&FVREditorActionCallbacks::OnLightButtonClicked, VRMode),
+			FCanExecuteAction::CreateStatic(&FLevelEditorActionCallbacks::DefaultCanExecuteAction),
+			FGetActionCheckState::CreateStatic(&FVREditorActionCallbacks::GetFlashlightState, VRMode)
+			),
+		NAME_None,
+		EUserInterfaceActionType::Check
+		);
+	ExtrasMenuBuilder.AddMenuEntry(
+		LOCTEXT("PlayInEditor", "Play"),
+		LOCTEXT("PlayTooltip", "Play"),
+		FSlateIcon(),
+		FUIAction
+		(
+			FExecuteAction::CreateStatic(&FVREditorActionCallbacks::OnPlayButtonClicked, VRMode),
+			FCanExecuteAction::CreateStatic(&FLevelEditorActionCallbacks::DefaultCanExecuteAction)
+			)
+		);
+	TSharedRef<SWidget> ExtrasMenuWidget = ExtrasMenuBuilder.MakeWidget(&VREditorMenuBuilderOverride);
+	const FSlateBrush* NoBorderBrush = FStyleDefaults::GetNoBrush();
+	TSharedRef<SWidget> MenuWidget = SNew(SBox)
+		.VAlign(VAlign_Top)
+		.HAlign(HAlign_Center)
+		[
+			SNew(SDPIScaler)
+			.DPIScale(3)
+			[
+				SNew(SVerticalBox)
+				+SVerticalBox::Slot()
+				.HAlign(HAlign_Fill)
+				.VAlign(VAlign_Bottom)
+				[
+					SNew(STextBlock)
+					.Justification(ETextJustify::Center)
+					.Text(LOCTEXT("GridSnap", "Grid Snap"))
+					.TextStyle(FVREditorStyle::Get(),"VREditorStyle.Heading")
+				]
+				+ SVerticalBox::Slot()
+				.HAlign(HAlign_Fill)
+				.VAlign(VAlign_Top)
+				.AutoHeight()
+				[
+					SnapMenuWidget
+				]
+				+ SVerticalBox::Slot()
+				.HAlign(HAlign_Fill)
+				.VAlign(VAlign_Bottom)
+				[
+					SNew(STextBlock)
+					.Justification(ETextJustify::Center)
+					.Text(LOCTEXT("Interface", "Interface"))
+					.TextStyle(FVREditorStyle::Get(), "VREditorStyle.Heading")
+				]
+				+ SVerticalBox::Slot()
+				.HAlign(HAlign_Fill)
+				.VAlign(VAlign_Top)
+				.AutoHeight()
+				[
+					InterfaceMenuWidget
+				]
+				+ SVerticalBox::Slot()
+				.HAlign(HAlign_Fill)
+				.VAlign(VAlign_Bottom)
+				[
+					SNew(STextBlock)
+					.Justification(ETextJustify::Center)
+					.Text(LOCTEXT("VRTools", "Tools"))
+					.TextStyle(FVREditorStyle::Get(), "VREditorStyle.Heading")
+				]
+				+ SVerticalBox::Slot()
+				.HAlign(HAlign_Fill)
+				.VAlign(VAlign_Top)
+				.AutoHeight()
+				[
+					ExtrasMenuWidget
+				]
+				+ SVerticalBox::Slot()
+				.HAlign(HAlign_Fill)
+				.VAlign(VAlign_Bottom)
+				[
+					SNew(STextBlock)
+					.Justification(ETextJustify::Center)
+					.Text(LOCTEXT("Gizmo", "Gizmo"))
+					.TextStyle(FVREditorStyle::Get(), "VREditorStyle.Heading")
+				]
+				+ SVerticalBox::Slot()
+				.HAlign(HAlign_Fill)
+				.VAlign(VAlign_Top)
+				.AutoHeight()
+				[
+					GizmoMenuWidget
+				]
+			]
+		]
+	;
+
+
+	return MenuWidget;
+}
+
+
+
+TSharedRef<SWidget> UVREditorUISystem::BuildRadialMenuWidget()
+{
+
+	FMultiBox::FOnMakeMultiBoxBuilderOverride VREditorMenuBuilderOverride;
+	VREditorMenuBuilderOverride.BindUObject(this, &UVREditorUISystem::MakeRadialBoxMenu);
+	const TSharedRef<FExtender> MenuExtender = MakeShareable(new FExtender());
+	TSharedPtr<FUICommandList> CommandList = MakeShareable(new FUICommandList());
+	FMenuBarBuilder MenuBuilder(CommandList, MenuExtender);
+
+	FVREditorActionCallbacks::UpdateGizmoCoordinateSystemText(VRMode);
+	FVREditorActionCallbacks::UpdateGizmoModeText(VRMode);
+
+	// First menu entry is at 90 degrees 
+	MenuBuilder.AddMenuEntry(
+		LOCTEXT("Delete", "Delete"),
+		LOCTEXT("DeleteTooltip", "Delete"),
+		FSlateIcon(),
+		FUIAction
+		(
+			FExecuteAction::CreateStatic(&FLevelEditorActionCallbacks::ExecuteExecCommand, FString(TEXT("DELETE"))),
+			FCanExecuteAction::CreateStatic(&FLevelEditorActionCallbacks::Delete_CanExecute)
+		)
+	);
+	MenuBuilder.AddMenuEntry(
+		LOCTEXT("Cut", "Cut"),
+		LOCTEXT("CutTooltip", "Cut"),
+		FSlateIcon(),
+		FUIAction
+		(
+			FExecuteAction::CreateStatic(&FLevelEditorActionCallbacks::ExecuteExecCommand, FString(TEXT("EDIT CUT"))),
+			FCanExecuteAction::CreateStatic(&FLevelEditorActionCallbacks::Cut_CanExecute)
+		)
+		);
+	MenuBuilder.AddMenuEntry(
+		LOCTEXT("Copy", "Copy"),
+		LOCTEXT("CopyTooltip", "Copy"),
+		FSlateIcon(),
+		FUIAction
+		(
+			FExecuteAction::CreateStatic(&FLevelEditorActionCallbacks::ExecuteExecCommand, FString(TEXT("EDIT COPY"))),
+			FCanExecuteAction::CreateStatic(&FLevelEditorActionCallbacks::Copy_CanExecute)
+		)
+	);
+	MenuBuilder.AddMenuEntry(
+		LOCTEXT("Duplicate", "Duplicate"),
+		LOCTEXT("DuplicateTooltip", "Duplicate"),
+		FSlateIcon(),
+		FUIAction
+		(
+			FExecuteAction::CreateStatic(&FLevelEditorActionCallbacks::ExecuteExecCommand, FString(TEXT("DUPLICATE"))),
+			FCanExecuteAction::CreateStatic(&FLevelEditorActionCallbacks::Duplicate_CanExecute)
+		)
+	);
+	MenuBuilder.AddMenuEntry(
+		LOCTEXT("Paste", "Paste"),
+		LOCTEXT("PasteTooltip", "Paste"),
+		FSlateIcon(),
+		FUIAction
+		(
+			FExecuteAction::CreateStatic(&FLevelEditorActionCallbacks::ExecuteExecCommand, FString(TEXT("EDIT PASTE"))),
+			FCanExecuteAction::CreateStatic(&FLevelEditorActionCallbacks::Paste_CanExecute)
+		)
+	);
+	MenuBuilder.AddMenuEntry(
+		LOCTEXT("Undo", "Undo"),
+		LOCTEXT("UndoTooltip", "Undo"),
+		FSlateIcon(),
+		FUIAction
+		(
+			FExecuteAction::CreateStatic(&FLevelEditorActionCallbacks::ExecuteExecCommand, FString(TEXT("TRANSACTION UNDO"))),
+			FCanExecuteAction::CreateStatic(&FLevelEditorActionCallbacks::DefaultCanExecuteAction)
+			)
+		);
+	MenuBuilder.AddMenuEntry(
+		LOCTEXT("SnapActorsToGround", "Snap Actors To Ground"),
+		LOCTEXT("SnapActorsToGroundTooltip", "Snap Actors To Ground"),
+		FSlateIcon(),
+		FUIAction
+		(
+			FExecuteAction::CreateStatic(&FVREditorActionCallbacks::OnSnapActorsToGroundClicked, VRMode),
+			FCanExecuteAction::CreateStatic(&FLevelEditorActionCallbacks::Copy_CanExecute)
+		)
+	);
+	MenuBuilder.AddMenuEntry(
+		LOCTEXT("Redo", "Redo"),
+		LOCTEXT("RedoTooltip", "Redo"),
+		FSlateIcon(),
+		FUIAction
+		(
+			FExecuteAction::CreateStatic(&FLevelEditorActionCallbacks::ExecuteExecCommand, FString(TEXT("TRANSACTION REDO"))),
+			FCanExecuteAction::CreateStatic(&FLevelEditorActionCallbacks::DefaultCanExecuteAction)
+			)
+		);
+
+
+	// Create the menu widget
+	TSharedRef<SWidget> MenuWidget = MenuBuilder.MakeWidget(&VREditorMenuBuilderOverride);
+	return MenuWidget;
+}
+
+const TSharedRef<SWidget>& UVREditorUISystem::FindWidgetOfType(const TSharedRef<SWidget>& Content, FName WidgetType)
+{
+	if (Content->GetType() == WidgetType)
+	{
+		return Content;
+	}
+
+	FChildren* Children = Content->GetChildren();
+	int32 NumChildren = Children->Num();
+
+	for (int32 Index = 0; Index < NumChildren; ++Index)
+	{
+		const TSharedRef<SWidget>& Found = FindWidgetOfType((Children->GetChildAt(Index)), WidgetType);
+		if (Found != SNullWidget::NullWidget)
+		{
+			return Found;
+		}
+	}
+	return SNullWidget::NullWidget;
+}
+
+void UVREditorUISystem::OnHoverBeginEffect(TSharedRef<SButton> Button)
+{
+	FVRButton* ButtonToAnimate = VRButtons.FindByPredicate([&Button](const FVRButton& ButtonData)
+	{
+		return Button == ButtonData.Button;
+	});
+
+	check(ButtonToAnimate);
+
+	//Set the newly hovered button's animation state to Forward
+	ButtonToAnimate->AnimationDirection = EVREditorAnimationState::Forward;
+
+	// Put the button in the higher z-order overlay slot
+	const TSharedRef<SWidget>& PlaceholderWidget = ButtonToAnimate->HoverSlot->DetachWidget().ToSharedRef();
+	const TSharedRef<SWidget>& WidgetToTransfer = ButtonToAnimate->BaseSlot->DetachWidget().ToSharedRef();
+	ButtonToAnimate->HoverSlot->AttachWidget(WidgetToTransfer);
+	ButtonToAnimate->BaseSlot->AttachWidget(PlaceholderWidget);
+}
+
+void UVREditorUISystem::OnHoverEndEffect(TSharedRef<SButton> Button)
+{
+	FVRButton* ButtonToAnimate = VRButtons.FindByPredicate([&Button](const FVRButton& ButtonData)
+	{
+		return Button == ButtonData.Button;
+	});
+
+	check(ButtonToAnimate);
+
+	// Set the unhovered button's animation state to Backward
+	ButtonToAnimate->AnimationDirection = EVREditorAnimationState::Backward;
+	// Put the button back in the lower z-order slot
+	TSharedRef<SWidget> PlaceholderWidget = (ButtonToAnimate->BaseSlot->DetachWidget()).ToSharedRef();
+	TSharedRef<SWidget> WidgetToTransfer = (ButtonToAnimate->HoverSlot->DetachWidget()).ToSharedRef();
+	ButtonToAnimate->BaseSlot->AttachWidget(WidgetToTransfer);
+	ButtonToAnimate->HoverSlot->AttachWidget(PlaceholderWidget);
+}
+
+const TSharedPtr<SWidget>& UVREditorUISystem::GetRadialWidget()
+{
+	return RadialWidget;
+}
+
+#undef LOCTEXT_NAMESPACE
\ No newline at end of file
diff --git a/Engine/Source/Editor/VREditor/UI/VREditorUISystem.h b/Engine/Source/Editor/VREditor/UI/VREditorUISystem.h
--- a/Engine/Source/Editor/VREditor/UI/VREditorUISystem.h
+++ b/Engine/Source/Editor/VREditor/UI/VREditorUISystem.h
@@ -8,6 +8,69 @@ class UVREditorInteractor;
 class UViewportInteractor;
 class SColorPicker;
 
+
+/** Stores the animation playback state of a VR UI element */
+enum class EVREditorAnimationState : uint8
+{
+	None,
+	Forward,
+	Backward
+};
+
+/** Structure to keep track of all relevant interaction and animation elements of a VR Button */
+struct FVRButton
+{
+	/** Pointer to button */
+	TSharedPtr<SButton> Button;
+
+	/** Pointer to slot in overlay with lower z-order */
+	SOverlay::FOverlaySlot* BaseSlot;
+
+	/** Pointer to slot in overlay with higher z-order */
+	SOverlay::FOverlaySlot* HoverSlot;
+
+	/** Pointer to button border */
+	TSharedPtr<SBorder> ButtonBorder;
+
+	/** Animation playback state of the button */
+	EVREditorAnimationState AnimationDirection;
+
+	/** Current scale of the button element */
+	float CurrentScale;
+
+	/** Minimum scale of the button element */
+	float MinScale;
+
+	/** Maximum scale of the button element */
+	float MaxScale;
+
+	/** Rate at which the button changes scale. Currently the same for scaling up and scaling down. */
+	float ScaleRate;
+
+	FVRButton()
+		: AnimationDirection(EVREditorAnimationState::None),
+		CurrentScale(1.0f),
+		MinScale(1.0f),
+		MaxScale(1.10f),
+		ScaleRate(2.0f)
+		{}
+
+	FVRButton(TSharedPtr<SButton> InButton, SOverlay::FOverlaySlot* InBaseSlot, SOverlay::FOverlaySlot* InHoverSlot, TSharedPtr<SBorder> InButtonBorder,
+		EVREditorAnimationState InAnimationDirection = EVREditorAnimationState::None, float InCurrentScale = 1.0f, float InMinScale = 1.0f, float InMaxScale = 1.25f, float InScaleRate = 2.0f)
+		: Button(InButton),
+		BaseSlot(InBaseSlot),
+		HoverSlot(InHoverSlot),
+		ButtonBorder(InButtonBorder),
+		AnimationDirection(InAnimationDirection),
+		CurrentScale(InCurrentScale),
+		MinScale(InMinScale),
+		MaxScale(InMaxScale),
+		ScaleRate(InScaleRate)
+		{}
+
+};
+
+
 /**
  * VR Editor user interface manager
  */
@@ -84,9 +147,6 @@ public:
 	/** Tries to spawn the radial menu (if the specified hand isn't doing anything else) */
 	void TryToSpawnRadialMenu( UVREditorInteractor* Interactor );
 
-	/** Updates the radial menu with the current hand trackpad location */
-	void UpdateRadialMenu( UVREditorInteractor* Interactor );
-
 	/** Hides the radial menu if the specified hand is showing it */
 	void HideRadialMenu( UVREditorInteractor* Interactor );
 
@@ -122,6 +182,17 @@ public:
 	/** Toggles the visibility of the panel, if the panel is in room space it will be hidden and docked to nothing */
 	void TogglePanelVisibility( const EEditorUIPanel EditorUIPanel );
 
+	/** Returns the radial widget so other classes, like the interactors, can access its functionality */
+	const TSharedPtr<SWidget>& GetRadialWidget();
+
+	/** 
+	 * Finds a widget with a given name inside the Content argument 
+	 * @param Content The widget to begin searching in
+	 * @param The FName of the widget type to find. 
+	 */
+	static const TSharedRef<SWidget>& FindWidgetOfType(const TSharedRef<SWidget>& Content, FName WidgetType);
+
+
 protected:
 
 	/** Called when the user presses a button on their motion controller device */
@@ -159,6 +230,26 @@ protected:
 	/** Hides the VR-specific color picker. Gets bound to SColorPicker's destruction override delegate */
 	void DestroyVRColorPicker();
 
+	/**  Makes a uniform grid widget from the menu information contained in a MultiBox and MultiBoxWidget */
+	void MakeUniformGridMenu( const TSharedRef<FMultiBox>& MultiBox, const TSharedRef<SMultiBoxWidget>& MultiBoxWidget, int32 Columns);
+	/**  Makes a radial box widget from the menu information contained in a MultiBox and MultiBoxWidget */
+	void MakeRadialBoxMenu(const TSharedRef<FMultiBox>& MultiBox, const TSharedRef<SMultiBoxWidget>& MultiBoxWidget);
+
+	/** Adds a hoverable button of a given type to an overlay, using menu data from a BlockWidget */
+	TSharedRef<SWidget> AddHoverableButton(TSharedRef<SWidget>& BlockWidget, FName ButtonType, TSharedRef<SOverlay>& TestOverlay);
+	/** Sets the text wrap size of the text block element nested in a BlockWidget */
+	TSharedRef<SWidget> SetButtonTextWrap(TSharedRef<SWidget>& BlockWidget, float WrapSize);
+	
+	/** Builds the quick menu Slate widget */
+	TSharedRef<SWidget> BuildQuickMenuWidget();
+	/** Builds the radial menu Slate widget */
+	TSharedRef<SWidget> BuildRadialMenuWidget();
+
+	/** Called when a laser or simulated mouse hover enters a button */
+	void OnHoverBeginEffect(TSharedRef<SButton> Button);
+	/** Called when a laser or simulated mouse hover leaves a button */
+	void OnHoverEndEffect(TSharedRef<SButton> Button);
+
 
 protected:
 
@@ -185,14 +276,11 @@ protected:
 	/** The time since the radial menu was updated */
 	float RadialMenuHideDelayTime;
 
-	/** Quick menu widget class */
-	UClass* QuickMenuWidgetClass;
-
-	/** Quick radial menu widget class */
-	UClass* QuickRadialWidgetClass;
-
 	/** Tutorial menu widget class */
 	UClass* TutorialWidgetClass;
+
+	/** Pointer to the radial menu widget */
+	TSharedPtr<SWidget> RadialWidget;
 	
 	//
 	// Dragging UI
@@ -252,6 +340,9 @@ protected:
 	/** Default transforms */
 	TArray<FTransform> DefaultWindowTransforms;
 
+	/** All buttons created for the radial and quick menus */
+	TArray<FVRButton> VRButtons;
+
 	/** If this is the first time using TogglePanelsVisibility */
 	bool bSetDefaultLayout;
 };
diff --git a/Engine/Source/Editor/VREditor/VREditorActions.cpp b/Engine/Source/Editor/VREditor/VREditorActions.cpp
new file mode 100644
--- /dev/null
+++ b/Engine/Source/Editor/VREditor/VREditorActions.cpp
@@ -0,0 +1,259 @@
+// Copyright 1998-2016 Epic Games, Inc. All Rights Reserved.
+
+
+#include "VREditorModule.h"
+#include "VREditorUISystem.h"
+#include "VREditorMode.h"
+#include "ViewportWorldInteraction.h"
+#include "VREditorInteractor.h"
+#include "EditorWorldManager.h"
+#include "VREditorFloatingUI.h"
+#include "VREditorTransformGizmo.h"
+#include "SLevelViewport.h"
+#include "ImageUtils.h"
+#include "VREditorActions.h"
+
+#define LOCTEXT_NAMESPACE "VREditorActions"
+
+ECheckBoxState FVREditorActionCallbacks::GetTranslationSnapState()
+{
+	return GetDefault<ULevelEditorViewportSettings>()->GridEnabled ? ECheckBoxState::Checked : ECheckBoxState::Unchecked;
+}
+
+void FVREditorActionCallbacks::OnTranslationSnapSizeButtonClicked()
+{
+	const TArray<float>& GridSizes = GEditor->GetCurrentPositionGridArray();
+	const int32 CurrentGridSize = GetDefault<ULevelEditorViewportSettings>()->CurrentPosGridSize;
+
+	int32 NewGridSize = CurrentGridSize + 1;
+	if (NewGridSize >= GridSizes.Num())
+	{
+		NewGridSize = 0;
+	}
+
+	GEditor->SetGridSize(NewGridSize);
+
+	GetTranslationSnapSizeText();
+
+}
+
+FText FVREditorActionCallbacks::GetTranslationSnapSizeText() 
+{
+	return FText::AsNumber(GEditor->GetGridSize());
+
+}
+
+ECheckBoxState FVREditorActionCallbacks::GetRotationSnapState()
+{
+	return GetDefault<ULevelEditorViewportSettings>()->RotGridEnabled ? ECheckBoxState::Checked : ECheckBoxState::Unchecked;
+}
+
+void FVREditorActionCallbacks::OnRotationSnapSizeButtonClicked()
+{
+	const TArray<float>& GridSizes = GEditor->GetCurrentRotationGridArray();
+	const int32 CurrentGridSize = GetDefault<ULevelEditorViewportSettings>()->CurrentRotGridSize;
+	const ERotationGridMode CurrentGridMode = GetDefault<ULevelEditorViewportSettings>()->CurrentRotGridMode;
+
+	int32 NewGridSize = CurrentGridSize + 1;
+	if (NewGridSize >= GridSizes.Num())
+	{
+		NewGridSize = 0;
+	}
+
+	GEditor->SetRotGridSize(NewGridSize, CurrentGridMode);
+
+}
+
+FText FVREditorActionCallbacks::GetRotationSnapSizeText() 
+{
+	return FText::AsNumber(GEditor->GetRotGridSize().Yaw);
+}
+
+ECheckBoxState FVREditorActionCallbacks::GetScaleSnapState()
+{
+	return GetDefault<ULevelEditorViewportSettings>()->SnapScaleEnabled ? ECheckBoxState::Checked : ECheckBoxState::Unchecked;
+}
+
+void FVREditorActionCallbacks::OnScaleSnapSizeButtonClicked()
+{
+	const ULevelEditorViewportSettings* ViewportSettings = GetMutableDefault<ULevelEditorViewportSettings>();
+	const int32 NumGridSizes = ViewportSettings->ScalingGridSizes.Num();
+
+	const int32 CurrentGridSize = GetDefault<ULevelEditorViewportSettings>()->CurrentScalingGridSize;
+
+	int32 NewGridSize = CurrentGridSize + 1;
+	if (NewGridSize >= NumGridSizes)
+	{
+		NewGridSize = 0;
+	}
+
+	GEditor->SetScaleGridSize(NewGridSize);
+
+	
+}
+
+FText FVREditorActionCallbacks::GetScaleSnapSizeText() 
+{
+	return FText::AsNumber(GEditor->GetScaleGridSize());
+}
+
+void FVREditorActionCallbacks::OnGizmoCoordinateSystemButtonClicked(class UVREditorMode* InVRMode)
+{
+	InVRMode->GetWorldInteraction().CycleTransformGizmoCoordinateSpace();
+	UpdateGizmoModeText(InVRMode);
+	UpdateGizmoCoordinateSystemText(InVRMode);
+
+}
+
+FText FVREditorActionCallbacks::GetGizmoCoordinateSystemText()
+{
+	return EVREditorActions::GizmoCoordinateSystemText;
+}
+
+void FVREditorActionCallbacks::UpdateGizmoCoordinateSystemText(class UVREditorMode* InVRMode) 
+{
+	const ECoordSystem CurrentCoordSystem = InVRMode->GetWorldInteraction().GetTransformGizmoCoordinateSpace(); //@todo VREditor
+	EVREditorActions::GizmoCoordinateSystemText = (CurrentCoordSystem == COORD_World ? LOCTEXT("WorldCoordinateSystem", "World Space") : LOCTEXT("LocalCoordinateSystem", "Local Space"));
+}
+
+void FVREditorActionCallbacks::OnGizmoModeButtonClicked(class UVREditorMode* InVRMode)
+{
+	const EGizmoHandleTypes CurrentGizmoMode = InVRMode->GetWorldInteraction().GetCurrentGizmoType();
+	InVRMode->CycleTransformGizmoHandleType();
+	UpdateGizmoModeText(InVRMode);
+	UpdateGizmoCoordinateSystemText(InVRMode);
+
+}
+
+FText FVREditorActionCallbacks::GetGizmoModeText()
+{
+	return EVREditorActions::GizmoModeText;
+}
+
+void FVREditorActionCallbacks::UpdateGizmoModeText(class UVREditorMode* InVRMode) 
+{
+	const EGizmoHandleTypes CurrentGizmoType = InVRMode->GetWorldInteraction().GetCurrentGizmoType();
+	FText GizmoTypeText;
+	switch (CurrentGizmoType)
+	{
+	case EGizmoHandleTypes::All:
+		GizmoTypeText = LOCTEXT("AllGizmoType", "Universal");
+		break;
+
+	case EGizmoHandleTypes::Translate:
+		GizmoTypeText = LOCTEXT("TranslateGizmoType", "Translate");
+		break;
+
+	case EGizmoHandleTypes::Rotate:
+		GizmoTypeText = LOCTEXT("RotationGizmoType", "Rotation");
+		break;
+
+	case EGizmoHandleTypes::Scale:
+		GizmoTypeText = LOCTEXT("ScaleGizmoType", "Scale");
+		break;
+
+	default:
+		check(0);	// Unrecognized type
+		break;
+	}
+
+	EVREditorActions::GizmoModeText = GizmoTypeText;
+}
+
+void FVREditorActionCallbacks::OnGizmoCycle(class UVREditorMode* InVRMode)
+{
+	InVRMode->CycleTransformGizmoHandleType();
+	UpdateGizmoModeText(InVRMode);
+	UpdateGizmoCoordinateSystemText(InVRMode);
+}
+
+void FVREditorActionCallbacks::OnUIToggleButtonClicked(class UVREditorMode* InVRMode, const UVREditorUISystem::EEditorUIPanel PanelToToggle)
+{
+	InVRMode->GetUISystem().TogglePanelVisibility(PanelToToggle);
+}
+
+ECheckBoxState FVREditorActionCallbacks::GetUIToggledState(class UVREditorMode* InVRMode, const UVREditorUISystem::EEditorUIPanel PanelToCheck)
+{
+	return InVRMode->GetUISystem().IsShowingEditorUIPanel(PanelToCheck) ? ECheckBoxState::Checked : ECheckBoxState::Unchecked;
+}
+
+
+void FVREditorActionCallbacks::OnLightButtonClicked(class UVREditorMode* InVRMode)
+{
+	// Always spawn the flashlight on the hand clicking on the UI
+	UVREditorInteractor* VREditorInteractor = InVRMode->GetHandInteractor(EControllerHand::Left);
+	if (!VREditorInteractor->IsHoveringOverUI())
+	{
+		VREditorInteractor = InVRMode->GetHandInteractor(EControllerHand::Right);
+
+	}
+	InVRMode->ToggleFlashlight(VREditorInteractor);
+
+}
+
+ECheckBoxState FVREditorActionCallbacks::GetFlashlightState(class UVREditorMode* InVRMode)
+{
+	return InVRMode->IsFlashlightOn() ? ECheckBoxState::Checked : ECheckBoxState::Unchecked;
+}
+
+
+void FVREditorActionCallbacks::OnScreenshotButtonClicked(class UVREditorMode* InVRMode)
+{
+	// @todo vreditor: update after direct buffer grab changes
+
+	FString GeneratedFilename = TEXT("");
+	FString Filename = TEXT("");
+	FScreenshotRequest::CreateViewportScreenShotFilename(GeneratedFilename);
+	const bool bRemovePath = false;
+	GeneratedFilename = FPaths::GetBaseFilename(GeneratedFilename, bRemovePath);
+	FFileHelper::GenerateNextBitmapFilename(GeneratedFilename, TEXT("png"), Filename);
+
+	TSharedRef<SWidget> WindowRef = InVRMode->GetLevelViewportPossessedForVR().AsWidget();
+
+
+	TArray<FColor> OutImageData;
+	FIntVector OutImageSize;
+
+	if (FSlateApplication::Get().TakeScreenshot(WindowRef, OutImageData, OutImageSize))
+	{
+		int32 Width = OutImageSize.X;
+		int32 Height = OutImageSize.Y;
+
+		TArray<FColor> ScaledBitmap;
+
+
+		ScaledBitmap = OutImageData;
+		//Clear the alpha channel before saving
+		for (int32 Index = 0; Index < Width*Height; Index++)
+		{
+			ScaledBitmap[Index].A = 255;
+		}
+
+
+		TArray<uint8> CompressedBitmap;
+		FImageUtils::CompressImageArray(Width, Height, ScaledBitmap, CompressedBitmap);
+
+
+		//Save locally
+		const bool bTree = true;
+		const FString FinalFileName = Filename;
+		IFileManager::Get().MakeDirectory(*FPaths::GetPath(FinalFileName), bTree);
+		FFileHelper::SaveArrayToFile(CompressedBitmap, *FinalFileName);
+
+	}
+
+
+}
+
+void FVREditorActionCallbacks::OnPlayButtonClicked(class UVREditorMode* InVRMode)
+{
+	InVRMode->StartExitingVRMode(EVREditorExitType::PIE_VR);
+}
+
+void FVREditorActionCallbacks::OnSnapActorsToGroundClicked(class UVREditorMode* InVRMode)
+{
+	InVRMode->SnapSelectedActorsToGround();
+}
+
+
+#undef LOCTEXT_NAMESPACE
diff --git a/Engine/Source/Editor/VREditor/VREditorActions.h b/Engine/Source/Editor/VREditor/VREditorActions.h
new file mode 100644
--- /dev/null
+++ b/Engine/Source/Editor/VREditor/VREditorActions.h
@@ -0,0 +1,133 @@
+// Copyright 1998-2016 Epic Games, Inc. All Rights Reserved.
+#pragma once
+
+#include "VREditorUISystem.h"
+
+namespace EVREditorActions
+{
+	//Gizmo coordinate button text starts blank before it is initialized
+	static FText GizmoCoordinateSystemText = NSLOCTEXT("GizmoCoordinateDefault", "", "");
+
+	//Gizmo mode button text starts blank before it is initialized
+	static FText GizmoModeText = NSLOCTEXT("GizmoModeDefault", "", "");;
+}
+
+/**
+* Implementation of various VR editor action callback functions
+*/
+class FVREditorActionCallbacks
+{
+public:
+
+	/** Returns the checked state of the translation snap enable/disable button */
+	static ECheckBoxState GetTranslationSnapState();
+
+	/** Rotates through the available translation snap sizes */
+	static void OnTranslationSnapSizeButtonClicked();
+
+	/** Returns the translation snap size as text to use as the button display text */
+	static FText GetTranslationSnapSizeText();
+
+	/** Returns the checked state of the rotation snap enable/disable button */
+	static ECheckBoxState GetRotationSnapState();
+
+	/** Rotates through the available rotation snap sizes */
+	static void OnRotationSnapSizeButtonClicked();
+
+	/** Returns the rotation snap size as text to use as the button display text */
+	static FText GetRotationSnapSizeText();
+
+	/** Returns the checked state of the scale snap enable/disable button */
+	static ECheckBoxState GetScaleSnapState();
+
+	/** Rotates through the available scale snap sizes */
+	static void OnScaleSnapSizeButtonClicked();
+
+	/** Returns the scale snap size as text to use as the button display text */
+	static FText GetScaleSnapSizeText();
+
+	/**
+	 * Toggles the gizmo coordinate system between local and world space
+	 * @param InVRMode Currently active VRMode
+	 */
+	static void OnGizmoCoordinateSystemButtonClicked(class UVREditorMode* InVRMode);
+
+	/** Returns the gizmo coordinate system as text to use as the button display text */
+	static FText GetGizmoCoordinateSystemText();
+
+	/**
+	 * Updates the gizmo coordinate system text if the coordinate system or gizmo type is changed
+	 * @param InVRMode Currently active VRMode
+	 */
+	static void UpdateGizmoCoordinateSystemText(class UVREditorMode* InVRMode);
+
+	/**
+	 * Rotates the gizmo type through universal, translate, rotate, and scale
+	 * @param InVRMode Currently active VRMode
+	 */
+	static void OnGizmoModeButtonClicked(class UVREditorMode* InVRMode);
+
+	/** Returns the gizmo type as text to use as the button display text */
+	static FText GetGizmoModeText();
+
+	/**
+	 * Updates the gizmo coordinate type text if the coordinate system or gizmo type is changed
+	 * @param InVRMode Currently active VRMode
+	 */
+	static void UpdateGizmoModeText(class UVREditorMode* InVRMode);
+
+	/**
+	* Cycles the gizmo through all coordinate spaces and Actor manipulation types
+	* @param InVRMode Currently active VRMode
+	*/
+	static void OnGizmoCycle(class UVREditorMode* InVRMode);
+
+	/**
+	 * Toggles a VR UI panel's state between visible and invisible
+	 * @param InVRMode Currently active VRMode
+	 * @param PanelToToggle UI panel to change the state of
+	 */
+	static void OnUIToggleButtonClicked(class UVREditorMode* InVRMode, const UVREditorUISystem::EEditorUIPanel PanelToToggle);
+
+	/**
+	 * Returns a VR UI panel's visibility - used for check boxes on the menu button
+	 * @param InVRMode Currently active VRMode
+	 * @param PanelToToggle UI panel to read the state of
+	 */
+	static ECheckBoxState GetUIToggledState(class UVREditorMode* InVRMode, const UVREditorUISystem::EEditorUIPanel PanelToCheck);
+
+	/**
+	 * Toggles a flashlight on and off on the interactor that clicked on the UI button
+	 * @param InVRMode Currently active VRMode
+	 */
+	static void OnLightButtonClicked(class UVREditorMode* InVRMode);
+
+	/**
+	 * Returns whether or not the flashlight is enabled - used for check box on the flashlight button
+	 * @param InVRMode Currently active VRMode
+	 */
+	static ECheckBoxState GetFlashlightState(class UVREditorMode* InVRMode);
+
+	/**
+	 * Takes a screenshot of the mirror viewport
+	 * @param InVRMode Currently active VRMode
+	 */
+	static void OnScreenshotButtonClicked(class UVREditorMode* InVRMode);
+
+	/**
+	 * Enters Play In Editor mode for testing of gameplay
+	 * @param InVRMode Currently active VRMode
+	 */
+	static void OnPlayButtonClicked(class UVREditorMode* InVRMode);
+
+	/**
+	 * Snaps currently selected Actors to the ground
+	 * @param InVRMode Currently active VRMode
+	 */
+	static void OnSnapActorsToGroundClicked(class UVREditorMode* InVRMode);
+
+}
+;
+
+
+
diff --git a/Engine/Source/Editor/VREditor/VREditorModule.cpp b/Engine/Source/Editor/VREditor/VREditorModule.cpp
--- a/Engine/Source/Editor/VREditor/VREditorModule.cpp
+++ b/Engine/Source/Editor/VREditor/VREditorModule.cpp
@@ -3,6 +3,7 @@
 #include "VREditorModule.h"
 #include "IVREditorModule.h"
 #include "VREditorModeManager.h"
+#include "VREditorStyle.h"
 #include "Kismet/HeadMountedDisplayFunctionLibrary.h"	// For EHMDWornState::Type
 
 class FVREditorModule : public IVREditorModule, public FTickableEditorObject
@@ -54,10 +55,20 @@ namespace VREd
 
 void FVREditorModule::StartupModule()
 {
+	if (GIsEditor)
+	{
+		FVREditorStyle::Initialize();
+	}
+
 }
 
 void FVREditorModule::ShutdownModule()
 {
+
+	if (GIsEditor)
+	{
+		FVREditorStyle::Shutdown();
+	}
 }
 
 void FVREditorModule::PostLoadCallback()
diff --git a/Engine/Source/Editor/VREditor/VREditorMotionControllerInteractor.cpp b/Engine/Source/Editor/VREditor/VREditorMotionControllerInteractor.cpp
--- a/Engine/Source/Editor/VREditor/VREditorMotionControllerInteractor.cpp
+++ b/Engine/Source/Editor/VREditor/VREditorMotionControllerInteractor.cpp
@@ -14,6 +14,9 @@
 #include "MotionControllerComponent.h"
 #include "Features/IModularFeatures.h"
 
+#include "SVRRadialPanel.h"
+#include "VREditorWidgetComponent.h"
+
 namespace VREd
 {
 	//Laser
@@ -43,7 +46,6 @@ namespace VREd
 	static FAutoConsoleVariable InvertTrackpadVertical( TEXT( "VREd.InvertTrackpadVertical" ), 1, TEXT( "Toggles inverting the touch pad vertical axis" ) );
 	static FAutoConsoleVariable TriggerLightlyPressedLockTime( TEXT( "VREd.TriggerLightlyPressedLockTime" ), 0.15f, TEXT( "If the trigger remains lightly pressed for longer than this, we'll continue to treat it as a light press in some cases" ) );
 	static FAutoConsoleVariable MinVelocityForInertia( TEXT( "VREd.MinVelocityForMotionControllerInertia" ), 1.0f, TEXT( "Minimum velocity (in cm/frame in unscaled room space) before inertia will kick in when releasing objects (or the world)" ) );
-	static FAutoConsoleVariable MinJoystickOffsetBeforeRadialMenu( TEXT( "VREd.MinJoystickOffsetBeforeRadialMenu" ), 0.15f, TEXT( "Toggles inverting the touch pad vertical axis" ) );
 
 	static FAutoConsoleVariable HelpLabelFadeDuration( TEXT( "VREd.HelpLabelFadeDuration" ), 0.4f, TEXT( "Duration to fade controller help labels in and out" ) );
 	static FAutoConsoleVariable HelpLabelFadeDistance( TEXT( "VREd.HelpLabelFadeDistance" ), 30.0f, TEXT( "Distance at which controller help labels should appear (in cm)" ) );
@@ -121,8 +123,6 @@ void UVREditorMotionControllerInteractor::Init( class UVREditorMode* InVRMode )
 		AddKeyAction( EKeys::MotionController_Left_Grip1, FViewportActionKeyInput( ViewportWorldActionTypes::WorldMovement ) );
 		AddKeyAction( UVREditorMotionControllerInteractor::MotionController_Left_FullyPressedTriggerAxis, FViewportActionKeyInput( ViewportWorldActionTypes::SelectAndMove ) );
 		AddKeyAction( UVREditorMotionControllerInteractor::MotionController_Left_LightlyPressedTriggerAxis, FViewportActionKeyInput( ViewportWorldActionTypes::SelectAndMove_LightlyPressed ) );
-		AddKeyAction( EKeys::MotionController_Left_Thumbstick, FViewportActionKeyInput( VRActionTypes::ConfirmRadialSelection ) );
-
 		AddKeyAction( SteamVRControllerKeyNames::Touch0, FViewportActionKeyInput( VRActionTypes::Touch ) );
 		AddKeyAction( EKeys::MotionController_Left_TriggerAxis, FViewportActionKeyInput( UVREditorMotionControllerInteractor::TriggerAxis ) );
 		AddKeyAction( EKeys::MotionController_Left_Thumbstick_X, FViewportActionKeyInput( UVREditorMotionControllerInteractor::TrackpadPositionX ) );
@@ -324,7 +324,7 @@ void UVREditorMotionControllerInteractor::Tick( const float DeltaTime )
 		}
 	}
 
-	UpdateRadialMenuInput( DeltaTime );
+	UpdateRadialMenuInput(DeltaTime);
 
 	{
 		const float WorldScaleFactor = WorldInteraction->GetWorldScaleFactor();
@@ -1279,24 +1279,11 @@ void UVREditorMotionControllerInteractor::UpdateRadialMenuInput( const float Del
 		!UISystem.IsInteractorDraggingDockUI( this ) &&
 		!UISystem.IsShowingRadialMenu( Cast<UVREditorInteractor>( OtherInteractor ) ) )
 	{
-		const EHMDDeviceType::Type HMDDevice = GetHMDDeviceType();
-
-		// Spawn the radial menu if we are using the touchpad for steamvr or the analog stick for the oculus
-		if ( ( HMDDevice == EHMDDeviceType::DT_SteamVR && bIsTouchingTrackpad ) ||
-			( HMDDevice == EHMDDeviceType::DT_OculusRift && TrackpadPosition.Size() > VREd::MinJoystickOffsetBeforeRadialMenu->GetFloat() ) )
-		{
-			UISystem.TryToSpawnRadialMenu( this );
-		}
-		else
-		{
-			// Hide it if we are not using the touchpad or analog stick
-			UISystem.HideRadialMenu( this );
-		}
-
 		// Update the radial menu if we are already showing the radial menu
 		if ( UISystem.IsShowingRadialMenu( this ) )
 		{
-			UISystem.UpdateRadialMenu( this );
+			const TSharedPtr<SRadialBox>& RadialBox = StaticCastSharedPtr<SRadialBox>(UISystem.GetRadialWidget());
+			RadialBox->HighlightSlot(TrackpadPosition);
 		}
 	}
 }
diff --git a/Engine/Source/Editor/VREditor/VREditorStyle.cpp b/Engine/Source/Editor/VREditor/VREditorStyle.cpp
new file mode 100644
--- /dev/null
+++ b/Engine/Source/Editor/VREditor/VREditorStyle.cpp
@@ -0,0 +1,92 @@
+// Copyright 1998-2016 Epic Games, Inc. All Rights Reserved.
+
+
+#include "VREditorModule.h"
+#include "VREditorStyle.h"
+#include "SlateTypes.h"
+#include "SlateGameResources.h"
+
+TSharedPtr< FSlateStyleSet > FVREditorStyle::VREditorStyleInstance = NULL;
+
+void FVREditorStyle::Initialize()
+{
+	if ( !VREditorStyleInstance.IsValid() )
+	{
+		VREditorStyleInstance = Create();
+		FSlateStyleRegistry::RegisterSlateStyle( *VREditorStyleInstance);
+	}
+}
+
+void FVREditorStyle::Shutdown()
+{
+	FSlateStyleRegistry::UnRegisterSlateStyle( *VREditorStyleInstance);
+	ensure(VREditorStyleInstance.IsUnique() );
+	VREditorStyleInstance.Reset();
+}
+
+FName FVREditorStyle::GetStyleSetName()
+{
+	static FName StyleSetName(TEXT("VREditorStyle"));
+	return StyleSetName;
+}
+
+#define IMAGE_BRUSH( RelativePath, ... ) FSlateImageBrush( Style->RootToContentDir( RelativePath, TEXT(".png") ), __VA_ARGS__ )
+#define BOX_BRUSH( RelativePath, ... ) FSlateBoxBrush( Style->RootToContentDir( RelativePath, TEXT(".png") ), __VA_ARGS__ )
+#define BORDER_BRUSH( RelativePath, ... ) FSlateBorderBrush( Style->RootToContentDir( RelativePath, TEXT(".png") ), __VA_ARGS__ )
+#define TTF_FONT( RelativePath, ... ) FSlateFontInfo( Style->RootToContentDir( RelativePath, TEXT(".ttf") ), __VA_ARGS__ )
+#define OTF_FONT( RelativePath, ... ) FSlateFontInfo( Style->RootToContentDir( RelativePath, TEXT(".otf") ), __VA_ARGS__ )
+#define TTF_CORE_FONT(RelativePath, ...) FSlateFontInfo(Style->RootToCoreContentDir(RelativePath, TEXT(".ttf") ), __VA_ARGS__)
+
+const FVector2D Icon16x16(16.0f, 16.0f);
+const FVector2D Icon20x20(20.0f, 20.0f);
+const FVector2D Icon40x40(40.0f, 40.0f);
+
+TSharedRef< FSlateStyleSet > FVREditorStyle::Create()
+{
+	TSharedRef< FSlateStyleSet > Style = MakeShareable(new FSlateStyleSet("VREditorStyle"));
+
+	const FButtonStyle OverrideButton = FEditorStyle::GetWidgetStyle<FButtonStyle>("VRButtonStyle");
+
+	// Use the default menu button style, but set the background to dark grey and use the sound set in the editor.
+	const FButtonStyle NormalButton = FEditorStyle::GetWidgetStyle<FButtonStyle>("Menu.Button");
+	Style->Set("VREditorStyle.Button", FButtonStyle(NormalButton)
+		.SetNormal(BOX_BRUSH("Common/RoundedSelection_16x", 4.0f / 16.0f, FLinearColor(0.1f,0.1f,0.1f)))
+		.SetPressedSound(OverrideButton.PressedSlateSound));
+
+	const FTextBlockStyle NormalText = FEditorStyle::GetWidgetStyle<FTextBlockStyle>("NormalText");
+	Style->Set("VREditorStyle.Label", FTextBlockStyle(NormalText)
+		.SetFont(TTF_CORE_FONT("Fonts/Roboto-Regular", 9)));
+
+	// Headings will have a font outline
+	FFontOutlineSettings HeadingOutline;
+	HeadingOutline.OutlineColor = FLinearColor(0.2f, 0.2f, 0.2f, 0.5f);
+	HeadingOutline.OutlineSize = 2.0f;
+	FSlateFontInfo HeadlineFont = TTF_CORE_FONT("Fonts/Roboto-Regular", 18);
+	HeadlineFont.OutlineSettings = HeadingOutline;
+
+	Style->Set("VREditorStyle.Heading", FTextBlockStyle(NormalText)
+			.SetFont(HeadlineFont)
+			.SetColorAndOpacity(FLinearColor::White));
+	
+	const FCheckBoxStyle CheckboxStyle = FEditorStyle::GetWidgetStyle<FCheckBoxStyle>("Menu.Check");
+	Style->Set("VREditorStyle.Check", FCheckBoxStyle(CheckboxStyle));
+
+	return Style;
+}
+
+#undef IMAGE_BRUSH
+#undef BOX_BRUSH
+#undef BORDER_BRUSH
+#undef TTF_FONT
+#undef OTF_FONT
+#undef TTF_CORE_FONT
+
+void FVREditorStyle::ReloadTextures()
+{
+	FSlateApplication::Get().GetRenderer()->ReloadTextureResources();
+}
+
+const ISlateStyle& FVREditorStyle::Get()
+{
+	return *VREditorStyleInstance;
+}
diff --git a/Engine/Source/Editor/VREditor/VREditorStyle.h b/Engine/Source/Editor/VREditor/VREditorStyle.h
new file mode 100644
--- /dev/null
+++ b/Engine/Source/Editor/VREditor/VREditorStyle.h
@@ -0,0 +1,35 @@
+// Copyright 1998-2016 Epic Games, Inc. All Rights Reserved.
+
+#pragma once
+
+#include "SlateBasics.h"
+
+/**  */
+class FVREditorStyle
+{
+public:
+
+	static void Initialize();
+
+	static void Shutdown();
+
+	/** reloads textures used by slate renderer */
+	static void ReloadTextures();
+
+	/** @return The Slate style set for the UMG Style */
+	static const ISlateStyle& Get();
+
+	static FName GetStyleSetName();
+
+	
+
+private:
+
+	static TSharedRef< class FSlateStyleSet > Create();
+
+private:
+
+	static TSharedPtr< class FSlateStyleSet > VREditorStyleInstance;
+
+	
+};
\ No newline at end of file
diff --git a/Engine/Source/Runtime/Slate/Private/Framework/MultiBox/MultiBox.cpp b/Engine/Source/Runtime/Slate/Private/Framework/MultiBox/MultiBox.cpp
--- a/Engine/Source/Runtime/Slate/Private/Framework/MultiBox/MultiBox.cpp
+++ b/Engine/Source/Runtime/Slate/Private/Framework/MultiBox/MultiBox.cpp
@@ -248,7 +248,7 @@ void FMultiBox::InsertCustomMultiBlock( TSharedRef<const FMultiBlock> InBlock, i
  *
  * @return  MultiBox widget object
  */
-TSharedRef< SMultiBoxWidget > FMultiBox::MakeWidget( bool bSearchable )
+TSharedRef< SMultiBoxWidget > FMultiBox::MakeWidget( bool bSearchable, FOnMakeMultiBoxBuilderOverride* InMakeMultiBoxBuilderOverride /* = nullptr */ )
 {	
 	ApplyCustomizedBlocks();
 
@@ -261,8 +261,16 @@ TSharedRef< SMultiBoxWidget > FMultiBox::MakeWidget( bool bSearchable )
 	// Assign ourselves to the MultiBox widget
 	NewMultiBoxWidget->SetMultiBox( AsShared() );
 
-	// Build up the widget
-	NewMultiBoxWidget->BuildMultiBoxWidget();
+	if( (InMakeMultiBoxBuilderOverride != nullptr) && (InMakeMultiBoxBuilderOverride->IsBound()) )
+	{
+		TSharedRef<FMultiBox> ThisMultiBox = AsShared();
+		InMakeMultiBoxBuilderOverride->Execute( ThisMultiBox, NewMultiBoxWidget );
+	}
+	else
+	{
+		// Build up the widget
+		NewMultiBoxWidget->BuildMultiBoxWidget();
+	}
 	
 #if PLATFORM_MAC
 	if(Type == EMultiBoxType::MenuBar)
diff --git a/Engine/Source/Runtime/Slate/Private/Framework/MultiBox/MultiBoxBuilder.cpp b/Engine/Source/Runtime/Slate/Private/Framework/MultiBox/MultiBoxBuilder.cpp
--- a/Engine/Source/Runtime/Slate/Private/Framework/MultiBox/MultiBoxBuilder.cpp
+++ b/Engine/Source/Runtime/Slate/Private/Framework/MultiBox/MultiBoxBuilder.cpp
@@ -80,9 +80,9 @@ FMultiBoxCustomization FMultiBoxBuilder::GetCustomization() const
 	return FMultiBoxCustomization( MultiBox->GetCustomizationName() ); 
 }
 
-TSharedRef< class SWidget > FMultiBoxBuilder::MakeWidget()
+TSharedRef< class SWidget > FMultiBoxBuilder::MakeWidget( FMultiBox::FOnMakeMultiBoxBuilderOverride* InMakeMultiBoxBuilderOverride /* = nullptr */ )
 {
-	return MultiBox->MakeWidget( false );
+	return MultiBox->MakeWidget( false, InMakeMultiBoxBuilderOverride );
 }
 
 TSharedRef< class FMultiBox > FMultiBoxBuilder::GetMultiBox()
@@ -163,10 +163,10 @@ void FBaseMenuBuilder::AddMenuEntry( const FUIAction& UIAction, const TSharedRef
 	ApplyHook(InExtensionHook, EExtensionHook::After);
 }
 
-TSharedRef< class SWidget > FMenuBuilder::MakeWidget()
+TSharedRef< class SWidget > FMenuBuilder::MakeWidget( FMultiBox::FOnMakeMultiBoxBuilderOverride* InMakeMultiBoxBuilderOverride /* = nullptr */ )
 {
 	// Make menu builders searchable
-	return MultiBox->MakeWidget( true );
+	return MultiBox->MakeWidget( true, InMakeMultiBoxBuilderOverride );
 }
 
 void FMenuBuilder::BeginSection( FName InExtensionHook, const TAttribute< FText >& InHeadingText )
diff --git a/Engine/Source/Runtime/Slate/Private/Widgets/Input/SButton.cpp b/Engine/Source/Runtime/Slate/Private/Widgets/Input/SButton.cpp
--- a/Engine/Source/Runtime/Slate/Private/Widgets/Input/SButton.cpp
+++ b/Engine/Source/Runtime/Slate/Private/Widgets/Input/SButton.cpp
@@ -418,6 +418,16 @@ void SButton::SetOnClicked(FOnClicked InOnClicked)
 	OnClicked = InOnClicked;
 }
 
+void SButton::SetOnHovered(FSimpleDelegate InOnHovered)
+{
+	OnHovered = InOnHovered;
+}
+
+void SButton::SetOnUnhovered(FSimpleDelegate InOnUnhovered)
+{
+	OnUnhovered = InOnUnhovered;
+}
+
 void SButton::SetButtonStyle(const FButtonStyle* ButtonStyle)
 {
 	/* Get pointer to the button style */
diff --git a/Engine/Source/Runtime/Slate/Public/Framework/MultiBox/MultiBox.h b/Engine/Source/Runtime/Slate/Public/Framework/MultiBox/MultiBox.h
--- a/Engine/Source/Runtime/Slate/Public/Framework/MultiBox/MultiBox.h
+++ b/Engine/Source/Runtime/Slate/Public/Framework/MultiBox/MultiBox.h
@@ -239,12 +239,14 @@ public:
 	 */
 	void InsertCustomMultiBlock( TSharedRef<const FMultiBlock> InBlock, int32 Index );
 
+	DECLARE_DELEGATE_TwoParams( FOnMakeMultiBoxBuilderOverride, const TSharedRef<FMultiBox>&, const TSharedRef<SMultiBoxWidget>& );
+
 	/**
 	 * Creates a MultiBox widget for this MultiBox
 	 *
 	 * @return  MultiBox widget object
 	 */
-	TSharedRef< class SMultiBoxWidget > MakeWidget( bool bSearchable );
+	TSharedRef< class SMultiBoxWidget > MakeWidget( bool bSearchable, FOnMakeMultiBoxBuilderOverride* InMakeMultiBoxBuilderOverride = nullptr );
 
 
 	/**
@@ -475,6 +477,18 @@ public:
 		MultiBox = InMultiBox;
 	}
 
+	/**
+	* Sets new content for the SMultiBoxWidget
+	*
+	* @param	InContent	The new content to place in the ChildSlot
+	*/
+	void SetContent( TSharedRef< SWidget > InContent )
+	{	
+		ChildSlot
+		[
+			InContent 
+		];
+	}
 
 	/**
 	 * Access the MultiBox associated with this widget
diff --git a/Engine/Source/Runtime/Slate/Public/Framework/MultiBox/MultiBoxBuilder.h b/Engine/Source/Runtime/Slate/Public/Framework/MultiBox/MultiBoxBuilder.h
--- a/Engine/Source/Runtime/Slate/Public/Framework/MultiBox/MultiBoxBuilder.h
+++ b/Engine/Source/Runtime/Slate/Public/Framework/MultiBox/MultiBoxBuilder.h
@@ -4,7 +4,7 @@
 
 #include "MultiBoxExtender.h"
 #include "SlateDelegates.h"
-
+#include "MultiBox.h"
 
 class FUICommandInfo;
 class FUICommandList;
@@ -53,7 +53,7 @@ public:
 	 *
 	 * @return  New widget object
 	 */
-	virtual TSharedRef< class SWidget > MakeWidget();
+	virtual TSharedRef< class SWidget > MakeWidget( FMultiBox::FOnMakeMultiBoxBuilderOverride* InMakeMultiBoxBuilderOverride = nullptr );
 	
 
 	/** 
@@ -210,7 +210,7 @@ public:
 	*
 	* @return  New widget object
 	*/
-	virtual TSharedRef< class SWidget > MakeWidget() override;
+	virtual TSharedRef< class SWidget > MakeWidget( FMultiBox::FOnMakeMultiBoxBuilderOverride* InMakeMultiBoxBuilderOverride = nullptr ) override;
 
 	/**
 	 * Adds a menu separator
diff --git a/Engine/Source/Runtime/Slate/Public/Widgets/Input/SButton.h b/Engine/Source/Runtime/Slate/Public/Widgets/Input/SButton.h
--- a/Engine/Source/Runtime/Slate/Public/Widgets/Input/SButton.h
+++ b/Engine/Source/Runtime/Slate/Public/Widgets/Input/SButton.h
@@ -130,6 +130,12 @@ public:
 	/** See OnClicked event */
 	void SetOnClicked(FOnClicked InOnClicked);
 
+	/** Set OnHovered event */
+	void SetOnHovered(FSimpleDelegate InOnHovered);
+
+	/** Set OnUnhovered event */
+	void SetOnUnhovered(FSimpleDelegate InOnUnhovered);
+
 	/** See ButtonStyle attribute */
 	void SetButtonStyle(const FButtonStyle* ButtonStyle);
 
