Fix for function name overlap in SRadialPanel.cpp

diff --git a/Engine/Source/Editor/VREditor/SVRRadialPanel.cpp b/Engine/Source/Editor/VREditor/SVRRadialPanel.cpp
--- a/Engine/Source/Editor/VREditor/SVRRadialPanel.cpp
+++ b/Engine/Source/Editor/VREditor/SVRRadialPanel.cpp
@@ -117,7 +117,7 @@ static void ArrangeChildrenAlong( const TPanelChildren<SRadialPanel::FSlot>& Chi
 			AlignmentArrangeResult YAlignmentResult = AlignChild<Orient_Vertical>( SlotSize.Y, CurChild, SlotPadding );
 			
 			FVector2D MyDesiredSize(0, 0);
-			MyDesiredSize = ComputeDesiredSizeForBox(Children);
+			MyDesiredSize = ComputeDesiredSizeForMenu(Children);
 			
 			// Set the x and y coordinates of the slot based on the current angle, a radius of MyDesiredSize.Y / 2.0, and a center of (0, MyDesiredSize.Y / 2.0)
 			// First element ends up at 90 degrees
@@ -166,7 +166,7 @@ void SRadialPanel::OnArrangeChildren( const FGeometry& AllottedGeometry, FArrang
  *
  * @return The size desired by the children given an orientation.
  */
-static FVector2D ComputeDesiredSizeForBox( const TPanelChildren<SRadialPanel::FSlot>& Children ) 
+static FVector2D ComputeDesiredSizeForMenu( const TPanelChildren<SRadialPanel::FSlot>& Children ) 
 {
 	// The desired size of this panel is the total size desired by its children plus any margins specified in this panel.
 	// The layout along the panel's axis is describe dy the SizeParam, while the perpendicular layout is described by the
@@ -207,7 +207,7 @@ static FVector2D ComputeDesiredSizeForBox( const TPanelChildren<SRadialPanel::FS
  */
 FVector2D SRadialPanel::ComputeDesiredSize( float ) const
 {
-	return ComputeDesiredSizeForBox(this->Children);
+	return ComputeDesiredSizeForMenu(this->Children);
 }
 
 /** @return  The children of a panel in a slot-agnostic way. */
