Number pad for variable entry in the VR Editor. Number pad appears when clicking on an editable text field, can be dismissed with center button. Trigger for key entry, grip for confirming entry, modifier key for backspace. Also includes menu positioning bug fixes and button pressed sounds.

diff --git a/Engine/Source/Editor/VREditor/SVRRadialPanel.cpp b/Engine/Source/Editor/VREditor/SVRRadialPanel.cpp
--- a/Engine/Source/Editor/VREditor/SVRRadialPanel.cpp
+++ b/Engine/Source/Editor/VREditor/SVRRadialPanel.cpp
@@ -23,6 +23,8 @@ void SRadialBox::Construct( const SRadialBox::FArguments& InArgs )
 	}
 
 	CurrentlyHoveredButton = nullptr;
+	RadiusRatioOverride = InArgs._RadiusRatio;
+
 }
 
 /**
@@ -68,7 +70,7 @@ static FVector2D ComputeDesiredSizeForMenu(const TPanelChildren<SRadialPanel::FS
 
 
 template<EOrientation Orientation>
-static void ArrangeChildrenAlong( const TPanelChildren<SRadialPanel::FSlot>& Children, const FGeometry& AllottedGeometry, FArrangedChildren& ArrangedChildren )
+static void ArrangeChildrenAlong( const TPanelChildren<SRadialPanel::FSlot>& Children, const FGeometry& AllottedGeometry, FArrangedChildren& ArrangedChildren, TAttribute<float> RadiusOverride )
 {
 	// Allotted space will be given to fixed-size children first.
 	// Remaining space will be proportionately divided between stretch children (SizeRule_Stretch)
@@ -158,10 +160,11 @@ static void ArrangeChildrenAlong( const TPanelChildren<SRadialPanel::FSlot>& Chi
 			FVector2D MyDesiredSize(0, 0);
 			MyDesiredSize = ComputeDesiredSizeForMenu(Children);
 			
-			// Set the x and y coordinates of the slot based on the current angle, a radius of MyDesiredSize.Y / 2.0, and a center of (0, MyDesiredSize.Y / 2.0)
+			// Set the x and y coordinates of the slot based on the current angle, a default radius of MyDesiredSize.Y / 2, and a center of (0, MyDesiredSize.Y / 2.0)
 			// First element ends up at 90 degrees
-			float RadialX = MyDesiredSize.Y / 2.0f * FMath::Cos(PositionSoFar * ( (2.0f*PI) / Children.Num() ));
-			float RadialY = MyDesiredSize.Y / 2.0f + MyDesiredSize.Y / 2.0f * FMath::Sin(PositionSoFar * ( (2.0f*PI) / Children.Num() ));
+			float Radius = (MyDesiredSize.Y / 2.0f) * RadiusOverride.Get();
+			float RadialX =  Radius * FMath::Cos(PositionSoFar * ( (2.0f*PI) / Children.Num() ));
+			float RadialY = MyDesiredSize.Y / 2.0f + Radius * FMath::Sin(PositionSoFar * ( (2.0f*PI) / Children.Num() ));
 			 
 
 			// Add the information about this child to the output list (ArrangedChildren)
@@ -193,7 +196,7 @@ static void ArrangeChildrenAlong( const TPanelChildren<SRadialPanel::FSlot>& Chi
 void SRadialPanel::OnArrangeChildren( const FGeometry& AllottedGeometry, FArrangedChildren& ArrangedChildren ) const
 {
 	
-	ArrangeChildrenAlong<Orient_Vertical>(this->Children, AllottedGeometry, ArrangedChildren );
+	ArrangeChildrenAlong<Orient_Vertical>(this->Children, AllottedGeometry, ArrangedChildren, RadiusRatioOverride );
 	
 }
 
diff --git a/Engine/Source/Editor/VREditor/SVRRadialPanel.h b/Engine/Source/Editor/VREditor/SVRRadialPanel.h
--- a/Engine/Source/Editor/VREditor/SVRRadialPanel.h
+++ b/Engine/Source/Editor/VREditor/SVRRadialPanel.h
@@ -88,6 +88,9 @@ protected:
 
 	/** The Radial Panel's children. */
 	TPanelChildren<FSlot> Children;
+
+	/** The ratio to adjust the radial menu's radius by (relative to the default)*/
+	TAttribute<float> RadiusRatioOverride;
 };
 
 
@@ -177,10 +180,12 @@ public:
 
 	SLATE_BEGIN_ARGS( SRadialBox )
 	{
-		_Visibility = EVisibility::SelfHitTestInvisible;
+		_Visibility = EVisibility::SelfHitTestInvisible,
+		/** The ratio to adjust the radial menu's radius by (relative to the default) */
+		_RadiusRatio = 1.0f;
 	}
-
-		SLATE_SUPPORTS_SLOT(SRadialBox::FSlot)
+	SLATE_ARGUMENT( float, RadiusRatio )
+	SLATE_SUPPORTS_SLOT(SRadialBox::FSlot)
 
 	SLATE_END_ARGS()
 
diff --git a/Engine/Source/Editor/VREditor/UI/VREditorUISystem.cpp b/Engine/Source/Editor/VREditor/UI/VREditorUISystem.cpp
--- a/Engine/Source/Editor/VREditor/UI/VREditorUISystem.cpp
+++ b/Engine/Source/Editor/VREditor/UI/VREditorUISystem.cpp
@@ -112,6 +112,7 @@ UVREditorUISystem::UVREditorUISystem( const FObjectInitializer& Initializer ) :
 	StopDragUISound( nullptr ),
 	HideUISound( nullptr ),
 	ShowUISound( nullptr ),
+	ButtonPressSound( nullptr ),
 	bDraggedDockFromHandPassedThreshold( false ),
 	LastDraggingHoverLocation( FVector::ZeroVector ),
 	bSetDefaultLayout( true )
@@ -147,6 +148,9 @@ void UVREditorUISystem::Init()
 	ShowUISound = LoadObject<USoundCue>( nullptr, TEXT( "/Engine/VREditor/Sounds/VR_open_Cue" ) );
 	check( ShowUISound != nullptr );
 
+	ButtonPressSound = LoadObject<USoundCue>(nullptr, TEXT("/Engine/VREditor/Sounds/VR_click1_Cue"));
+	check(ButtonPressSound != nullptr);
+
 	// Create all of our UI panels
 	CreateUIs();
 
@@ -177,6 +181,7 @@ void UVREditorUISystem::Shutdown()
 	StopDragUISound = nullptr;
 	HideUISound = nullptr;
 	ShowUISound = nullptr;
+	ButtonPressSound = nullptr;
 	VRMode = nullptr;
 	DraggingUI = nullptr;
 	ColorPickerUI = nullptr;
@@ -192,6 +197,32 @@ void UVREditorUISystem::OnVRAction( FEditorViewportClient& ViewportClient, UView
 	{
 		if ( !bWasHandled )
 		{
+			// If the numpad is currently showing
+			if (bRadialMenuIsNumpad)
+			{
+				// Modifier button is backspace
+				if (Action.ActionType == VRActionTypes::Modifier)
+				{
+					const bool bRepeat = false;
+					FVREditorActionCallbacks::SimulateKeyDown( EKeys::BackSpace, bRepeat );
+					FVREditorActionCallbacks::SimulateKeyUp( EKeys::BackSpace );
+					bWasHandled = true;
+				}
+				// Side triggers function as enter keys
+				if (Action.ActionType == ViewportWorldActionTypes::WorldMovement)
+				{
+					const bool bRepeat = false;
+					FVREditorActionCallbacks::SimulateKeyDown( EKeys::Enter, bRepeat );
+					FVREditorActionCallbacks::SimulateKeyUp( EKeys::Enter );
+					// After pressing enter, dismiss the numpad
+					SwapRadialMenu();
+					if (!bRadialMenuVisibleAtSwap)
+					{
+						HideRadialMenu( VREditorInteractor );
+					}
+					bWasHandled = true;
+				}
+			}
 			if ( Action.ActionType == VRActionTypes::ConfirmRadialSelection )
 			{
 				const EViewportInteractionDraggingMode DraggingMode = Interactor->GetDraggingMode();
@@ -203,11 +234,22 @@ void UVREditorUISystem::OnVRAction( FEditorViewportClient& ViewportClient, UView
 				{
 					if ( !IsShowingRadialMenu(VREditorInteractor) )
 					{
-						TryToSpawnRadialMenu(VREditorInteractor);
+						// Don't force the UI to spawn if using the touchpad/thumbstick to spawn the menu
+						const bool bForceOverUI = false;
+						TryToSpawnRadialMenu(VREditorInteractor, bForceOverUI);
 					}
 					else
-					{
-						HideRadialMenu(VREditorInteractor);
+					{	
+						// If you are hiding the number pad, also toggle back to the radial menu
+						if(bRadialMenuIsNumpad)
+						{
+							SwapRadialMenu();
+						}
+						// If we are not swapping back to the radial menu from the number pad, hide the menu
+						if (!bRadialMenuVisibleAtSwap)
+						{
+							HideRadialMenu( VREditorInteractor );
+						}
 					}
 					bWasHandled = true;
 				}
@@ -219,20 +261,64 @@ void UVREditorUISystem::OnVRAction( FEditorViewportClient& ViewportClient, UView
 				if ( IsShowingRadialMenu( VREditorInteractor ) )
 				{
 					TSharedPtr<SRadialBox> RadialBox = StaticCastSharedPtr<SRadialBox>(RadialWidget);
-					if(Action.Event == IE_Pressed)
+					if(RadialBox->GetCurrentlyHoveredButton().IsValid())
 					{
-						RadialBox->SimulateLeftClickDown();
+						if((Action.Event == IE_Pressed))
+						{
+							RadialBox->SimulateLeftClickDown();
+							bOutIsInputCaptured = true;
+						}
+						if (Action.Event == IE_Released)
+						{
+							RadialBox->SimulateLeftClickUp();
+							bOutIsInputCaptured = false;
+						}
+						bWasHandled = true;
 					}
-					if (Action.Event == IE_Released)
+					else
 					{
-						RadialBox->SimulateLeftClickUp();
-						// Hide the radial menu after selection
-						HideRadialMenu(VREditorInteractor);
+						FVector LaserPointerStart, LaserPointerEnd;
+						if (VREditorInteractor->GetLaserPointer( LaserPointerStart, LaserPointerEnd ))
+						{
+							FHitResult HitResult = VREditorInteractor->GetHitResultFromLaserPointer();
+							if (HitResult.Actor.IsValid())
+							{
+								UWidgetComponent* WidgetComponent = Cast<UWidgetComponent>( HitResult.GetComponent() );
+								// If we are clicking on an Actor but not a widget component, send a fake mouse click event to toggle focus
+								if (WidgetComponent == nullptr)
+								{
+									const bool bIsRightClicking =
+										(Action.Event == IE_Pressed && VREditorInteractor->IsModifierPressed()) ||
+										(Action.Event == IE_Released && VREditorInteractor->IsRightClickingOnUI());
+									TSet<FKey> PressedButtons;
+									FPointerEvent PointerEvent(
+										1 + (uint8)VREditorInteractor->GetControllerSide(),
+										FVector2D::ZeroVector,
+										FVector2D::ZeroVector,
+										PressedButtons,
+										bIsRightClicking ? EKeys::RightMouseButton : EKeys::LeftMouseButton,
+										0.0f,	// Wheel delta
+										FModifierKeysState() );
+
+									FWidgetPath EmptyWidgetPath;
+									FReply Reply = FSlateApplication::Get().RoutePointerDownEvent( EmptyWidgetPath, PointerEvent );
+									bWasHandled = true;
+									if (bRadialMenuIsNumpad)
+									{
+										// If clicking somewhere outside UI so the widget loses focus
+										SwapRadialMenu();
+										if (!bRadialMenuVisibleAtSwap)
+										{
+											HideRadialMenu( VREditorInteractor );
+										}
+									}
+								}
+							}
+						}
 					}
-					bWasHandled = true;
 				}
 
-				else
+				if(!bWasHandled)
 				{
 					FVector LaserPointerStart, LaserPointerEnd;
 					if ( VREditorInteractor->GetLaserPointer( LaserPointerStart, LaserPointerEnd ) )
@@ -306,6 +392,21 @@ void UVREditorUISystem::OnVRAction( FEditorViewportClient& ViewportClient, UView
 											}
 											else
 											{
+												// If we are clicking on an editable text field and the radial menu is not a numpad, show the numpad
+												if (WidgetPathUnderFinger.Widgets.Last().Widget->GetTypeAsString() == TEXT("SEditableText") && (bRadialMenuIsNumpad == false))
+												{
+													if (!QuickRadialMenu->bHidden)
+													{
+														bRadialMenuVisibleAtSwap = true;
+													}
+													else
+													{
+														// Force the radial menu to spawn even if the laser is over UI
+														const bool bForceOverUI = true;
+														TryToSpawnRadialMenu( VREditorInteractor, bForceOverUI );
+													}
+													SwapRadialMenu();
+												}
 												Reply = FSlateApplication::Get().RoutePointerDownEvent(WidgetPathUnderFinger, PointerEvent);
 											}
 
@@ -331,6 +432,25 @@ void UVREditorUISystem::OnVRAction( FEditorViewportClient& ViewportClient, UView
 									}
 								}
 							}
+							else
+							{
+								// If we didn't handle the input in any other way, send an empty mouse down event so Slate focus is handled correctly
+								const bool bIsRightClicking =
+									(Action.Event == IE_Pressed && VREditorInteractor->IsModifierPressed()) ||
+									(Action.Event == IE_Released && VREditorInteractor->IsRightClickingOnUI());
+								TSet<FKey> PressedButtons;
+								FPointerEvent PointerEvent(
+									1 + (uint8)VREditorInteractor->GetControllerSide(),
+									FVector2D::ZeroVector,
+									FVector2D::ZeroVector,
+									PressedButtons,
+									bIsRightClicking ? EKeys::RightMouseButton : EKeys::LeftMouseButton,
+									0.0f,	// Wheel delta
+									FModifierKeysState() );
+
+								FWidgetPath EmptyWidgetPath;
+								FReply Reply = FSlateApplication::Get().RoutePointerDownEvent( EmptyWidgetPath, PointerEvent );
+							}
 						}
 					}
 				}
@@ -430,7 +550,7 @@ void UVREditorUISystem::OnVRHoverUpdate( FEditorViewportClient& ViewportClient,
 								FModifierKeysState());
 
 							FSlateApplication::Get().RoutePointerMoveEvent(WidgetPathUnderFinger, PointerEvent, false);
-
+							
 							bWasHandled = true;
 							HoverImpactPoint = HitResult.ImpactPoint;
 							VREditorInteractor->SetHoveringOverWidgetComponent( WidgetComponent );
@@ -447,7 +567,8 @@ void UVREditorUISystem::OnVRHoverUpdate( FEditorViewportClient& ViewportClient,
 								const bool bIsAbsolute = ( GetOwner().GetHMDDeviceType() == EHMDDeviceType::DT_SteamVR );
 
 								float ScrollDelta = 0.0f;
-								if( VREditorInteractor->IsTouchingTrackpad() || !bIsAbsolute )
+								// Don't scroll if the radial menu is a number pad
+								if(( VREditorInteractor->IsTouchingTrackpad() || !bIsAbsolute) && !bRadialMenuIsNumpad)
 								{
 									if( bIsAbsolute )
 									{
@@ -1094,7 +1215,7 @@ void UVREditorUISystem::ShowEditorUIPanel( AVREditorFloatingUI* Panel, UVREditor
 
 bool UVREditorUISystem::IsShowingRadialMenu( UVREditorInteractor* Interactor ) const
 {
-	const EControllerHand DockedToHand = QuickRadialMenu->GetDockedTo() == AVREditorFloatingUI::EDockedTo::LeftHand ? EControllerHand::Left : EControllerHand::Right;
+	const EControllerHand DockedToHand = QuickRadialMenu->GetDockedTo() == AVREditorFloatingUI::EDockedTo::LeftArm ? EControllerHand::Left : EControllerHand::Right;
 	const UVREditorInteractor* DockedToInteractor = GetOwner().GetHandInteractor( DockedToHand );
 	return DockedToInteractor == Interactor && !QuickRadialMenu->bHidden;
 }
@@ -1103,7 +1224,7 @@ bool UVREditorUISystem::IsShowingRadialMenu( UVREditorInteractor* Interactor ) c
 
 
 
-void UVREditorUISystem::TryToSpawnRadialMenu( UVREditorInteractor* Interactor )
+void UVREditorUISystem::TryToSpawnRadialMenu( UVREditorInteractor* Interactor, const bool bForceOverUI )
 {
 	if( Interactor != nullptr )
 	{
@@ -1113,7 +1234,7 @@ void UVREditorUISystem::TryToSpawnRadialMenu( UVREditorInteractor* Interactor )
 			UVREditorInteractor* DockedToInteractor = nullptr;
 			if( !QuickRadialMenu->bHidden )
 			{
-				const EControllerHand ControllerHand = QuickRadialMenu->GetDockedTo() == AVREditorFloatingUI::EDockedTo::LeftHand ? EControllerHand::Left : EControllerHand::Right;
+				const EControllerHand ControllerHand = QuickRadialMenu->GetDockedTo() == AVREditorFloatingUI::EDockedTo::LeftArm ? EControllerHand::Left : EControllerHand::Right;
 				DockedToInteractor = GetOwner().GetHandInteractor( ControllerHand );
 			}
 
@@ -1127,11 +1248,11 @@ void UVREditorUISystem::TryToSpawnRadialMenu( UVREditorInteractor* Interactor )
 				DraggingMode != EViewportInteractionDraggingMode::World &&
 				DraggingMode != EViewportInteractionDraggingMode::AssistingDrag &&
 				( InteractorDraggingUI == nullptr || InteractorDraggingUI != Interactor ) &&
-				!Interactor->IsHoveringOverUI();	// Don't show radial menu when aiming at a UI  (too much clutter)
+				(!Interactor->IsHoveringOverUI() || bForceOverUI);	// Don't show radial menu when aiming at a UI unless forcing it (too much clutter)
 
 			if (bNeedsSpawn)
 			{
-				const AVREditorFloatingUI::EDockedTo DockedTo = MotionControllerInteractor->GetControllerSide() == EControllerHand::Left ? AVREditorFloatingUI::EDockedTo::LeftHand : AVREditorFloatingUI::EDockedTo::RightHand;
+				const AVREditorFloatingUI::EDockedTo DockedTo = MotionControllerInteractor->GetControllerSide() == EControllerHand::Left ? AVREditorFloatingUI::EDockedTo::LeftArm : AVREditorFloatingUI::EDockedTo::RightArm;
 				QuickRadialMenu->SetDockedTo(DockedTo);
 				QuickRadialMenu->ShowUI(true);
 			}
@@ -1637,7 +1758,7 @@ TSharedRef<SWidget> UVREditorUISystem::AddHoverableButton(TSharedRef<SWidget>& B
 		]
 		];
 	// Placeholder slot to put the button into when it is hovered
-	TestOverlay->AddSlot(1500)
+	TestOverlay->AddSlot(10)
 	.Expose(HoverSlot)
 		[
 			SNew(SOverlay)
@@ -1651,6 +1772,9 @@ TSharedRef<SWidget> UVREditorUISystem::AddHoverableButton(TSharedRef<SWidget>& B
 		Button->SetRenderTransformPivot(FVector2D(0.5, 0.5));
 		Button->SetOnHovered(FSimpleDelegate::CreateUObject(this, &UVREditorUISystem::OnHoverBeginEffect, Button));
 		Button->SetOnUnhovered(FSimpleDelegate::CreateUObject(this, &UVREditorUISystem::OnHoverEndEffect, Button));
+		FSlateSound SlateButtonPressSound = FSlateSound();
+		SlateButtonPressSound.SetResourceObject(ButtonPressSound);
+		Button->SetPressedSound(SlateButtonPressSound);
 		ButtonBorder->SetRenderTransformPivot(FVector2D(0.5, 0.5));
 		FVRButton ButtonToAdd = FVRButton(Button, BaseSlot, HoverSlot, ButtonBorder);
 		VRButtons.Add(ButtonToAdd);
@@ -1670,7 +1794,7 @@ TSharedRef<SWidget> UVREditorUISystem::SetButtonTextWrap(TSharedRef<SWidget>& Bl
 	return BlockWidget;
 }
 
-void UVREditorUISystem::MakeRadialBoxMenu(const TSharedRef<FMultiBox>& MultiBox, const TSharedRef<SMultiBoxWidget>& MultiBoxWidget)
+void UVREditorUISystem::MakeRadialBoxMenu(const TSharedRef<FMultiBox>& MultiBox, const TSharedRef<SMultiBoxWidget>& MultiBoxWidget, float RadiusRatioOverride)
 {
 	// Grab the list of blocks, early out if there's nothing to fill the widget with
 	const TArray< TSharedRef< const FMultiBlock > >& Blocks = MultiBox->GetBlocks();
@@ -1680,14 +1804,13 @@ void UVREditorUISystem::MakeRadialBoxMenu(const TSharedRef<FMultiBox>& MultiBox,
 		return;
 	}
 	MultiBox.Get().SetStyle(&FVREditorStyle::Get(), FVREditorStyle::GetStyleSetName());
-	TSharedRef<SRadialBox> GridPanel = SNew(SRadialBox);
+	TSharedRef<SRadialBox> GridPanel = SNew(SRadialBox)
+		.RadiusRatio(RadiusRatioOverride);
 	GridPanel->NumEntries = Blocks.Num();
 	for (const TSharedRef<const FMultiBlock>& MultiBlock : Blocks)
 	{
 		const TSharedRef<const FMultiBlock>& Block = MultiBlock;
 		TSharedRef<SWidget> BlockWidget = Block->MakeWidget(MultiBoxWidget, EMultiBlockLocation::Middle, true)->AsWidget();
-
-		
 		TSharedRef<SOverlay> TestOverlay = SNew(SOverlay);
 		
 		GridPanel->AddSlot()
@@ -1702,7 +1825,6 @@ void UVREditorUISystem::MakeRadialBoxMenu(const TSharedRef<FMultiBox>& MultiBox,
 
 	}
 	RadialWidget = GridPanel;
-	StaticCastSharedPtr<SRadialBox>(RadialWidget)->NumEntries = Blocks.Num();
 	MultiBoxWidget->SetContent(
 			SNew(SDPIScaler)
 			.DPIScale(3)
@@ -2086,7 +2208,7 @@ TSharedRef<SWidget> UVREditorUISystem::BuildRadialMenuWidget()
 {
 
 	FMultiBox::FOnMakeMultiBoxBuilderOverride VREditorMenuBuilderOverride;
-	VREditorMenuBuilderOverride.BindUObject(this, &UVREditorUISystem::MakeRadialBoxMenu);
+	VREditorMenuBuilderOverride.BindUObject(this, &UVREditorUISystem::MakeRadialBoxMenu, 1.0f);
 	const TSharedRef<FExtender> MenuExtender = MakeShareable(new FExtender());
 	TSharedPtr<FUICommandList> CommandList = MakeShareable(new FUICommandList());
 	FMenuBarBuilder MenuBuilder(CommandList, MenuExtender);
@@ -2096,90 +2218,248 @@ TSharedRef<SWidget> UVREditorUISystem::BuildRadialMenuWidget()
 
 	// First menu entry is at 90 degrees 
 	MenuBuilder.AddMenuEntry(
-		LOCTEXT("Delete", "Delete"),
-		LOCTEXT("DeleteTooltip", "Delete"),
+		LOCTEXT( "Delete", "Delete" ),
+		LOCTEXT( "DeleteTooltip", "Delete" ),
 		FSlateIcon(),
 		FUIAction
 		(
-			FExecuteAction::CreateStatic(&FLevelEditorActionCallbacks::ExecuteExecCommand, FString(TEXT("DELETE"))),
-			FCanExecuteAction::CreateStatic(&FLevelEditorActionCallbacks::Delete_CanExecute)
+			FExecuteAction::CreateStatic( &FLevelEditorActionCallbacks::ExecuteExecCommand, FString( TEXT( "DELETE" ) ) ),
+			FCanExecuteAction::CreateStatic( &FLevelEditorActionCallbacks::Delete_CanExecute )
 		)
 	);
 	MenuBuilder.AddMenuEntry(
-		LOCTEXT("Cut", "Cut"),
-		LOCTEXT("CutTooltip", "Cut"),
+		LOCTEXT( "Cut", "Cut" ),
+		LOCTEXT( "CutTooltip", "Cut" ),
 		FSlateIcon(),
 		FUIAction
 		(
-			FExecuteAction::CreateStatic(&FLevelEditorActionCallbacks::ExecuteExecCommand, FString(TEXT("EDIT CUT"))),
-			FCanExecuteAction::CreateStatic(&FLevelEditorActionCallbacks::Cut_CanExecute)
+			FExecuteAction::CreateStatic( &FLevelEditorActionCallbacks::ExecuteExecCommand, FString( TEXT( "EDIT CUT" ) ) ),
+			FCanExecuteAction::CreateStatic( &FLevelEditorActionCallbacks::Cut_CanExecute )
 		)
-		);
+	);
 	MenuBuilder.AddMenuEntry(
-		LOCTEXT("Copy", "Copy"),
-		LOCTEXT("CopyTooltip", "Copy"),
+		LOCTEXT( "Copy", "Copy" ),
+		LOCTEXT( "CopyTooltip", "Copy" ),
 		FSlateIcon(),
 		FUIAction
 		(
-			FExecuteAction::CreateStatic(&FLevelEditorActionCallbacks::ExecuteExecCommand, FString(TEXT("EDIT COPY"))),
-			FCanExecuteAction::CreateStatic(&FLevelEditorActionCallbacks::Copy_CanExecute)
+			FExecuteAction::CreateStatic( &FLevelEditorActionCallbacks::ExecuteExecCommand, FString( TEXT( "EDIT COPY" ) ) ),
+			FCanExecuteAction::CreateStatic( &FLevelEditorActionCallbacks::Copy_CanExecute )
 		)
 	);
 	MenuBuilder.AddMenuEntry(
-		LOCTEXT("Duplicate", "Duplicate"),
-		LOCTEXT("DuplicateTooltip", "Duplicate"),
+		LOCTEXT( "Duplicate", "Duplicate" ),
+		LOCTEXT( "DuplicateTooltip", "Duplicate" ),
 		FSlateIcon(),
 		FUIAction
 		(
-			FExecuteAction::CreateStatic(&FLevelEditorActionCallbacks::ExecuteExecCommand, FString(TEXT("DUPLICATE"))),
-			FCanExecuteAction::CreateStatic(&FLevelEditorActionCallbacks::Duplicate_CanExecute)
+			FExecuteAction::CreateStatic( &FLevelEditorActionCallbacks::ExecuteExecCommand, FString( TEXT( "DUPLICATE" ) ) ),
+			FCanExecuteAction::CreateStatic( &FLevelEditorActionCallbacks::Duplicate_CanExecute )
 		)
 	);
 	MenuBuilder.AddMenuEntry(
-		LOCTEXT("Paste", "Paste"),
-		LOCTEXT("PasteTooltip", "Paste"),
+		LOCTEXT( "Paste", "Paste" ),
+		LOCTEXT( "PasteTooltip", "Paste" ),
 		FSlateIcon(),
 		FUIAction
 		(
-			FExecuteAction::CreateStatic(&FLevelEditorActionCallbacks::ExecuteExecCommand, FString(TEXT("EDIT PASTE"))),
-			FCanExecuteAction::CreateStatic(&FLevelEditorActionCallbacks::Paste_CanExecute)
+			FExecuteAction::CreateStatic( &FLevelEditorActionCallbacks::ExecuteExecCommand, FString( TEXT( "EDIT PASTE" ) ) ),
+			FCanExecuteAction::CreateStatic( &FLevelEditorActionCallbacks::Paste_CanExecute )
 		)
 	);
 	MenuBuilder.AddMenuEntry(
-		LOCTEXT("Undo", "Undo"),
-		LOCTEXT("UndoTooltip", "Undo"),
+		LOCTEXT( "Undo", "Undo" ),
+		LOCTEXT( "UndoTooltip", "Undo" ),
 		FSlateIcon(),
 		FUIAction
 		(
-			FExecuteAction::CreateStatic(&FLevelEditorActionCallbacks::ExecuteExecCommand, FString(TEXT("TRANSACTION UNDO"))),
-			FCanExecuteAction::CreateStatic(&FLevelEditorActionCallbacks::DefaultCanExecuteAction)
-			)
-		);
+			FExecuteAction::CreateStatic( &FLevelEditorActionCallbacks::ExecuteExecCommand, FString( TEXT( "TRANSACTION UNDO" ) ) ),
+			FCanExecuteAction::CreateStatic( &FLevelEditorActionCallbacks::DefaultCanExecuteAction )
+		)
+	);
 	MenuBuilder.AddMenuEntry(
-		LOCTEXT("SnapActorsToGround", "Snap Actors To Ground"),
-		LOCTEXT("SnapActorsToGroundTooltip", "Snap Actors To Ground"),
+		LOCTEXT( "SnapActorsToGround", "Snap Actors To Ground" ),
+		LOCTEXT( "SnapActorsToGroundTooltip", "Snap Actors To Ground" ),
 		FSlateIcon(),
 		FUIAction
 		(
-			FExecuteAction::CreateStatic(&FVREditorActionCallbacks::OnSnapActorsToGroundClicked, VRMode),
-			FCanExecuteAction::CreateStatic(&FLevelEditorActionCallbacks::Copy_CanExecute)
+			FExecuteAction::CreateStatic( &FVREditorActionCallbacks::OnSnapActorsToGroundClicked, VRMode ),
+			FCanExecuteAction::CreateStatic( &FLevelEditorActionCallbacks::Copy_CanExecute )
 		)
 	);
 	MenuBuilder.AddMenuEntry(
-		LOCTEXT("Redo", "Redo"),
-		LOCTEXT("RedoTooltip", "Redo"),
+		LOCTEXT( "Redo", "Redo" ),
+		LOCTEXT( "RedoTooltip", "Redo" ),
 		FSlateIcon(),
 		FUIAction
 		(
-			FExecuteAction::CreateStatic(&FLevelEditorActionCallbacks::ExecuteExecCommand, FString(TEXT("TRANSACTION REDO"))),
-			FCanExecuteAction::CreateStatic(&FLevelEditorActionCallbacks::DefaultCanExecuteAction)
+			FExecuteAction::CreateStatic( &FLevelEditorActionCallbacks::ExecuteExecCommand, FString( TEXT( "TRANSACTION REDO" ) ) ),
+			FCanExecuteAction::CreateStatic( &FLevelEditorActionCallbacks::DefaultCanExecuteAction )
+		)
+	);
+
+	// Create the menu widget
+	TSharedRef<SWidget> MenuWidget = MenuBuilder.MakeWidget(&VREditorMenuBuilderOverride);
+	return MenuWidget;
+}
+
+TSharedRef<SWidget> UVREditorUISystem::BuildNumPadWidget()
+{
+
+	FMultiBox::FOnMakeMultiBoxBuilderOverride VREditorMenuBuilderOverride;
+	VREditorMenuBuilderOverride.BindUObject(this, &UVREditorUISystem::MakeRadialBoxMenu, 0.5f);
+	const TSharedRef<FExtender> MenuExtender = MakeShareable(new FExtender());
+	TSharedPtr<FUICommandList> CommandList = MakeShareable(new FUICommandList());
+	FMenuBarBuilder MenuBuilder(CommandList, MenuExtender);
+
+	// First menu entry is at 90 degrees 
+
+	MenuBuilder.AddMenuEntry(
+		LOCTEXT("NumThree", "3"),
+		LOCTEXT("", ""),
+		FSlateIcon(),
+		FUIAction
+		(
+			FExecuteAction::CreateStatic(&FVREditorActionCallbacks::SimulateCharacterEntry, FString::FromInt(3))
+		)
+		);
+	MenuBuilder.AddMenuEntry(
+		LOCTEXT("NumFour", "4"),
+		LOCTEXT("", ""),
+		FSlateIcon(),
+		FUIAction
+		(
+			FExecuteAction::CreateStatic(&FVREditorActionCallbacks::SimulateCharacterEntry, FString::FromInt(4))
+		)
+		);
+	MenuBuilder.AddMenuEntry(
+		LOCTEXT("NumFive", "5"),
+		LOCTEXT("", ""),
+		FSlateIcon(),
+		FUIAction
+		(
+			FExecuteAction::CreateStatic(&FVREditorActionCallbacks::SimulateCharacterEntry, FString::FromInt(5))
+		)
+		);
+	MenuBuilder.AddMenuEntry(
+		LOCTEXT("NumSix", "6"),
+		LOCTEXT("", ""),
+		FSlateIcon(),
+		FUIAction
+		(
+			FExecuteAction::CreateStatic(&FVREditorActionCallbacks::SimulateCharacterEntry, FString::FromInt(6))
+		)
+		);
+	MenuBuilder.AddMenuEntry(
+		LOCTEXT("NumSeven", "7"),
+		LOCTEXT("", ""),
+		FSlateIcon(),
+		FUIAction
+		(
+			FExecuteAction::CreateStatic(&FVREditorActionCallbacks::SimulateCharacterEntry, FString::FromInt(7))
 			)
 		);
-
+	MenuBuilder.AddMenuEntry(
+		LOCTEXT("NumEight", "8"),
+		LOCTEXT("", ""),
+		FSlateIcon(),
+		FUIAction
+		(
+			FExecuteAction::CreateStatic(&FVREditorActionCallbacks::SimulateCharacterEntry, FString::FromInt(8))
+			)
+		);
+	MenuBuilder.AddMenuEntry(
+		LOCTEXT("NumNine", "9"),
+		LOCTEXT("", ""),
+		FSlateIcon(),
+		FUIAction
+		(
+			FExecuteAction::CreateStatic(&FVREditorActionCallbacks::SimulateCharacterEntry, FString::FromInt(9))
+		)
+		);
+	MenuBuilder.AddMenuEntry(
+		LOCTEXT( "Hyphen", "-" ),
+		LOCTEXT( "", "" ),
+		FSlateIcon(),
+		FUIAction
+		(
+			FExecuteAction::CreateStatic( &FVREditorActionCallbacks::SimulateCharacterEntry, FString::FString( TEXT( "-" ) ) )
+		)
+	);
+	MenuBuilder.AddMenuEntry(
+		LOCTEXT( "Decimal", "." ),
+		LOCTEXT( "", "" ),
+		FSlateIcon(),
+		FUIAction
+		(
+			FExecuteAction::CreateStatic( &FVREditorActionCallbacks::SimulateCharacterEntry, FString::FString( TEXT( "." ) ) )
+		)
+	);
+	MenuBuilder.AddMenuEntry(
+		LOCTEXT( "NumZero", "0" ),
+		LOCTEXT( "", "" ),
+		FSlateIcon(),
+		FUIAction
+		(
+			FExecuteAction::CreateStatic( &FVREditorActionCallbacks::SimulateCharacterEntry, FString::FromInt( 0 ) )
+		)
+	);
+	MenuBuilder.AddMenuEntry(
+		LOCTEXT( "NumOne", "1" ),
+		LOCTEXT( "", "" ),
+		FSlateIcon(),
+		FUIAction
+		(
+			FExecuteAction::CreateStatic( &FVREditorActionCallbacks::SimulateCharacterEntry, FString::FromInt( 1 ) )
+		)
+	);
+	MenuBuilder.AddMenuEntry(
+		LOCTEXT( "NumTwo", "2" ),
+		LOCTEXT( "", "" ),
+		FSlateIcon(),
+		FUIAction
+		(
+			FExecuteAction::CreateStatic( &FVREditorActionCallbacks::SimulateCharacterEntry, FString::FromInt( 2 ) )
+		)
+	);
 
 	// Create the menu widget
 	TSharedRef<SWidget> MenuWidget = MenuBuilder.MakeWidget(&VREditorMenuBuilderOverride);
-	return MenuWidget;
+
+	TSharedRef<SWidget> FinalWidget = SNew( SOverlay )
+		+ SOverlay::Slot()
+		[
+			MenuWidget
+		]
+		+ SOverlay::Slot()
+		.HAlign( HAlign_Center )
+		.VAlign( VAlign_Center )
+		[
+			SNew(SVerticalBox)
+			+ SVerticalBox::Slot()
+			.HAlign( HAlign_Fill )
+			.VAlign( VAlign_Center )
+			.AutoHeight()
+			[
+				SNew( STextBlock )
+				.Justification( ETextJustify::Center )
+				.Text( LOCTEXT( "BackspaceHelp", "Modifier: Backspace" ) )
+				.TextStyle( FVREditorStyle::Get(), "VREditorStyle.HelperText" )
+			]
+		+ SVerticalBox::Slot()
+			.HAlign( HAlign_Fill )
+			.VAlign( VAlign_Center )
+			.AutoHeight()
+			[
+				SNew( STextBlock )
+				.Justification( ETextJustify::Center )
+				.Text( LOCTEXT( "EnterHelp", "Grip: Enter" ) )
+			.	TextStyle( FVREditorStyle::Get(), "VREditorStyle.HelperText" )
+			]
+		]
+		;
+
+	return FinalWidget;
 }
 
 const TSharedRef<SWidget>& UVREditorUISystem::FindWidgetOfType(const TSharedRef<SWidget>& Content, FName WidgetType)
@@ -2245,4 +2525,23 @@ const TSharedPtr<SWidget>& UVREditorUISystem::GetRadialWidget()
 	return RadialWidget;
 }
 
+void UVREditorUISystem::SwapRadialMenu()
+{
+	const AVREditorFloatingUI::EDockedTo RadialMenuInteractor = QuickRadialMenu->GetDockedTo();
+	const FIntPoint DefaultResolution( VREd::DefaultEditorUIResolutionX->GetInt(), VREd::DefaultEditorUIResolutionY->GetInt() );
+	if (bRadialMenuIsNumpad)
+	{
+		QuickRadialMenu->SetSlateWidget( *this, BuildRadialMenuWidget(), DefaultResolution, 40.0f, RadialMenuInteractor );
+		bRadialMenuIsNumpad = false;
+	}
+	else
+	{
+		QuickRadialMenu->SetSlateWidget( *this, BuildNumPadWidget(), DefaultResolution, 40.0f, RadialMenuInteractor );
+		bRadialMenuIsNumpad = true;
+	}
+	QuickRadialMenu->GetWidgetComponent()->SetOpacityFromTexture( 1.0f );
+	QuickRadialMenu->GetWidgetComponent()->SetBackgroundColor( FLinearColor::Transparent );
+	QuickRadialMenu->GetWidgetComponent()->SetBlendMode( EWidgetBlendMode::Masked );
+}
+
 #undef LOCTEXT_NAMESPACE
\ No newline at end of file
diff --git a/Engine/Source/Editor/VREditor/UI/VREditorUISystem.h b/Engine/Source/Editor/VREditor/UI/VREditorUISystem.h
--- a/Engine/Source/Editor/VREditor/UI/VREditorUISystem.h
+++ b/Engine/Source/Editor/VREditor/UI/VREditorUISystem.h
@@ -145,7 +145,7 @@ public:
 	bool IsShowingRadialMenu( UVREditorInteractor* Interactor ) const;
 
 	/** Tries to spawn the radial menu (if the specified hand isn't doing anything else) */
-	void TryToSpawnRadialMenu( UVREditorInteractor* Interactor );
+	void TryToSpawnRadialMenu( UVREditorInteractor* Interactor, const bool bForceOverUI );
 
 	/** Hides the radial menu if the specified hand is showing it */
 	void HideRadialMenu( UVREditorInteractor* Interactor );
@@ -233,7 +233,7 @@ protected:
 	/**  Makes a uniform grid widget from the menu information contained in a MultiBox and MultiBoxWidget */
 	void MakeUniformGridMenu( const TSharedRef<FMultiBox>& MultiBox, const TSharedRef<SMultiBoxWidget>& MultiBoxWidget, int32 Columns);
 	/**  Makes a radial box widget from the menu information contained in a MultiBox and MultiBoxWidget */
-	void MakeRadialBoxMenu(const TSharedRef<FMultiBox>& MultiBox, const TSharedRef<SMultiBoxWidget>& MultiBoxWidget);
+	void MakeRadialBoxMenu(const TSharedRef<FMultiBox>& MultiBox, const TSharedRef<SMultiBoxWidget>& MultiBoxWidget, float RadiusRatioOverride);
 
 	/** Adds a hoverable button of a given type to an overlay, using menu data from a BlockWidget */
 	TSharedRef<SWidget> AddHoverableButton(TSharedRef<SWidget>& BlockWidget, FName ButtonType, TSharedRef<SOverlay>& TestOverlay);
@@ -244,6 +244,10 @@ protected:
 	TSharedRef<SWidget> BuildQuickMenuWidget();
 	/** Builds the radial menu Slate widget */
 	TSharedRef<SWidget> BuildRadialMenuWidget();
+	/** Builds the numpad Slate widget */
+	TSharedRef<SWidget> BuildNumPadWidget();
+	/** Swaps the content of the radial menu between the radial menu and the numpad */
+	void SwapRadialMenu();
 
 	/** Called when a laser or simulated mouse hover enters a button */
 	void OnHoverBeginEffect(TSharedRef<SButton> Button);
@@ -281,6 +285,12 @@ protected:
 
 	/** Pointer to the radial menu widget */
 	TSharedPtr<SWidget> RadialWidget;
+
+	/** True if the radial menu was visible when the content was swapped */
+	bool bRadialMenuVisibleAtSwap;
+
+	/** True if the radial menu is currently displaying the numpad */
+	bool bRadialMenuIsNumpad;
 	
 	//
 	// Dragging UI
@@ -331,6 +341,10 @@ protected:
 	UPROPERTY()
 	class USoundCue* ShowUISound;
 
+	/** Button press sound */
+	UPROPERTY()
+	class USoundCue* ButtonPressSound;
+
 	/** If the current dragged dock passed a certain distance if dragged from a hand */
 	bool bDraggedDockFromHandPassedThreshold;
 
diff --git a/Engine/Source/Editor/VREditor/VREditorActions.cpp b/Engine/Source/Editor/VREditor/VREditorActions.cpp
--- a/Engine/Source/Editor/VREditor/VREditorActions.cpp
+++ b/Engine/Source/Editor/VREditor/VREditorActions.cpp
@@ -256,4 +256,49 @@ void FVREditorActionCallbacks::OnSnapActorsToGroundClicked(class UVREditorMode*
 }
 
 
+void FVREditorActionCallbacks::SimulateCharacterEntry(const FString InChar)
+{
+
+	for (int32 CharIndex = 0; CharIndex < InChar.Len(); CharIndex++)
+	{
+		TCHAR CharKey = InChar[CharIndex];
+		const bool bRepeat = false;
+		FCharacterEvent CharacterEvent(CharKey, FModifierKeysState::FModifierKeysState(), 0, bRepeat);
+		FSlateApplication::Get().ProcessKeyCharEvent(CharacterEvent);
+	}
+
+}
+
+void FVREditorActionCallbacks::SimulateKeyDown(FKey Key, bool bRepeat)
+{
+	const uint32* KeyCodePtr;
+	const uint32* CharCodePtr;
+	FInputKeyManager::Get().GetCodesFromKey( Key, KeyCodePtr, CharCodePtr );
+
+	uint32 KeyCode = KeyCodePtr ? *KeyCodePtr : 0;
+	uint32 CharCode = CharCodePtr ? *CharCodePtr : 0;
+
+	FKeyEvent KeyEvent( Key, FModifierKeysState::FModifierKeysState(), 0, bRepeat, KeyCode, CharCode );
+	bool DownResult = FSlateApplication::Get().ProcessKeyDownEvent( KeyEvent );
+
+	if (CharCodePtr)
+	{
+		FCharacterEvent CharacterEvent( CharCode, FModifierKeysState::FModifierKeysState(), 0, bRepeat );
+		FSlateApplication::Get().ProcessKeyCharEvent( CharacterEvent );
+	}
+}
+
+void FVREditorActionCallbacks::SimulateKeyUp(FKey Key)
+{
+	const uint32* KeyCodePtr;
+	const uint32* CharCodePtr;
+	FInputKeyManager::Get().GetCodesFromKey( Key, KeyCodePtr, CharCodePtr );
+
+	uint32 KeyCode = KeyCodePtr ? *KeyCodePtr : 0;
+	uint32 CharCode = CharCodePtr ? *CharCodePtr : 0;
+
+	FKeyEvent KeyEvent( Key, FModifierKeysState::FModifierKeysState(), 0, false, KeyCode, CharCode );
+	FSlateApplication::Get().ProcessKeyUpEvent( KeyEvent );
+}
+
 #undef LOCTEXT_NAMESPACE
diff --git a/Engine/Source/Editor/VREditor/VREditorActions.h b/Engine/Source/Editor/VREditor/VREditorActions.h
--- a/Engine/Source/Editor/VREditor/VREditorActions.h
+++ b/Engine/Source/Editor/VREditor/VREditorActions.h
@@ -126,6 +126,25 @@ public:
 	 */
 	static void OnSnapActorsToGroundClicked(class UVREditorMode* InVRMode);
 
+	/**
+	 * Simulates the user entering characters with a keyboard for data entry
+	 * @param InChar String of characters to enter
+	 */
+	static void SimulateCharacterEntry(const FString InChar);
+
+	/**
+	 * Simulates the user pressing a key down
+	 * @param Key Key to press
+	 * @param bRepeat Whether or not to repeat
+	 */
+	static void SimulateKeyDown( FKey Key, bool bRepeat );
+
+	/**
+	 * Simulates the user releasing a key
+	 * @param Key Key to release
+	 */
+	static void SimulateKeyUp( FKey Key );
+
 }
 ;
 
diff --git a/Engine/Source/Editor/VREditor/VREditorBaseActor.cpp b/Engine/Source/Editor/VREditor/VREditorBaseActor.cpp
--- a/Engine/Source/Editor/VREditor/VREditorBaseActor.cpp
+++ b/Engine/Source/Editor/VREditor/VREditorBaseActor.cpp
@@ -8,6 +8,7 @@
 namespace VREd
 {
 	static FAutoConsoleVariable UIOnHandRotationOffset( TEXT( "VREd.UIOnHandRotationOffset" ), 45.0f, TEXT( "Rotation offset for UI that's docked to your hand, to make it more comfortable to hold" ) );
+	static FAutoConsoleVariable UIOnArmRotationOffset( TEXT( "VREd.UIOnArmRotationOffset" ), 15.0f, TEXT( "Rotation offset for UI that's docked to your arm, so it aligns with the controllers" ) );
 }
 
 AVREditorBaseActor::AVREditorBaseActor() 
@@ -90,10 +91,16 @@ FTransform AVREditorBaseActor::MakeUITransformLockedToHand( UViewportInteractor*
 	const float WorldScaleFactor = VRMode->GetWorldScaleFactor();
 
 	FTransform UIToHandTransform( InLocalRotation, InRelativeOffset * WorldScaleFactor );
+	// If the UI is on the hand
 	if ( !bOnArm )
 	{
 		UIToHandTransform *= FTransform( FRotator( VREd::UIOnHandRotationOffset->GetFloat(), 0.0f, 0.0f ).Quaternion(), FVector::ZeroVector );
 	}
+	// If the UI is on the arm
+	else
+	{
+		UIToHandTransform *= FTransform( FRotator( VREd::UIOnArmRotationOffset->GetFloat(), 0.0f, 0.0f ).Quaternion(), FVector::ZeroVector );
+	}
 
 	const FTransform HandToWorldTransform = Interactor->GetTransform();
 	FTransform UIToWorldTransform = UIToHandTransform * HandToWorldTransform;
diff --git a/Engine/Source/Editor/VREditor/VREditorMotionControllerInteractor.cpp b/Engine/Source/Editor/VREditor/VREditorMotionControllerInteractor.cpp
--- a/Engine/Source/Editor/VREditor/VREditorMotionControllerInteractor.cpp
+++ b/Engine/Source/Editor/VREditor/VREditorMotionControllerInteractor.cpp
@@ -1267,7 +1267,7 @@ void UVREditorMotionControllerInteractor::SetLaserVisuals( const FLinearColor& N
 void UVREditorMotionControllerInteractor::UpdateRadialMenuInput( const float DeltaTime )
 {
 	UVREditorUISystem& UISystem = GetVRMode().GetUISystem();
-
+	const EHMDDeviceType::Type HMDDeviceType = GetVRMode().GetHMDDeviceType();
 	//Update the radial menu
 	EViewportInteractionDraggingMode DraggingMode = GetDraggingMode();
 	if ( !HasUIInFront() &&
@@ -1286,6 +1286,16 @@ void UVREditorMotionControllerInteractor::UpdateRadialMenuInput( const float Del
 			RadialBox->HighlightSlot(TrackpadPosition);
 		}
 	}
+	// If we are not currently touching the Vive touchpad, reset the highlighted button
+	else if (HMDDeviceType == EHMDDeviceType::DT_SteamVR && !bIsTouchingTrackpad)
+	{
+		if (UISystem.IsShowingRadialMenu( this ))
+		{
+			const TSharedPtr<SRadialBox>& RadialBox = StaticCastSharedPtr<SRadialBox>( UISystem.GetRadialWidget() );
+			FVector2D ReturnToCenter = FVector2D::ZeroVector;
+			RadialBox->HighlightSlot( ReturnToCenter );
+		}
+	}
 }
 
 #undef LOCTEXT_NAMESPACE
\ No newline at end of file
diff --git a/Engine/Source/Editor/VREditor/VREditorStyle.cpp b/Engine/Source/Editor/VREditor/VREditorStyle.cpp
--- a/Engine/Source/Editor/VREditor/VREditorStyle.cpp
+++ b/Engine/Source/Editor/VREditor/VREditorStyle.cpp
@@ -44,14 +44,10 @@ const FVector2D Icon40x40(40.0f, 40.0f);
 TSharedRef< FSlateStyleSet > FVREditorStyle::Create()
 {
 	TSharedRef< FSlateStyleSet > Style = MakeShareable(new FSlateStyleSet("VREditorStyle"));
-
-	const FButtonStyle OverrideButton = FEditorStyle::GetWidgetStyle<FButtonStyle>("VRButtonStyle");
-
-	// Use the default menu button style, but set the background to dark grey and use the sound set in the editor.
+	// Use the default menu button style, but set the background to dark grey.
 	const FButtonStyle NormalButton = FEditorStyle::GetWidgetStyle<FButtonStyle>("Menu.Button");
 	Style->Set("VREditorStyle.Button", FButtonStyle(NormalButton)
-		.SetNormal(BOX_BRUSH("Common/RoundedSelection_16x", 4.0f / 16.0f, FLinearColor(0.1f,0.1f,0.1f)))
-		.SetPressedSound(OverrideButton.PressedSlateSound));
+		.SetNormal(BOX_BRUSH("Common/RoundedSelection_16x", 4.0f / 16.0f, FLinearColor(0.1f,0.1f,0.1f))));
 
 	const FTextBlockStyle NormalText = FEditorStyle::GetWidgetStyle<FTextBlockStyle>("NormalText");
 	Style->Set("VREditorStyle.Label", FTextBlockStyle(NormalText)
@@ -67,6 +63,17 @@ TSharedRef< FSlateStyleSet > FVREditorStyle::Create()
 	Style->Set("VREditorStyle.Heading", FTextBlockStyle(NormalText)
 			.SetFont(HeadlineFont)
 			.SetColorAndOpacity(FLinearColor::White));
+
+	// Headings will have a font outline
+	FFontOutlineSettings HelperOutline;
+	HelperOutline.OutlineColor = FLinearColor( 0.2f, 0.2f, 0.2f, 0.5f );
+	HelperOutline.OutlineSize = 3.0f;
+	FSlateFontInfo HelperFont = TTF_CORE_FONT( "Fonts/Roboto-Regular", 24 );
+	HelperFont.OutlineSettings = HelperOutline;
+
+	Style->Set( "VREditorStyle.HelperText", FTextBlockStyle( NormalText )
+		.SetFont( HelperFont )
+		.SetColorAndOpacity( FLinearColor::White ) );
 	
 	const FCheckBoxStyle CheckboxStyle = FEditorStyle::GetWidgetStyle<FCheckBoxStyle>("Menu.Check");
 	Style->Set("VREditorStyle.Check", FCheckBoxStyle(CheckboxStyle));
