Fixing case where radial menu could get stuck on after number pad is dismissed

diff --git a/Engine/Source/Editor/VREditor/UI/VREditorUISystem.cpp b/Engine/Source/Editor/VREditor/UI/VREditorUISystem.cpp
--- a/Engine/Source/Editor/VREditor/UI/VREditorUISystem.cpp
+++ b/Engine/Source/Editor/VREditor/UI/VREditorUISystem.cpp
@@ -240,15 +240,19 @@ void UVREditorUISystem::OnVRAction( FEditorViewportClient& ViewportClient, UView
 					}
 					else
 					{	
-						// If you are hiding the number pad, also toggle back to the radial menu
-						if(bRadialMenuIsNumpad)
+						// If the numpad is showing, either swap to the radial menu or dismiss the menu entirely
+						if (bRadialMenuIsNumpad)
 						{
 							SwapRadialMenu();
+							if (!bRadialMenuVisibleAtSwap)
+							{
+								HideRadialMenu(VREditorInteractor);
+							}
 						}
-						// If we are not swapping back to the radial menu from the number pad, hide the menu
-						if (!bRadialMenuVisibleAtSwap)
+						// Always dismiss the regular radial menu
+						else
 						{
-							HideRadialMenu( VREditorInteractor );
+							HideRadialMenu(VREditorInteractor);
 						}
 					}
 					bWasHandled = true;
@@ -401,6 +405,7 @@ void UVREditorUISystem::OnVRAction( FEditorViewportClient& ViewportClient, UView
 													}
 													else
 													{
+														bRadialMenuVisibleAtSwap = false;
 														// Force the radial menu to spawn even if the laser is over UI
 														const bool bForceOverUI = true;
 														TryToSpawnRadialMenu( VREditorInteractor, bForceOverUI );
