Adjusting quick menu positioning for UI clipping

diff --git a/Engine/Source/Editor/VREditor/UI/VREditorUISystem.cpp b/Engine/Source/Editor/VREditor/UI/VREditorUISystem.cpp
--- a/Engine/Source/Editor/VREditor/UI/VREditorUISystem.cpp
+++ b/Engine/Source/Editor/VREditor/UI/VREditorUISystem.cpp
@@ -838,7 +838,7 @@ void UVREditorUISystem::CreateUIs()
 			QuickMenuUI = GetOwner().SpawnTransientSceneActor< AVREditorFloatingUI >(TEXT("QuickMenu"), bWithSceneComponent);
 			QuickMenuUI->SetSlateWidget(*this, BuildQuickMenuWidget(), Resolution, 40.0f, AVREditorFloatingUI::EDockedTo::Nothing);
 			QuickMenuUI->ShowUI(false);
-			QuickMenuUI->SetRelativeOffset(FVector(-10.0f, 0.0f, 5.0f));
+			QuickMenuUI->SetRelativeOffset(FVector(-12.0f, 0.0f, 5.0f));
 			QuickMenuUI->GetWidgetComponent()->SetOpacityFromTexture(1.0f);
 			QuickMenuUI->GetWidgetComponent()->SetBackgroundColor(FLinearColor::Transparent);
 			QuickMenuUI->GetWidgetComponent()->SetBlendMode(EWidgetBlendMode::Masked);
