Fixing near clip plane issues when the user is small

diff --git a/Engine/Source/Editor/VREditor/Public/VREditorMode.h b/Engine/Source/Editor/VREditor/Public/VREditorMode.h
--- a/Engine/Source/Editor/VREditor/Public/VREditorMode.h
+++ b/Engine/Source/Editor/VREditor/Public/VREditorMode.h
@@ -429,6 +429,9 @@ public:
 	// Gets the color
 	FLinearColor GetColor( const EColors Color ) const;
 
+	// Get the default near clipping plane for VR editing
+	float GetDefaultVRNearClipPlane() const;
+
 private:
 
 	// All the colors for this mode
diff --git a/Engine/Source/Editor/VREditor/VREditorMode.cpp b/Engine/Source/Editor/VREditor/VREditorMode.cpp
--- a/Engine/Source/Editor/VREditor/VREditorMode.cpp
+++ b/Engine/Source/Editor/VREditor/VREditorMode.cpp
@@ -39,6 +39,7 @@ namespace VREd
 {
 	static FAutoConsoleVariable UseMouseAsHandInForcedVRMode( TEXT( "VREd.UseMouseAsHandInForcedVRMode" ), 1, TEXT( "When in forced VR mode, enabling this setting uses the mouse cursor as a virtual hand instead of motion controllers" ) );
 	static FAutoConsoleVariable ForceOculusMirrorMode( TEXT( "VREd.ForceOculusMirrorMode" ), 3, TEXT( "Which Oculus display mirroring mode to use (see MirrorWindowModeType in OculusRiftHMD.h)" ) );
+	static FAutoConsoleVariable VRNearClipPlane(TEXT("VREd.VRNearClipPlane"), 1.0f, TEXT("The near clip plane to use for VR"));
 }
 
 // @todo vreditor: Hacky that we have to import these this way. (Plugin has them in a .cpp, not exported)
@@ -275,9 +276,8 @@ void UVREditorMode::Enter(const bool bReenteringVREditing)
 			VRViewportClient.bAlwaysShowModeWidgetAfterSelectionChanges = false;
 
 			// Force tiny near clip plane distance, because user can scale themselves to be very small.
-			// @todo vreditor: Make this automatically change based on WorldToMetersScale?
 			SavedEditorState.NearClipPlane = GNearClippingPlane;
-			GNearClippingPlane = 1.0f;	// Normally defaults to 10cm (NOTE: If we go too small, skyboxes become affected)
+			GNearClippingPlane = GetDefaultVRNearClipPlane();	
 
 			SavedEditorState.bOnScreenMessages = GAreScreenMessagesEnabled;
 			GAreScreenMessagesEnabled = false;
@@ -903,6 +903,11 @@ FLinearColor UVREditorMode::GetColor( const EColors Color ) const
 	return Colors[ (int32)Color ];
 }
 
+float UVREditorMode::GetDefaultVRNearClipPlane() const 
+{
+	return VREd::VRNearClipPlane->GetFloat();
+}
+
 bool UVREditorMode::IsHandAimingTowardsCapsule( UViewportInteractor* Interactor, const FTransform& CapsuleTransform, FVector CapsuleStart, const FVector CapsuleEnd, const float CapsuleRadius, const float MinDistanceToCapsule, const FVector CapsuleFrontDirection, const float MinDotForAimingAtCapsule ) const
 {
 	bool bIsAimingTowards = false;
diff --git a/Engine/Source/Editor/VREditor/VREditorWorldInteraction.cpp b/Engine/Source/Editor/VREditor/VREditorWorldInteraction.cpp
--- a/Engine/Source/Editor/VREditor/VREditorWorldInteraction.cpp
+++ b/Engine/Source/Editor/VREditor/VREditorWorldInteraction.cpp
@@ -57,12 +57,14 @@ void UVREditorWorldInteraction::Init( UVREditorMode* InOwner, UViewportWorldInte
 	FEditorDelegates::OnAssetDragStarted.AddUObject( this, &UVREditorWorldInteraction::OnAssetDragStartedFromContentBrowser );
 
 	ViewportWorldInteraction->OnStopDragging().AddUObject( this, &UVREditorWorldInteraction::StopDragging );
+	ViewportWorldInteraction->OnWorldScaleChanged().AddUObject( this, &UVREditorWorldInteraction::UpdateNearClipPlaneOnScaleChange );
 }
 
 void UVREditorWorldInteraction::Shutdown()
 {
 	FEditorDelegates::OnAssetDragStarted.RemoveAll( this );
 	ViewportWorldInteraction->OnStopDragging().RemoveAll( this );
+	ViewportWorldInteraction->OnWorldScaleChanged().RemoveAll( this );
 
 	PlacingMaterialOrTextureAsset = nullptr;
 	FloatingUIAssetDraggedFrom = nullptr;
@@ -112,6 +114,13 @@ void UVREditorWorldInteraction::StopDragging( UViewportInteractor* Interactor )
 	}
 }
 
+void UVREditorWorldInteraction::UpdateNearClipPlaneOnScaleChange(const float NewWorldToMetersScale)
+{
+	// Adjust the clipping plane for the user's scale, but don't let it be larger than the engine default
+	const float DefaultWorldToMetersScale = Owner->GetSavedEditorState().WorldToMetersScale;
+	GNearClippingPlane = FMath::Min((Owner->GetDefaultVRNearClipPlane()) * (NewWorldToMetersScale / DefaultWorldToMetersScale), Owner->GetSavedEditorState().NearClipPlane);
+}
+
 void UVREditorWorldInteraction::StartDraggingMaterialOrTexture( UViewportInteractor* Interactor, const FViewportActionKeyInput& Action, const FVector HitLocation, UObject* MaterialOrTextureAsset )
 {
 	check( MaterialOrTextureAsset != nullptr );
diff --git a/Engine/Source/Editor/VREditor/VREditorWorldInteraction.h b/Engine/Source/Editor/VREditor/VREditorWorldInteraction.h
--- a/Engine/Source/Editor/VREditor/VREditorWorldInteraction.h
+++ b/Engine/Source/Editor/VREditor/VREditorWorldInteraction.h
@@ -34,6 +34,9 @@ protected:
 	/** When an interactor stops dragging */
 	void StopDragging( UViewportInteractor* Interactor );
 
+	/** When the world scale changes, update the near clip plane */
+	void UpdateNearClipPlaneOnScaleChange(const float NewWorldToMetersScale);
+
 	/** Starts dragging a material, allowing the user to drop it on an object in the scene to place it */
 	void StartDraggingMaterialOrTexture( UViewportInteractor* Interactor, const FViewportActionKeyInput& Action, const FVector HitLocation, UObject* MaterialOrTextureAsset );
 
diff --git a/Engine/Source/Editor/ViewportInteraction/Public/ViewportWorldInteraction.h b/Engine/Source/Editor/ViewportInteraction/Public/ViewportWorldInteraction.h
--- a/Engine/Source/Editor/ViewportInteraction/Public/ViewportWorldInteraction.h
+++ b/Engine/Source/Editor/ViewportInteraction/Public/ViewportWorldInteraction.h
@@ -161,6 +161,8 @@ public:
 
 	virtual FOnStopDragging& OnStopDragging() override { return OnStopDraggingEvent; };
 
+	virtual FOnWorldScaleChanged& OnWorldScaleChanged() override { return OnWorldScaleChangedEvent; };
+
 	/** Starts dragging selected actors around, changing selection if needed.  Called when clicking in the world, or when placing new objects */
 	void StartDraggingActors( UViewportInteractor* Interactor, const FViewportActionKeyInput& Action, UActorComponent* ClickedComponent, const FVector& HitLocation, const bool bIsPlacingActors );
 	
@@ -530,9 +532,12 @@ private:
 	/** Event that is fired when new axis input is received by the InputProcessor */
 	FOnHandleInputAxis OnAxisInputEvent;
 
-	/** Event this is fired when a interactor stopped dragging */
+	/** Event that is fired when a interactor stopped dragging */
 	FOnStopDragging OnStopDraggingEvent;
 
+	/** Event that is fired when the world scale changes */
+	FOnWorldScaleChanged OnWorldScaleChangedEvent;
+
 	/** If the world interaction is active and running the tick function */
 	bool bActive;
 };
\ No newline at end of file
diff --git a/Engine/Source/Editor/ViewportInteraction/Public/ViewportWorldInteractionInterface.h b/Engine/Source/Editor/ViewportInteraction/Public/ViewportWorldInteractionInterface.h
--- a/Engine/Source/Editor/ViewportInteraction/Public/ViewportWorldInteractionInterface.h
+++ b/Engine/Source/Editor/ViewportInteraction/Public/ViewportWorldInteractionInterface.h
@@ -50,4 +50,8 @@ public:
 	/** Gets the event for when an interactor stops dragging */
 	DECLARE_EVENT_OneParam( IViewportWorldInteractionInterface, FOnStopDragging, class UViewportInteractor* /** Interactor */ );
 	virtual FOnStopDragging& OnStopDragging() = 0;
+
+	/** Gets the event for when the world scale changes */
+	DECLARE_EVENT_OneParam(IViewportWorldInteractionInterface, FOnWorldScaleChanged, const float /* NewWorldToMetersScale */);
+	virtual FOnWorldScaleChanged& OnWorldScaleChanged() = 0;
 };
\ No newline at end of file
diff --git a/Engine/Source/Editor/ViewportInteraction/ViewportWorldInteraction.cpp b/Engine/Source/Editor/ViewportInteraction/ViewportWorldInteraction.cpp
--- a/Engine/Source/Editor/ViewportInteraction/ViewportWorldInteraction.cpp
+++ b/Engine/Source/Editor/ViewportInteraction/ViewportWorldInteraction.cpp
@@ -2794,6 +2794,7 @@ void UViewportWorldInteraction::SetWorldToMetersScale( const float NewWorldToMet
 
 	ENGINE_API extern float GNewWorldToMetersScale;
 	GNewWorldToMetersScale = NewWorldToMetersScale;
+	OnWorldScaleChangedEvent.Broadcast(NewWorldToMetersScale);
 }
 
 UViewportInteractor* UViewportWorldInteraction::GetOtherInteractorIntertiaContribute( UViewportInteractor* Interactor )
