Adding other (not-yet-used) menu styles to VREditorStyle so Slate can find them when using VR Editor Style

diff --git a/Engine/Source/Editor/VREditor/VREditorStyle.cpp b/Engine/Source/Editor/VREditor/VREditorStyle.cpp
--- a/Engine/Source/Editor/VREditor/VREditorStyle.cpp
+++ b/Engine/Source/Editor/VREditor/VREditorStyle.cpp
@@ -43,7 +43,7 @@ const FVector2D Icon40x40(40.0f, 40.0f);
 
 TSharedRef< FSlateStyleSet > FVREditorStyle::Create()
 {
-	TSharedRef< FSlateStyleSet > Style = MakeShareable(new FSlateStyleSet("VREditorStyle"));
+	TSharedRef< FSlateStyleSet > Style = MakeShareable(new FSlateStyleSet(FVREditorStyle::GetStyleSetName()));
 	// Use the default menu button style, but set the background to dark grey.
 	const FButtonStyle NormalButton = FEditorStyle::GetWidgetStyle<FButtonStyle>("Menu.Button");
 	Style->Set("VREditorStyle.Button", FButtonStyle(NormalButton)
@@ -75,8 +75,20 @@ TSharedRef< FSlateStyleSet > FVREditorStyle::Create()
 		.SetFont( HelperFont )
 		.SetColorAndOpacity( FLinearColor::White ) );
 	
-	const FCheckBoxStyle CheckboxStyle = FEditorStyle::GetWidgetStyle<FCheckBoxStyle>("Menu.Check");
-	Style->Set("VREditorStyle.Check", FCheckBoxStyle(CheckboxStyle));
+	const FCheckBoxStyle CheckStyle = FEditorStyle::GetWidgetStyle<FCheckBoxStyle>("Menu.Check");
+	Style->Set("VREditorStyle.Check", FCheckBoxStyle(CheckStyle));
+
+	const FCheckBoxStyle CheckboxStyle = FEditorStyle::GetWidgetStyle<FCheckBoxStyle>("Menu.CheckBox");
+	Style->Set("VREditorStyle.CheckBox", FCheckBoxStyle(CheckboxStyle));
+
+	const FCheckBoxStyle RadioButtonStyle = FEditorStyle::GetWidgetStyle<FCheckBoxStyle>("Menu.RadioButton");
+	Style->Set("VREditorStyle.RadioButton", FCheckBoxStyle(RadioButtonStyle));
+
+	const FTextBlockStyle KeybindingStyle = FEditorStyle::GetWidgetStyle<FTextBlockStyle>("Menu.Keybinding");
+	Style->Set("VREditorStyle.Keybinding", FTextBlockStyle(KeybindingStyle));
+
+	const FTextBlockStyle LabelStyle = FEditorStyle::GetWidgetStyle<FTextBlockStyle>("Menu.Label");
+	Style->Set("VREditorStyle.Label", FTextBlockStyle(LabelStyle));
 
 	return Style;
 }
