VR Editor: Fixed difficulty with double clicking with Oculus Touch (UE-37521) (integrated from CL 3227095)
- Recent Oculus Runtime changes altered the input curve on the trigger which affected our hard-coded dead zone
- This prevented us from being able to detect a double click with the trigger reliably
- We now have different thresholds for Oculus Rift versus HTC Vive

diff --git a/Engine/Source/Editor/VREditor/VREditorMotionControllerInteractor.cpp b/Engine/Source/Editor/VREditor/VREditorMotionControllerInteractor.cpp
--- a/Engine/Source/Editor/VREditor/VREditorMotionControllerInteractor.cpp
+++ b/Engine/Source/Editor/VREditor/VREditorMotionControllerInteractor.cpp
@@ -32,8 +32,10 @@ namespace VREd
 	static FAutoConsoleVariable LaserPointerLightRadius( TEXT( "VREd.LaserPointLightRadius" ), 20.0f, TEXT( "How big our hover light is" ) );
 
 	//Trigger
-	static FAutoConsoleVariable TriggerLightlyPressedThreshold( TEXT( "VI.TriggerLightlyPressedThreshold" ), 0.03f, TEXT( "Minimum trigger threshold before we consider the trigger at least 'lightly pressed'" ) );
-	static FAutoConsoleVariable TriggerDeadZone( TEXT( "VI.TriggerDeadZone" ), 0.01f, TEXT( "Trigger dead zone.  The trigger must be fully released before we'll trigger a new 'light press'" ) );
+	static FAutoConsoleVariable TriggerLightlyPressedThreshold_Vive( TEXT( "VI.TriggerLightlyPressedThreshold_Vive" ), 0.03f, TEXT( "Minimum trigger threshold before we consider the trigger at least 'lightly pressed'" ) );
+	static FAutoConsoleVariable TriggerDeadZone_Vive( TEXT( "VI.TriggerDeadZone_Vive" ), 0.01f, TEXT( "Trigger dead zone.  The trigger must be fully released before we'll trigger a new 'light press'" ) );
+	static FAutoConsoleVariable TriggerLightlyPressedThreshold_Rift( TEXT( "VI.TriggerLightlyPressedThreshold_Rift" ), 0.2f, TEXT( "Minimum trigger threshold before we consider the trigger at least 'lightly pressed'" ) );
+	static FAutoConsoleVariable TriggerDeadZone_Rift( TEXT( "VI.TriggerDeadZone_Rift" ), 0.15f, TEXT( "Trigger dead zone.  The trigger must be fully released before we'll trigger a new 'light press'" ) );
 	static FAutoConsoleVariable TriggerFullyPressedThreshold( TEXT( "VI.TriggerFullyPressedThreshold" ), 0.95f, TEXT( "Minimum trigger threshold before we consider the trigger 'fully pressed'" ) );
 	static FAutoConsoleVariable TriggerFullyPressedReleaseThreshold( TEXT( "VI.TriggerFullyPressedReleaseThreshold" ), 0.8f, TEXT( "After fully pressing the trigger, if the axis falls below this threshold we no longer consider it fully pressed" ) );
 
@@ -651,6 +653,9 @@ void UVREditorMotionControllerInteractor::HandleInputAxis( FViewportActionKeyInp
 {
 	if ( Action.ActionType == TriggerAxis )
 	{
+		const float TriggerLightlyPressedThreshold = ( GetHMDDeviceType() == EHMDDeviceType::DT_OculusRift ) ? VREd::TriggerLightlyPressedThreshold_Rift->GetFloat() : VREd::TriggerLightlyPressedThreshold_Vive->GetFloat();
+		const float TriggerDeadZone = ( GetHMDDeviceType() == EHMDDeviceType::DT_OculusRift ) ? VREd::TriggerDeadZone_Rift->GetFloat() : VREd::TriggerDeadZone_Vive->GetFloat();
+
 		// Synthesize "lightly pressed" events for the trigger
 		{
 			// Store latest trigger value amount
@@ -662,7 +667,7 @@ void UVREditorMotionControllerInteractor::HandleInputAxis( FViewportActionKeyInp
 			if ( !bIsFullPressAlreadyCapturing &&		// Don't fire if we're already capturing input for a full press
 				!bIsTriggerAtLeastLightlyPressed &&	// Don't fire unless we are already pressed
 				bHasTriggerBeenReleasedSinceLastPress &&	// Only if we've been fully released since the last time we fired
-				Delta >= VREd::TriggerLightlyPressedThreshold->GetFloat() )
+				Delta >= TriggerLightlyPressedThreshold )
 			{
 				bIsTriggerAtLeastLightlyPressed = true;
 				SetAllowTriggerLightPressLocking( true );
@@ -674,7 +679,7 @@ void UVREditorMotionControllerInteractor::HandleInputAxis( FViewportActionKeyInp
 				const EInputEvent InputEvent = IE_Pressed;
 				const bool bWasLightPressHandled = UViewportInteractor::HandleInputKey( ControllerHandSide == EControllerHand::Left ? MotionController_Left_LightlyPressedTriggerAxis : MotionController_Right_LightlyPressedTriggerAxis, InputEvent );
 			}
-			else if ( bIsTriggerAtLeastLightlyPressed && Delta < VREd::TriggerLightlyPressedThreshold->GetFloat() )
+			else if ( bIsTriggerAtLeastLightlyPressed && Delta < TriggerLightlyPressedThreshold )
 			{
 				bIsTriggerAtLeastLightlyPressed = false;
 
@@ -684,7 +689,7 @@ void UVREditorMotionControllerInteractor::HandleInputAxis( FViewportActionKeyInp
 			}
 		}
 
-		if ( Delta < VREd::TriggerDeadZone->GetFloat() )
+		if ( !bHasTriggerBeenReleasedSinceLastPress && Delta < TriggerDeadZone )
 		{
 			bHasTriggerBeenReleasedSinceLastPress = true;
 		}
