Updating font location in VR Editor Slate style

diff --git a/Engine/Source/Editor/VREditor/VREditorStyle.cpp b/Engine/Source/Editor/VREditor/VREditorStyle.cpp
--- a/Engine/Source/Editor/VREditor/VREditorStyle.cpp
+++ b/Engine/Source/Editor/VREditor/VREditorStyle.cpp
@@ -52,13 +52,13 @@ TSharedRef< FSlateStyleSet > FVREditorStyle::Create()
 
 	const FTextBlockStyle NormalText = FEditorStyle::GetWidgetStyle<FTextBlockStyle>("NormalText");
 	Style->Set("VREditorStyle.Label", FTextBlockStyle(NormalText)
-		.SetFont(TTF_CORE_FONT("Fonts/Roboto-Regular", 9)));
+		.SetFont(TTF_CORE_FONT("Slate/Fonts/Roboto-Regular", 9)));
 
 	// Headings will have a font outline
 	FFontOutlineSettings HeadingOutline;
 	HeadingOutline.OutlineColor = FLinearColor(0.2f, 0.2f, 0.2f, 0.5f);
 	HeadingOutline.OutlineSize = 2.0f;
-	FSlateFontInfo HeadlineFont = TTF_CORE_FONT("Fonts/Roboto-Regular", 18);
+	FSlateFontInfo HeadlineFont = TTF_CORE_FONT("Slate/Fonts/Roboto-Regular", 18);
 	HeadlineFont.OutlineSettings = HeadingOutline;
 
 	Style->Set("VREditorStyle.Heading", FTextBlockStyle(NormalText)
@@ -69,7 +69,7 @@ TSharedRef< FSlateStyleSet > FVREditorStyle::Create()
 	FFontOutlineSettings HelperOutline;
 	HelperOutline.OutlineColor = FLinearColor( 0.2f, 0.2f, 0.2f, 0.5f );
 	HelperOutline.OutlineSize = 3.0f;
-	FSlateFontInfo HelperFont = TTF_CORE_FONT( "Fonts/Roboto-Regular", 24 );
+	FSlateFontInfo HelperFont = TTF_CORE_FONT( "Slate/Fonts/Roboto-Regular", 24 );
 	HelperFont.OutlineSettings = HelperOutline;
 
 	Style->Set( "VREditorStyle.HelperText", FTextBlockStyle( NormalText )
