Updating copyright notice on new VR Editor files

diff --git a/Engine/Source/Editor/VREditor/VREditorActions.h b/Engine/Source/Editor/VREditor/VREditorActions.h
--- a/Engine/Source/Editor/VREditor/VREditorActions.h
+++ b/Engine/Source/Editor/VREditor/VREditorActions.h
@@ -1,4 +1,4 @@
-// Copyright 1998-2016 Epic Games, Inc. All Rights Reserved.
+// Copyright 1998-2017 Epic Games, Inc. All Rights Reserved.
 #pragma once
 
 #include "VREditorUISystem.h"
diff --git a/Engine/Source/Editor/VREditor/VREditorStyle.h b/Engine/Source/Editor/VREditor/VREditorStyle.h
--- a/Engine/Source/Editor/VREditor/VREditorStyle.h
+++ b/Engine/Source/Editor/VREditor/VREditorStyle.h
@@ -1,4 +1,4 @@
-// Copyright 1998-2016 Epic Games, Inc. All Rights Reserved.
+// Copyright 1998-2017 Epic Games, Inc. All Rights Reserved.
 
 #pragma once
 
