Updating VR  editor interaction code to not allow selecting primitive components that have been set as non-selectable

diff --git a/Engine/Source/Editor/VREditor/VREditorWorldInteraction.cpp b/Engine/Source/Editor/VREditor/VREditorWorldInteraction.cpp
--- a/Engine/Source/Editor/VREditor/VREditorWorldInteraction.cpp
+++ b/Engine/Source/Editor/VREditor/VREditorWorldInteraction.cpp
@@ -90,12 +90,17 @@ bool UVREditorWorldInteraction::IsInteractableComponent( const UActorComponent*
 	// Don't bother if the base class already found out that this is a custom interactable
 	if ( bResult )
 	{
-		// Don't interact with frozen actors or our UI
+		// Don't interact with frozen actors, our UI, or primitive components that have been set as not selectable
 		if ( Component != nullptr )
 		{
 			static const bool bIsVREditorDemo = FParse::Param( FCommandLine::Get(), TEXT( "VREditorDemo" ) );	// @todo vreditor: Remove this when no longer needed (console variable, too!)
 			const bool bIsFrozen = bIsVREditorDemo && Component->GetOwner()->GetActorLabel().StartsWith( TEXT( "Frozen_" ) );
-			bResult = !bIsFrozen && !Owner->GetUISystem().IsWidgetAnEditorUIWidget( Component ); // Don't allow user to move around our UI widgets //@todo VREditor: This wouldn't be necessary if the UI used the interactable interface
+			const UPrimitiveComponent* ComponentAsPrimitive = Cast<UPrimitiveComponent>(Component);
+			// Don't allow user to move around our UI widgets 
+			//@todo VREditor: This wouldn't be necessary if the UI used the interactable interface
+			bResult = !bIsFrozen && 
+				!Owner->GetUISystem().IsWidgetAnEditorUIWidget( Component )
+				&& (ComponentAsPrimitive->bSelectable == true); 
 		}
 	}
 	
