Copying //UE4/Dev-VREditor to //UE4/Dev-Main (Source: //UE4/Dev-VREditor @ 3236454)

==========================
MAJOR FEATURES + CHANGES
==========================

Change 3150239 on 2016/10/04 by [REDACTED_PERSON]

	Fix incorrect spelling on target name for "Editor, Tools & Monolithics" in Dev-VREditor CIS.

Change 3174915 on 2016/10/26 by [REDACTED_PERSON]

	New Slate-based dynamic menu, with delegates for overriding widget creating. New behavior is that you must click the thumbstick or touchpad to summon/dismiss the radial menu, select options with a direction on the thumbstick/touchpad, and then select with the trigger. Quick menu functionality is mostly the same, just the laser is now hidden.

Change 3174961 on 2016/10/26 by [REDACTED_PERSON]

	Resaving new VRButtonStyle.uasset with engine version

Change 3175228 on 2016/10/26 by [REDACTED_PERSON]

	Fix for function name overlap in SRadialPanel.cpp

Change 3175375 on 2016/10/26 by [REDACTED_PERSON]

	Reordering functions for Mac compilation

Change 3196880 on 2016/11/14 by [REDACTED_PERSON]

	Adding editor-only flag and function on SkeletalMeshComponent to allow playback of animations on skeletal meshes placed in the level.

Change 3197790 on 2016/11/14 by [REDACTED_PERSON]

	Fixing typo in #if WITH_EDITOR statement

Change 3199071 on 2016/11/15 by [REDACTED_PERSON]

	Fix for User Struct Editor details panels not scrolling. Replaced Vertical Box with Splitter

Change 3207474 on 2016/11/22 by [REDACTED_PERSON]

	Number pad for variable entry in the VR Editor. Number pad appears when clicking on an editable text field, can be dismissed with center button. Trigger for key entry, grip for confirming entry, modifier key for backspace. Also includes menu positioning bug fixes and button pressed sounds.

Change 3207528 on 2016/11/22 by [REDACTED_PERSON]

	Fixing case where radial menu could get stuck on after number pad is dismissed

Change 3213143 on 2016/11/29 by [REDACTED_PERSON]

	VR Editor: Allow objects to snap to ground while translating them (press touchpad or analog stick when dragging object) #jira UETOOL-670

Change 3221501 on 2016/12/05 by [REDACTED_PERSON]

	Adjusting quick menu positioning for UI clipping

Change 3221504 on 2016/12/05 by [REDACTED_PERSON]

	Adding "VR Editor" category to VR Editor tutorial widgets

Change 3224946 on 2016/12/07 by [REDACTED_PERSON]

	Foliage mode now ignores trigger input if the radial menu is present on the interactor. Fixes undo/redo not working on foliage mode.

Change 3224988 on 2016/12/07 by [REDACTED_PERSON]

	Adding analytics event for Vive device initialization

Change 3225656 on 2016/12/07 by [REDACTED_PERSON]

	Fixing near clip plane issues when the user is small

Change 3225823 on 2016/12/07 by [REDACTED_PERSON]

	Adding other (not-yet-used) menu styles to VREditorStyle so Slate can find them when using VR Editor Style

Change 3227098 on 2016/12/08 by [REDACTED_PERSON]

	VR Editor: Fixed difficulty with double clicking with Oculus Touch (UE-37521) (integrated from CL 3227095)
	- Recent Oculus Runtime changes altered the input curve on the trigger which affected our hard-coded dead zone
	- This prevented us from being able to detect a double click with the trigger reliably
	- We now have different thresholds for Oculus Rift versus HTC Vive

Change 3233389 on 2016/12/13 by [REDACTED_PERSON]

	Physics: Fixed line traces not working properly in editor worlds when physics substepping was enabled (UE-36408)
	- Substepping relies on interpolating transforms over frames, but only game worlds will be ticked, so we now disallow this feature in non-game worlds.

Change 3235252 on 2016/12/14 by [REDACTED_PERSON]

	Updating font location in VR Editor Slate style

Change 3236411 on 2016/12/15 by [REDACTED_PERSON]

	Fixing include ordering in MultiboxBuilder.h

Change 3236424 on 2016/12/15 by [REDACTED_PERSON]

	Updating copyright notice on new VR Editor files

Change 3236439 on 2016/12/15 by [REDACTED_PERSON]

	Fixed IWYU regression with PhysPublic.h (introduced in 3233389)

diff --git a/Engine/Plugins/Runtime/Steam/SteamVR/Source/SteamVR/Private/SteamVRHMD.cpp b/Engine/Plugins/Runtime/Steam/SteamVR/Source/SteamVR/Private/SteamVRHMD.cpp
--- a/Engine/Plugins/Runtime/Steam/SteamVR/Source/SteamVR/Private/SteamVRHMD.cpp
+++ b/Engine/Plugins/Runtime/Steam/SteamVR/Source/SteamVR/Private/SteamVRHMD.cpp
@@ -16,6 +16,9 @@
 
 #include "SteamVRMeshAssets.h"
 
+#include "EngineAnalytics.h"
+#include "Runtime/Analytics/Analytics/Public/Interfaces/IAnalyticsProvider.h"
+
 #if WITH_EDITOR
 #include "Editor/UnrealEd/Classes/Editor/EditorEngine.h"
 #endif
@@ -525,6 +528,41 @@ EHMDTrackingOrigin::Type FSteamVRHMD::GetTrackingOrigin()
 	return EHMDTrackingOrigin::Floor;
 }
 
+
+void FSteamVRHMD::RecordAnalytics()
+{
+	if (FEngineAnalytics::IsAvailable())
+	{
+		// prepare and send analytics data
+		TArray<FAnalyticsEventAttribute> EventAttributes;
+
+		IHeadMountedDisplay::MonitorInfo MonitorInfo;
+		GetHMDMonitorInfo(MonitorInfo);
+
+		uint64 MonitorId = MonitorInfo.MonitorId;
+
+		char Buf[128];
+		vr::TrackedPropertyError Error;
+		FString DeviceName = "SteamVR - Default Device Name";
+		VRSystem->GetStringTrackedDeviceProperty(vr::k_unTrackedDeviceIndex_Hmd, vr::Prop_ModelNumber_String, Buf, sizeof(Buf), &Error);
+		if (Error == vr::TrackedProp_Success)
+		{
+			DeviceName = FString(UTF8_TO_TCHAR(Buf));
+		}
+		EventAttributes.Add(FAnalyticsEventAttribute(TEXT("DeviceName"), DeviceName));
+		EventAttributes.Add(FAnalyticsEventAttribute(TEXT("DisplayDeviceName"), *MonitorInfo.MonitorName));
+		EventAttributes.Add(FAnalyticsEventAttribute(TEXT("DisplayId"), MonitorId));
+		FString MonResolution(FString::Printf(TEXT("(%d, %d)"), MonitorInfo.ResolutionX, MonitorInfo.ResolutionY));
+		EventAttributes.Add(FAnalyticsEventAttribute(TEXT("Resolution"), MonResolution));
+
+		EventAttributes.Add(FAnalyticsEventAttribute(TEXT("InterpupillaryDistance"), GetInterpupillaryDistance()));
+
+
+		FString OutStr(TEXT("Editor.VR.DeviceInitialised"));
+		FEngineAnalytics::GetProvider().RecordEvent(OutStr, EventAttributes);
+	}
+}
+
 void FSteamVRHMD::SetUnrealControllerIdAndHandToDeviceIdMap(int32 InUnrealControllerIdAndHandToDeviceIdMap[ MAX_STEAMVR_CONTROLLER_PAIRS ][ 2 ] )
 {
 	for( int32 UnrealControllerIndex = 0; UnrealControllerIndex < MAX_STEAMVR_CONTROLLER_PAIRS; ++UnrealControllerIndex )
diff --git a/Engine/Plugins/Runtime/Steam/SteamVR/Source/SteamVR/Private/SteamVRHMD.h b/Engine/Plugins/Runtime/Steam/SteamVR/Source/SteamVR/Private/SteamVRHMD.h
--- a/Engine/Plugins/Runtime/Steam/SteamVR/Source/SteamVR/Private/SteamVRHMD.h
+++ b/Engine/Plugins/Runtime/Steam/SteamVR/Source/SteamVR/Private/SteamVRHMD.h
@@ -108,6 +108,8 @@ public:
 	virtual void SetTrackingOrigin(EHMDTrackingOrigin::Type NewOrigin) override;
 	virtual EHMDTrackingOrigin::Type GetTrackingOrigin() override;
 
+	virtual void RecordAnalytics() override;
+
 	/** IStereoRendering interface */
 	virtual bool IsStereoEnabled() const override;
 	virtual bool EnableStereo(bool stereo = true) override;
diff --git a/Engine/Source/Editor/FoliageEdit/Private/FoliageEdMode.cpp b/Engine/Source/Editor/FoliageEdit/Private/FoliageEdMode.cpp
--- a/Engine/Source/Editor/FoliageEdit/Private/FoliageEdMode.cpp
+++ b/Engine/Source/Editor/FoliageEdit/Private/FoliageEdMode.cpp
@@ -491,9 +491,9 @@ void FEdModeFoliage::OnVRAction(class FEditorViewportClient& ViewportClient, UVi
 		// Only allow light press
 		Interactor->SetAllowTriggerFullPress(false);
 		// Consume light press
-		if (Action.ActionType == ViewportWorldActionTypes::SelectAndMove_LightlyPressed)
+		if ((Action.ActionType == ViewportWorldActionTypes::SelectAndMove_LightlyPressed) && (!VREditorMode->IsShowingRadialMenu(VRInteractor)))
 		{
-			if (Action.Event == IE_Pressed && !VRInteractor->IsHoveringOverPriorityType())
+			if ((Action.Event == IE_Pressed) && (!VRInteractor->IsHoveringOverPriorityType()))
 			{
 				// Go ahead and paint immediately
 				FVector LaserPointerStart, LaserPointerEnd;
@@ -561,7 +561,7 @@ void FEdModeFoliage::OnVRAction(class FEditorViewportClient& ViewportClient, UVi
 			}
 
 			// Stop current tracking if the user is no longer painting
-			else if (Action.Event == IE_Released && FoliageInteractor && FoliageInteractor == Interactor)
+			else if (Action.Event == IE_Released && FoliageInteractor && (FoliageInteractor == Interactor))
 			{
 				EndFoliageBrushTrace();
 				FoliageInteractor = nullptr;
diff --git a/Engine/Source/Editor/Kismet/Private/UserDefinedStructureEditor.cpp b/Engine/Source/Editor/Kismet/Private/UserDefinedStructureEditor.cpp
--- a/Engine/Source/Editor/Kismet/Private/UserDefinedStructureEditor.cpp
+++ b/Engine/Source/Editor/Kismet/Private/UserDefinedStructureEditor.cpp
@@ -378,7 +378,10 @@ TSharedRef<SDockTab> FUserDefinedStructureEditor::SpawnStructureTab(const FSpawn
 		EditedStruct = Cast<UUserDefinedStruct>(EditingObjs[ 0 ]);
 	}
 
-	auto Box = SNew(SVerticalBox);
+	TSharedRef<SSplitter> Splitter = SNew(SSplitter)
+		.Orientation(Orient_Vertical)
+		.PhysicalSplitterHandleSize(10.0f)
+		.ResizeMode(ESplitterResizeMode::Fixed);
 
 	{
 		// Create a property view
@@ -389,9 +392,8 @@ TSharedRef<SDockTab> FUserDefinedStructureEditor::SpawnStructureTab(const FSpawn
 		FOnGetDetailCustomizationInstance LayoutStructDetails = FOnGetDetailCustomizationInstance::CreateStatic(&FUserDefinedStructureDetails::MakeInstance);
 		PropertyView->RegisterInstancedCustomPropertyLayout(UUserDefinedStruct::StaticClass(), LayoutStructDetails);
 		PropertyView->SetObject(EditedStruct);
-		Box->AddSlot()
-		.VAlign(EVerticalAlignment::VAlign_Top)
-		.AutoHeight()
+		Splitter->AddSlot()
+		.Value(0.25f)
 		[
 			PropertyView.ToSharedRef()
 		];
@@ -407,10 +409,7 @@ TSharedRef<SDockTab> FUserDefinedStructureEditor::SpawnStructureTab(const FSpawn
 		auto DefaultValueWidget = DefaultValueView->GetWidget();
 		if (DefaultValueWidget.IsValid())
 		{
-			Box->AddSlot()
-			.VAlign(EVerticalAlignment::VAlign_Top)
-			.Padding(2.0f, 0.0f, 0.0f, 0.0f)
-			.AutoHeight()
+			Splitter->AddSlot()
 			[
 				DefaultValueWidget.ToSharedRef()
 			];
@@ -422,7 +421,7 @@ TSharedRef<SDockTab> FUserDefinedStructureEditor::SpawnStructureTab(const FSpawn
 		.Label( LOCTEXT("UserDefinedStructureEditor", "Structure") )
 		.TabColorScale( GetTabColorScale() )
 		[
-			Box
+			Splitter
 		];
 }
 
diff --git a/Engine/Source/Editor/VREditor/Public/VREditorMode.h b/Engine/Source/Editor/VREditor/Public/VREditorMode.h
--- a/Engine/Source/Editor/VREditor/Public/VREditorMode.h
+++ b/Engine/Source/Editor/VREditor/Public/VREditorMode.h
@@ -207,6 +207,9 @@ public:
 		return *UISystem;
 	}
 
+	/** Lets other modules know if the radial menu is visible on a given interactor so input should be handled differently */
+	bool IsShowingRadialMenu( const class UVREditorInteractor* Interactor ) const;
+
 	/** Gets the viewport that VR Mode is activated in.  Even though editor modes are available in all
 	    level viewports simultaneously, only one viewport is "possessed" by the HMD.  Generally try to avoid using
 		this function and instead use the ViewportClient that is passed around through various FEdMode overrides */
@@ -439,6 +442,9 @@ public:
 	// Gets the color
 	FLinearColor GetColor( const EColors Color ) const;
 
+	// Get the default near clipping plane for VR editing
+	float GetDefaultVRNearClipPlane() const;
+
 private:
 
 	// All the colors for this mode
diff --git a/Engine/Source/Editor/VREditor/SVRRadialPanel.cpp b/Engine/Source/Editor/VREditor/SVRRadialPanel.cpp
new file mode 100644
--- /dev/null
+++ b/Engine/Source/Editor/VREditor/SVRRadialPanel.cpp
@@ -0,0 +1,332 @@
+// Copyright 1998-2017 Epic Games, Inc. All Rights Reserved.
+
+#include "SVRRadialPanel.h"
+#include "VREditorModule.h"
+#include "VREditorUISystem.h"
+
+
+namespace VREd
+{
+	static FAutoConsoleVariable MinJoystickOffsetBeforeRadialMenu(TEXT("VREd.MinJoystickOffsetBeforeRadialMenu"), 0.15f, TEXT("Toggles inverting the touch pad vertical axis"));
+}
+
+/**
+ * Construct this widget
+ *
+ * @param	InArgs	The declaration data for this widget
+ */
+void SRadialBox::Construct( const SRadialBox::FArguments& InArgs )
+{
+	const int32 NumSlots = InArgs.Slots.Num();
+	for ( int32 SlotIndex = 0; SlotIndex < NumSlots; ++SlotIndex )
+	{
+		Children.Add( InArgs.Slots[SlotIndex] );
+	}
+
+	CurrentlyHoveredButton = nullptr;
+	RadiusRatioOverride = InArgs._RadiusRatio;
+
+}
+
+/**
+* Helper to ComputeDesiredSize().
+*
+* @param Orientation   Template parameters that controls the orientation in which the children are layed out
+* @param Children      The children whose size we want to assess in a horizontal or vertical arrangement.
+*
+* @return The size desired by the children given an orientation.
+*/
+static FVector2D ComputeDesiredSizeForMenu(const TPanelChildren<SRadialPanel::FSlot>& Children)
+{
+	// The desired size of this panel is the total size desired by its children plus any margins specified in this panel.
+	// The layout along the panel's axis is describe dy the SizeParam, while the perpendicular layout is described by the
+	// alignment property.
+	FVector2D MyDesiredSize(0, 0);
+	for (int32 ChildIndex = 0; ChildIndex < Children.Num(); ++ChildIndex)
+	{
+		const SRadialPanel::FSlot& CurChild = Children[ChildIndex];
+		const FVector2D& CurChildDesiredSize = CurChild.GetWidget()->GetDesiredSize();
+		if (CurChild.GetWidget()->GetVisibility() != EVisibility::Collapsed)
+		{
+
+			// For a vertical panel, we want to find the maximum desired width (including margin).
+			// That will be the desired width of the whole panel.
+			MyDesiredSize.X = FMath::Max(MyDesiredSize.X, (2.0f * CurChildDesiredSize.X + CurChild.SlotPadding.Get().GetTotalSpaceAlong<Orient_Horizontal>()));
+
+			// Clamp to the max size if it was specified
+			float FinalChildDesiredSize = CurChildDesiredSize.Y;
+			float MaxSize = CurChild.MaxSize.Get();
+			if (MaxSize > 0)
+			{
+				FinalChildDesiredSize = FMath::Min(MaxSize, FinalChildDesiredSize);
+			}
+
+			MyDesiredSize.Y += FinalChildDesiredSize + CurChild.SlotPadding.Get().GetTotalSpaceAlong<Orient_Vertical>();
+
+		}
+	}
+
+	return MyDesiredSize;
+}
+
+
+template<EOrientation Orientation>
+static void ArrangeChildrenAlong( const TPanelChildren<SRadialPanel::FSlot>& Children, const FGeometry& AllottedGeometry, FArrangedChildren& ArrangedChildren, TAttribute<float> RadiusOverride )
+{
+	// Allotted space will be given to fixed-size children first.
+	// Remaining space will be proportionately divided between stretch children (SizeRule_Stretch)
+	// based on their stretch coefficient
+
+	if (Children.Num() > 0)
+	{
+		float StretchCoefficientTotal = 0;
+		float FixedTotal = 0;
+
+		// Compute the sum of stretch coefficients (SizeRule_Stretch) and space required by fixed-size widgets
+		// (SizeRule_Auto).
+		for( int32 ChildIndex=0; ChildIndex < Children.Num(); ++ChildIndex )
+		{
+			const SRadialPanel::FSlot& CurChild = Children[ChildIndex];
+
+			if ( CurChild.GetWidget()->GetVisibility() != EVisibility::Collapsed )
+			{
+				// All widgets contribute their margin to the fixed space requirement
+				FixedTotal += CurChild.SlotPadding.Get().GetTotalSpaceAlong<Orientation>();
+
+				if ( CurChild.SizeParam.SizeRule == FSizeParam::SizeRule_Stretch )
+				{
+					// for stretch children we save sum up the stretch coefficients
+					StretchCoefficientTotal += CurChild.SizeParam.Value.Get();
+				}
+				else
+				{
+					// Auto-sized children contribute their desired size to the fixed space requirement
+					float ChildSize =  CurChild.GetWidget()->GetDesiredSize().Y;
+
+					// Clamp to the max size if it was specified
+					float MaxSize = CurChild.MaxSize.Get();
+					FixedTotal += MaxSize > 0 ? FMath::Min( MaxSize, ChildSize ) : ChildSize ;
+				}
+			}
+
+		}
+
+		// The space available for SizeRule_Stretch widgets is any space that wasn't taken up by fixed-sized widgets.
+		const float NonFixedSpace = FMath::Max( 0.0f, AllottedGeometry.Size.Y - FixedTotal );
+
+		int PositionSoFar = 0;
+
+		// Now that we have the total fixed-space requirement and the total stretch coefficients we can
+		// arrange widgets top-to-bottom or left-to-right (depending on the orientation).
+		for( int32 ChildIndex=0; ChildIndex < Children.Num(); ++ChildIndex )
+		{
+			const SRadialPanel::FSlot& CurChild = Children[ChildIndex];
+			const EVisibility ChildVisibility = CurChild.GetWidget()->GetVisibility();
+
+			FVector2D ChildDesiredSize = CurChild.GetWidget()->GetDesiredSize();
+			
+			// Figure out the area allocated to the child in the direction of BoxPanel
+			// The area allocated to the slot is ChildSize + the associated margin.
+			float ChildSize = 0;
+			if ( ChildVisibility != EVisibility::Collapsed )
+			{
+				// The size of the widget depends on its size type
+				if ( CurChild.SizeParam.SizeRule == FSizeParam::SizeRule_Stretch )
+				{
+					// Stretch widgets get a fraction of the space remaining after all the fixed-space requirements are met
+					ChildSize = NonFixedSpace * CurChild.SizeParam.Value.Get() / StretchCoefficientTotal;
+				}
+				else
+				{
+					// Auto-sized widgets get their desired-size value
+					ChildSize =  ChildDesiredSize.Y;
+				}
+
+				// Clamp to the max size if it was specified
+				float MaxSize = CurChild.MaxSize.Get();
+				if( MaxSize > 0 )
+				{
+					ChildSize = FMath::Min( MaxSize, ChildSize );
+				}
+			}
+
+			const FMargin SlotPadding(CurChild.SlotPadding.Get());
+
+			FVector2D SlotSize =  FVector2D( AllottedGeometry.Size.X, ChildSize + SlotPadding.GetTotalSpaceAlong<Orient_Vertical>() );
+
+			// Figure out the size and local position of the child within the slot			
+			AlignmentArrangeResult XAlignmentResult = AlignChild<Orient_Horizontal>( SlotSize.X, CurChild, SlotPadding );					
+			AlignmentArrangeResult YAlignmentResult = AlignChild<Orient_Vertical>( SlotSize.Y, CurChild, SlotPadding );
+			
+			FVector2D MyDesiredSize(0, 0);
+			MyDesiredSize = ComputeDesiredSizeForMenu(Children);
+			
+			// Set the x and y coordinates of the slot based on the current angle, a default radius of MyDesiredSize.Y / 2, and a center of (0, MyDesiredSize.Y / 2.0)
+			// First element ends up at 90 degrees
+			float Radius = (MyDesiredSize.Y / 2.0f) * RadiusOverride.Get();
+			float RadialX =  Radius * FMath::Cos(PositionSoFar * ( (2.0f*PI) / Children.Num() ));
+			float RadialY = MyDesiredSize.Y / 2.0f + Radius * FMath::Sin(PositionSoFar * ( (2.0f*PI) / Children.Num() ));
+			 
+
+			// Add the information about this child to the output list (ArrangedChildren)
+			ArrangedChildren.AddWidget( ChildVisibility, AllottedGeometry.MakeChild(
+				// The child widget being arranged
+				CurChild.GetWidget(),
+				// Child's local position (i.e. position within parent)
+				 FVector2D( RadialX + XAlignmentResult.Offset, RadialY),
+				// Child's size
+				FVector2D( XAlignmentResult.Size, YAlignmentResult.Size )
+				));
+
+			if ( ChildVisibility != EVisibility::Collapsed )
+			{
+				// Offset the next child by the size of the current child and any post-child (bottom/right) margin
+				PositionSoFar++;
+			}
+		}
+	}
+}
+
+/**
+ * Panels arrange their children in a space described by the AllottedGeometry parameter. The results of the arrangement
+ * should be returned by appending a FArrangedWidget pair for every child widget. See StackPanel for an example
+ *
+ * @param AllottedGeometry    The geometry allotted for this widget by its parent.
+ * @param ArrangedChildren    The array to which to add the WidgetGeometries that represent the arranged children.
+ */
+void SRadialPanel::OnArrangeChildren( const FGeometry& AllottedGeometry, FArrangedChildren& ArrangedChildren ) const
+{
+	
+	ArrangeChildrenAlong<Orient_Vertical>(this->Children, AllottedGeometry, ArrangedChildren, RadiusRatioOverride );
+	
+}
+
+
+
+/**
+ * A Panel's desired size in the space required to arrange of its children on the screen while respecting all of
+ * the children's desired sizes and any layout-related options specified by the user. See StackPanel for an example.
+ *
+ * @return The desired size.
+ */
+FVector2D SRadialPanel::ComputeDesiredSize( float ) const
+{
+	return ComputeDesiredSizeForMenu(this->Children);
+}
+
+/** @return  The children of a panel in a slot-agnostic way. */
+FChildren* SRadialPanel::GetChildren()
+{
+	return &Children;
+}
+
+int32 SRadialPanel::RemoveSlot( const TSharedRef<SWidget>& SlotWidget )
+{
+	for (int32 SlotIdx = 0; SlotIdx < Children.Num(); ++SlotIdx)
+	{
+		if ( SlotWidget == Children[SlotIdx].GetWidget() )
+		{
+			Children.RemoveAt(SlotIdx);
+			return SlotIdx;
+		}
+	}
+
+	return -1;
+}
+
+void SRadialPanel::ClearChildren()
+{
+	Children.Empty();
+}
+
+SRadialPanel::SRadialPanel( )
+: Children()
+{
+
+}
+
+void SRadialBox::HighlightSlot(FVector2D& TrackpadPosition)
+{
+	if (FMath::Abs(TrackpadPosition.Size()) < VREd::MinJoystickOffsetBeforeRadialMenu->GetFloat())
+	{
+		if (CurrentlyHoveredButton.Get() != nullptr)
+		{
+			const FPointerEvent& SimulatedPointer = FPointerEvent::FPointerEvent();
+			CurrentlyHoveredButton->OnMouseLeave(SimulatedPointer);
+			CurrentlyHoveredButton = nullptr;
+		}
+		
+		return;
+	}
+
+	const float AnglePerItem = 360.0f / NumEntries;
+	float Angle = FRotator::NormalizeAxis(FMath::RadiansToDegrees(FMath::Atan2(TrackpadPosition.X, TrackpadPosition.Y)));
+
+	Angle += AnglePerItem / 2.0f;
+	// first element of the menu is at 90
+	Angle += -90.0f;
+	if (Angle < 0)
+	{
+		Angle = Angle + 360.0f;
+	}
+
+	const int32 Index = (Angle / AnglePerItem);
+	
+	TSharedRef<SWidget> CurrentChild = GetChildren()->GetChildAt(Index);
+	TSharedRef<SWidget> TestWidget = UVREditorUISystem::FindWidgetOfType(CurrentChild, FName(TEXT("SButton")));
+
+	if (TestWidget != SNullWidget::NullWidget)
+	{
+		CurrentlyHoveredButton = StaticCastSharedRef<SButton>(TestWidget);
+		const FPointerEvent& SimulatedPointer = FPointerEvent::FPointerEvent();
+		const FGeometry& ChildGeometry = FGeometry();
+		
+		// Simulate mouse entering event for the button if it was not previously hovered
+		if(!(CurrentlyHoveredButton->IsHovered()))
+		{
+			CurrentlyHoveredButton->OnMouseEnter(ChildGeometry, SimulatedPointer);
+		}
+
+		// Simulate mouse leaving events for any buttons that were previously hovered
+		for (int32 ButtonCount = 0; ButtonCount < (NumEntries); ButtonCount++)
+		{
+			if (ButtonCount != Index)
+			{
+				CurrentChild = GetChildren()->GetChildAt(ButtonCount);
+				TestWidget = UVREditorUISystem::FindWidgetOfType(CurrentChild, FName(TEXT("SButton")));
+				TSharedRef<SButton> TestButton = StaticCastSharedRef<SButton>(TestWidget);
+				if(TestButton->IsHovered())
+				{
+					TestButton->OnMouseLeave(SimulatedPointer);
+				}
+			}
+
+		}
+	}
+	
+}
+
+void SRadialBox::SimulateLeftClickDown()
+{
+	if(CurrentlyHoveredButton.Get() != nullptr)
+	{
+		const FPointerEvent& SimulatedPointer = FPointerEvent(uint32(0), uint32(0), FVector2D::ZeroVector, FVector2D::ZeroVector, true);
+		const FGeometry& ChildGeometry = FGeometry();
+		CurrentlyHoveredButton->OnMouseButtonDown(ChildGeometry, SimulatedPointer);
+	}	
+}
+
+void SRadialBox::SimulateLeftClickUp()
+{
+	if (CurrentlyHoveredButton.Get() != nullptr)
+	{
+		const FPointerEvent& SimulatedPointer = FPointerEvent(uint32(0), uint32(0), FVector2D::ZeroVector, FVector2D::ZeroVector, true);
+		const FGeometry& ChildGeometry = FGeometry();
+		CurrentlyHoveredButton->OnMouseButtonUp(ChildGeometry, SimulatedPointer);
+	}
+}
+
+const TSharedPtr<SButton>& SRadialBox::GetCurrentlyHoveredButton()
+{
+	return CurrentlyHoveredButton;
+}
diff --git a/Engine/Source/Editor/VREditor/SVRRadialPanel.h b/Engine/Source/Editor/VREditor/SVRRadialPanel.h
new file mode 100644
--- /dev/null
+++ b/Engine/Source/Editor/VREditor/SVRRadialPanel.h
@@ -0,0 +1,254 @@
+// Copyright 1998-2017 Epic Games, Inc. All Rights Reserved.
+
+#pragma once
+
+#include "SPanel.h"
+
+/**
+ * A RadialPanel contains one child and describes how that child should be arranged on the screen.
+ */
+class SRadialPanel
+	: public SPanel
+{
+public:
+
+	/**
+	 * A RadialPanel contains one RadialPanel child and describes how that
+	 * child should be arranged on the screen.
+	 */
+	class FSlot : public TSlotBase<FSlot>
+	{
+	public:		
+		/** Horizontal and Vertical Boxes inherit from FSlot */
+		virtual ~FSlot(){}
+		/**
+		 * How much space this slot should occupy along panel's direction.
+		 *   When SizeRule is SizeRule_Auto, the widget's DesiredSize will be used as the space required.
+		 *   When SizeRule is SizeRule_Stretch, the available space will be distributed proportionately between
+		 *   peer Widgets depending on the Value property. Available space is space remaining after all the
+		 *   peers' SizeRule_Auto requirements have been satisfied.
+		 */
+		FSizeParam SizeParam;
+		
+		/** Horizontal positioning of child within the allocated slot */
+		EHorizontalAlignment HAlignment;
+		
+		/** Vertical positioning of child within the allocated slot */
+		EVerticalAlignment VAlignment;
+		
+		/** The padding to add around the child. */
+		TAttribute<FMargin> SlotPadding;
+		
+		/** The max size that this slot can be (0 if no max) */
+		TAttribute<float> MaxSize;
+			
+	protected:
+		/** Default values for a slot. */
+		FSlot()
+			: TSlotBase<FSlot>()
+			, SizeParam( FStretch(1) )
+			, HAlignment( HAlign_Fill )
+			, VAlignment( VAlign_Fill )
+			, SlotPadding( FMargin(0) )
+			, MaxSize( 0.0f )
+		{ }
+	};
+
+public:
+
+	/** Removes a slot from this box panel which contains the specified SWidget
+	 *
+	 * @param SlotWidget The widget to match when searching through the slots
+	 * @returns The index in the children array where the slot was removed and -1 if no slot was found matching the widget
+	 */
+	int32 RemoveSlot( const TSharedRef<SWidget>& SlotWidget );
+
+	/**
+	 * Removes all children from the box.
+	 */
+	void ClearChildren();
+
+public:
+
+	// Begin SWidget overrides.
+
+	virtual void OnArrangeChildren( const FGeometry& AllottedGeometry, FArrangedChildren& ArrangedChildren ) const;
+	
+	virtual FVector2D ComputeDesiredSize(float) const;
+
+	virtual FChildren* GetChildren();
+
+	// End SWidget overrides.
+
+protected:
+
+	SRadialPanel();
+
+	/** The Radial Panel's children. */
+	TPanelChildren<FSlot> Children;
+
+	/** The ratio to adjust the radial menu's radius by (relative to the default)*/
+	TAttribute<float> RadiusRatioOverride;
+};
+
+
+/** A Radial Box Panel. See SRadialPanel for more info. */
+class SRadialBox : public SRadialPanel
+{
+public:
+	class FSlot : public SRadialPanel::FSlot
+	{
+		public:
+		FSlot()
+		: SRadialPanel::FSlot()
+		{
+		}
+
+		FSlot& AutoWidth()
+		{
+			SizeParam = FAuto();
+			return *this;
+		}
+
+		FSlot& MaxWidth( const TAttribute< float >& InMaxWidth )
+		{
+			MaxSize = InMaxWidth;
+			return *this;
+		}
+
+		FSlot& FillWidth( const TAttribute< float >& StretchCoefficient )
+		{
+			SizeParam = FStretch( StretchCoefficient );
+			return *this;
+		}
+		
+		FSlot& Padding( float Uniform )
+		{
+			SlotPadding = FMargin(Uniform);
+			return *this;
+		}
+
+		FSlot& Padding( float Horizontal, float Vertical )
+		{
+			SlotPadding = FMargin(Horizontal, Vertical);
+			return *this;
+		}
+
+		FSlot& Padding( float Left, float Top, float Right, float Bottom )
+		{
+			SlotPadding = FMargin(Left, Top, Right, Bottom);
+			return *this;
+		}
+		
+		FSlot& Padding( TAttribute<FMargin> InPadding )
+		{
+			SlotPadding = InPadding;
+			return *this;
+		}
+
+		FSlot& HAlign( EHorizontalAlignment InHAlignment )
+		{
+			HAlignment = InHAlignment;
+			return *this;
+		}
+
+		FSlot& VAlign( EVerticalAlignment InVAlignment )
+		{
+			VAlignment = InVAlignment;
+			return *this;
+		}
+
+		FSlot& operator[]( TSharedRef<SWidget> InWidget )
+		{
+			SRadialPanel::FSlot::operator[](InWidget);
+			return *this;
+		}
+
+		FSlot& Expose( FSlot*& OutVarToInit )
+		{
+			OutVarToInit = this;
+			return *this;
+		}
+	};
+
+	static FSlot& Slot()
+	{
+		return *(new FSlot());
+	}
+
+	SLATE_BEGIN_ARGS( SRadialBox )
+	{
+		_Visibility = EVisibility::SelfHitTestInvisible,
+		/** The ratio to adjust the radial menu's radius by (relative to the default) */
+		_RadiusRatio = 1.0f;
+	}
+	SLATE_ARGUMENT( float, RadiusRatio )
+	SLATE_SUPPORTS_SLOT(SRadialBox::FSlot)
+
+	SLATE_END_ARGS()
+
+	FSlot& AddSlot()
+	{
+		SRadialBox::FSlot& NewSlot = *new SRadialBox::FSlot();
+		this->Children.Add( &NewSlot );
+
+		Invalidate(EInvalidateWidget::Layout);
+
+		return NewSlot;
+	}
+
+	FSlot& InsertSlot(int32 Index = INDEX_NONE)
+	{
+		if (Index == INDEX_NONE)
+		{
+			return AddSlot();
+		}
+		SRadialBox::FSlot& NewSlot = *new SRadialBox::FSlot();
+		this->Children.Insert(&NewSlot, Index);
+
+		Invalidate(EInvalidateWidget::Layout);
+
+		return NewSlot;
+	}
+
+	int32 NumSlots() const
+	{
+		return this->Children.Num();
+	}
+
+	int32 NumEntries;
+
+	FORCENOINLINE SRadialBox()
+	: SRadialPanel()
+	{
+		bCanTick = false;
+		bCanSupportFocus = false;
+	}
+
+	/**
+	 * Construct this widget
+	 *
+	 * @param	InArgs	The declaration data for this widget
+	 */
+	void Construct( const FArguments& InArgs );
+
+
+	/**
+	* Highlight the widget in a slot based on a given trackpad position
+	*
+	* @param	TrackpadPosition	The current trackpad or thumbstick position
+	*/
+	void HighlightSlot(FVector2D& TrackpadPosition);
+
+	/** Simulate a left-mouse click (down) on the currently hovered button */
+	void SimulateLeftClickDown();
+
+	/** Simulate a left-mouse click (up) on the currently hovered button */
+	void SimulateLeftClickUp();
+
+	const TSharedPtr<SButton>& GetCurrentlyHoveredButton();
+
+	private:
+	/** Stores the currently hovered button */
+	TSharedPtr<SButton> CurrentlyHoveredButton;
+};
diff --git a/Engine/Source/Editor/VREditor/UI/VREditorQuickMenu.cpp b/Engine/Source/Editor/VREditor/UI/VREditorQuickMenu.cpp
deleted file mode 100644
diff --git a/Engine/Source/Editor/VREditor/UI/VREditorQuickMenu.h b/Engine/Source/Editor/VREditor/UI/VREditorQuickMenu.h
deleted file mode 100644
diff --git a/Engine/Source/Editor/VREditor/UI/VREditorRadialMenu.cpp b/Engine/Source/Editor/VREditor/UI/VREditorRadialMenu.cpp
deleted file mode 100644
diff --git a/Engine/Source/Editor/VREditor/UI/VREditorRadialMenu.h b/Engine/Source/Editor/VREditor/UI/VREditorRadialMenu.h
deleted file mode 100644
diff --git a/Engine/Source/Editor/VREditor/UI/VREditorRadialMenuItem.cpp b/Engine/Source/Editor/VREditor/UI/VREditorRadialMenuItem.cpp
deleted file mode 100644
diff --git a/Engine/Source/Editor/VREditor/UI/VREditorRadialMenuItem.h b/Engine/Source/Editor/VREditor/UI/VREditorRadialMenuItem.h
deleted file mode 100644
diff --git a/Engine/Source/Editor/VREditor/UI/VREditorUISystem.cpp b/Engine/Source/Editor/VREditor/UI/VREditorUISystem.cpp
--- a/Engine/Source/Editor/VREditor/UI/VREditorUISystem.cpp
+++ b/Engine/Source/Editor/VREditor/UI/VREditorUISystem.cpp
@@ -25,12 +25,11 @@
 #include "UnrealEdGlobals.h"
 #include "VREditorMode.h"
 #include "VREditorBaseActor.h"
+#include "VREditorBaseUserWidget.h"
 #include "VREditorFloatingUI.h"
 #include "VREditorDockableWindow.h"
-#include "VREditorBaseUserWidget.h"
-#include "VREditorQuickMenu.h"
-#include "VREditorRadialMenu.h"
 #include "ViewportInteractionTypes.h"
+#include "IHeadMountedDisplay.h"
 #include "ViewportWorldInteraction.h"
 #include "VREditorInteractor.h"
 #include "VREditorMotionControllerInteractor.h"
@@ -40,6 +39,9 @@
 // UI
 #include "Components/WidgetComponent.h"
 #include "VREditorWidgetComponent.h"
+#include "VREditorStyle.h"
+#include "SUniformGridPanel.h"
+#include "SScrollBox.h"
 
 // Content Browser
 #include "CollectionManagerTypes.h"
@@ -61,9 +63,23 @@
 #include "PropertyEditorModule.h"
 #include "IDetailsView.h"
 
+#include "SLevelViewport.h"
+#include "SScaleBox.h"
+#include "SDPIScaler.h"
+#include "SWidget.h"
+#include "SDockTab.h"
+#include "SColorPicker.h"
 #include "Widgets/Layout/SDPIScaler.h"
 #include "Widgets/Docking/SDockTab.h"
 #include "Widgets/Colors/SColorPicker.h"
+#include "MultiBox.h"
+#include "LevelEditorActions.h"
+#include "VREditorActions.h"
+#include "GenericCommands.h"
+#include "SVRRadialPanel.h"
+
+
+#define LOCTEXT_NAMESPACE "VREditor"
 
 namespace VREd
 {
@@ -71,8 +87,8 @@ namespace VREd
 	static FAutoConsoleVariable ContentBrowserUIResolutionY( TEXT( "VREd.ContentBrowserUIResolutionY" ), 1200, TEXT( "Vertical resolution to use for content browser UI render targets" ) );
 	static FAutoConsoleVariable DefaultEditorUIResolutionX( TEXT( "VREd.DefaultEditorUIResolutionX" ), 1024, TEXT( "Horizontal resolution to use for VR editor UI render targets" ) );
 	static FAutoConsoleVariable DefaultEditorUIResolutionY( TEXT( "VREd.DefaultEditorUIResolutionY" ), 1024, TEXT( "Vertical resolution to use for VR editor UI render targets" ) );
-	static FAutoConsoleVariable QuickMenuUIResolutionX( TEXT( "VREd.QuickMenuUIResolutionX" ), 1024, TEXT( "Horizontal resolution to use for Quick Menu VR UI render targets" ) );
-	static FAutoConsoleVariable QuickMenuUIResolutionY( TEXT( "VREd.QuickMenuUIResolutionY" ), 900, TEXT( "Vertical resolution to use for Quick Menu VR UI render targets" ) );
+	static FAutoConsoleVariable QuickMenuUIResolutionX( TEXT( "VREd.QuickMenuUIResolutionX" ), 1440, TEXT( "Horizontal resolution to use for Quick Menu VR UI render targets" ) );
+	static FAutoConsoleVariable QuickMenuUIResolutionY( TEXT( "VREd.QuickMenuUIResolutionY" ), 1440, TEXT( "Vertical resolution to use for Quick Menu VR UI render targets" ) );
 	static FAutoConsoleVariable ContentBrowserUISize( TEXT( "VREd.ContentBrowserUISize" ), 60.0f, TEXT( "How big content browser UIs should be" ) );
 	static FAutoConsoleVariable EditorUISize( TEXT( "VREd.EditorUISize" ), 50.0f, TEXT( "How big editor UIs should be" ) );
 	static FAutoConsoleVariable ContentBrowserUIScale( TEXT( "VREd.ContentBrowserUIScale" ), 2.0f, TEXT( "How much to scale up (or down) the content browser for VR" ) );
@@ -121,8 +137,6 @@ UVREditorUISystem::UVREditorUISystem( const FObjectInitializer& Initializer ) :
 	QuickMenuUI( nullptr ),
 	QuickRadialMenu( nullptr ),
 	RadialMenuHideDelayTime( 0.0f ),
-	QuickMenuWidgetClass( nullptr ),
-	QuickRadialWidgetClass( nullptr ),
 	TutorialWidgetClass( nullptr ),
 	InteractorDraggingUI( nullptr ),
 	DraggingUIOffsetTransform( FTransform::Identity ),
@@ -133,6 +147,7 @@ UVREditorUISystem::UVREditorUISystem( const FObjectInitializer& Initializer ) :
 	StopDragUISound( nullptr ),
 	HideUISound( nullptr ),
 	ShowUISound( nullptr ),
+	ButtonPressSound( nullptr ),
 	bDraggedDockFromHandPassedThreshold( false ),
 	LastDraggingHoverLocation( FVector::ZeroVector ),
 	bSetDefaultLayout( true )
@@ -152,12 +167,6 @@ void UVREditorUISystem::Init()
 	GetOwner().GetWorldInteraction().OnViewportInteractionInputAction().AddUObject( this, &UVREditorUISystem::OnVRAction );
 	GetOwner().GetWorldInteraction().OnViewportInteractionHoverUpdate().AddUObject( this, &UVREditorUISystem::OnVRHoverUpdate );
 
-	QuickMenuWidgetClass = LoadClass<UVREditorQuickMenu>( nullptr, TEXT( "/Engine/VREditor/UI/VRQuickMenu.VRQuickMenu_C" ) );
-	check( QuickMenuWidgetClass != nullptr );
-
-	QuickRadialWidgetClass = LoadClass<UVREditorRadialMenu>( nullptr, TEXT( "/Engine/VREditor/UI/VRRadialQuickMenu.VRRadialQuickMenu_C" ) );
-	check( QuickRadialWidgetClass != nullptr );
-
 	TutorialWidgetClass = LoadClass<UVREditorBaseUserWidget>( nullptr, TEXT( "/Engine/VREditor/Tutorial/UI_VR_Tutorial_00.UI_VR_Tutorial_00_C" ) );
 	check( TutorialWidgetClass != nullptr );
 
@@ -174,11 +183,12 @@ void UVREditorUISystem::Init()
 	ShowUISound = LoadObject<USoundCue>( nullptr, TEXT( "/Engine/VREditor/Sounds/VR_open_Cue" ) );
 	check( ShowUISound != nullptr );
 
-	
+	ButtonPressSound = LoadObject<USoundCue>(nullptr, TEXT("/Engine/VREditor/Sounds/VR_click1_Cue"));
+	check(ButtonPressSound != nullptr);
+
 	// Create all of our UI panels
 	CreateUIs();
 
-
 	// Bind the color picker creation & destruction overrides
 	SColorPicker::OnColorPickerNonModalCreateOverride.BindUObject(this, &UVREditorUISystem::CreateVRColorPicker);
 	SColorPicker::OnColorPickerDestroyOverride.BindUObject(this, &UVREditorUISystem::DestroyVRColorPicker);
@@ -201,13 +211,12 @@ void UVREditorUISystem::Shutdown()
 
 	DefaultWindowTransforms.Empty();
 
-	QuickMenuWidgetClass = nullptr;
-	QuickRadialWidgetClass = nullptr;
 	TutorialWidgetClass = nullptr;
 	StartDragUISound = nullptr;
 	StopDragUISound = nullptr;
 	HideUISound = nullptr;
 	ShowUISound = nullptr;
+	ButtonPressSound = nullptr;
 	VRMode = nullptr;
 	DraggingUI = nullptr;
 	ColorPickerUI = nullptr;
@@ -223,125 +232,268 @@ void UVREditorUISystem::OnVRAction( FEditorViewportClient& ViewportClient, UView
 	{
 		if ( !bWasHandled )
 		{
+			// If the numpad is currently showing
+			if (bRadialMenuIsNumpad)
+			{
+				// Modifier button is backspace
+				if (Action.ActionType == VRActionTypes::Modifier)
+				{
+					const bool bRepeat = false;
+					FVREditorActionCallbacks::SimulateKeyDown( EKeys::BackSpace, bRepeat );
+					FVREditorActionCallbacks::SimulateKeyUp( EKeys::BackSpace );
+					bWasHandled = true;
+				}
+				// Side triggers function as enter keys
+				if (Action.ActionType == ViewportWorldActionTypes::WorldMovement)
+				{
+					const bool bRepeat = false;
+					FVREditorActionCallbacks::SimulateKeyDown( EKeys::Enter, bRepeat );
+					FVREditorActionCallbacks::SimulateKeyUp( EKeys::Enter );
+					// After pressing enter, dismiss the numpad
+					SwapRadialMenu();
+					if (!bRadialMenuVisibleAtSwap)
+					{
+						HideRadialMenu( VREditorInteractor );
+					}
+					bWasHandled = true;
+				}
+			}
 			if ( Action.ActionType == VRActionTypes::ConfirmRadialSelection )
 			{
 				const EViewportInteractionDraggingMode DraggingMode = Interactor->GetDraggingMode();
+				// Toggle the radial menu by pressing the touchpad or thumb stick
 				if ( Action.Event == IE_Pressed && 
 					 ( DraggingMode != EViewportInteractionDraggingMode::ActorsFreely && DraggingMode != EViewportInteractionDraggingMode::World && 
-					   ( DraggingMode != EViewportInteractionDraggingMode::AssistingDrag && Interactor->GetOtherInteractor() != nullptr && Interactor->GetOtherInteractor()->GetDraggingMode() != EViewportInteractionDraggingMode::World ) ) )
+					   ( DraggingMode != EViewportInteractionDraggingMode::AssistingDrag &&
+						 Interactor->GetOtherInteractor() != nullptr ) ) )
 				{
-					UVREditorRadialMenu* RadialMenu = QuickRadialMenu->GetUserWidget<UVREditorRadialMenu>();
-					if ( IsShowingRadialMenu( VREditorInteractor ) )
+					if ( !IsShowingRadialMenu(VREditorInteractor) )
 					{
-						RadialMenu->Update( VREditorInteractor );
+						// Don't force the UI to spawn if using the touchpad/thumbstick to spawn the menu
+						const bool bForceOverUI = false;
+						TryToSpawnRadialMenu(VREditorInteractor, bForceOverUI);
+					}
+					else
+					{	
+						// If the numpad is showing, either swap to the radial menu or dismiss the menu entirely
+						if (bRadialMenuIsNumpad)
+						{
+							SwapRadialMenu();
+							if (!bRadialMenuVisibleAtSwap)
+							{
+								HideRadialMenu(VREditorInteractor);
+							}
+						}
+						// Always dismiss the regular radial menu
+						else
+						{
+							HideRadialMenu(VREditorInteractor);
+						}
 					}
-
-					RadialMenu->SelectCurrentItem();
 					bWasHandled = true;
 				}
 			}
-			else if ( Action.ActionType == ViewportWorldActionTypes::SelectAndMove_LightlyPressed )
+			else 
+			if ( Action.ActionType == ViewportWorldActionTypes::SelectAndMove_LightlyPressed )
 			{
-				FVector LaserPointerStart, LaserPointerEnd;
-				if ( VREditorInteractor->GetLaserPointer( LaserPointerStart, LaserPointerEnd ) )
+				// If the radial menu is showing, select the currently highlighted button by pressing the trigger
+				if ( IsShowingRadialMenu( VREditorInteractor ) )
 				{
-					FHitResult HitResult = VREditorInteractor->GetHitResultFromLaserPointer();
-					if ( HitResult.Actor.IsValid() )
+					TSharedPtr<SRadialBox> RadialBox = StaticCastSharedPtr<SRadialBox>(RadialWidget);
+					if(RadialBox->GetCurrentlyHoveredButton().IsValid())
 					{
-						// Only allow clicks to our own widget components
-						UWidgetComponent* WidgetComponent = Cast<UWidgetComponent>( HitResult.GetComponent() );
-						if ( WidgetComponent != nullptr && IsWidgetAnEditorUIWidget( WidgetComponent ) )
+						if((Action.Event == IE_Pressed))
 						{
-							// Always mark the event as handled so that the editor doesn't try to select the widget component
-							bWasHandled = true;
-
-							if ( Action.Event != IE_Repeat )
+							RadialBox->SimulateLeftClickDown();
+							bOutIsInputCaptured = true;
+						}
+						if (Action.Event == IE_Released)
+						{
+							RadialBox->SimulateLeftClickUp();
+							bOutIsInputCaptured = false;
+						}
+						bWasHandled = true;
+					}
+					else
+					{
+						FVector LaserPointerStart, LaserPointerEnd;
+						if (VREditorInteractor->GetLaserPointer( LaserPointerStart, LaserPointerEnd ))
+						{
+							FHitResult HitResult = VREditorInteractor->GetHitResultFromLaserPointer();
+							if (HitResult.Actor.IsValid())
 							{
-								// If the Modifier button is held down, treat this like a right click instead of a left click
-								const bool bIsRightClicking =
-									( Action.Event == IE_Pressed && VREditorInteractor->IsModifierPressed() ) ||
-									( Action.Event == IE_Released && VREditorInteractor->IsRightClickingOnUI() );
-
-								FVector2D LastLocalHitLocation = WidgetComponent->GetLastLocalHitLocation();
-
-								FVector2D LocalHitLocation;
-								WidgetComponent->GetLocalHitLocation( HitResult.ImpactPoint, LocalHitLocation );
-
-								// If we weren't already hovering over this widget, then we'll reset the last hit location
-								if ( WidgetComponent != VREditorInteractor->GetHoveringOverWidgetComponent() )
-								{
-									LastLocalHitLocation = LocalHitLocation;
-
-									if ( UVREditorWidgetComponent* VRWidgetComponent = Cast<UVREditorWidgetComponent>( VREditorInteractor->GetHoveringOverWidgetComponent() ) )
-									{
-										VRWidgetComponent->SetIsHovering( false );
-									}
-								}
-
-								FWidgetPath WidgetPathUnderFinger = FWidgetPath( WidgetComponent->GetHitWidgetPath( HitResult.ImpactPoint, /*bIgnoreEnabledStatus*/ false ) );
-								if ( WidgetPathUnderFinger.IsValid() )
+								UWidgetComponent* WidgetComponent = Cast<UWidgetComponent>( HitResult.GetComponent() );
+								// If we are clicking on an Actor but not a widget component, send a fake mouse click event to toggle focus
+								if (WidgetComponent == nullptr)
 								{
+									const bool bIsRightClicking =
+										(Action.Event == IE_Pressed && VREditorInteractor->IsModifierPressed()) ||
+										(Action.Event == IE_Released && VREditorInteractor->IsRightClickingOnUI());
 									TSet<FKey> PressedButtons;
-									if ( Action.Event == IE_Pressed )
-									{
-										PressedButtons.Add( bIsRightClicking ? EKeys::RightMouseButton : EKeys::LeftMouseButton );
-									}
-
 									FPointerEvent PointerEvent(
 										1 + (uint8)VREditorInteractor->GetControllerSide(),
-										LocalHitLocation,
-										LastLocalHitLocation,
+										FVector2D::ZeroVector,
+										FVector2D::ZeroVector,
 										PressedButtons,
 										bIsRightClicking ? EKeys::RightMouseButton : EKeys::LeftMouseButton,
 										0.0f,	// Wheel delta
 										FModifierKeysState() );
 
-									VREditorInteractor->SetHoveringOverWidgetComponent( WidgetComponent );
-
-									if ( UVREditorWidgetComponent* VRWidgetComponent = Cast<UVREditorWidgetComponent>( VREditorInteractor->GetHoveringOverWidgetComponent() ) )
+									FWidgetPath EmptyWidgetPath;
+									FReply Reply = FSlateApplication::Get().RoutePointerDownEvent( EmptyWidgetPath, PointerEvent );
+									bWasHandled = true;
+									if (bRadialMenuIsNumpad)
 									{
-										VRWidgetComponent->SetIsHovering( true );
+										// If clicking somewhere outside UI so the widget loses focus
+										SwapRadialMenu();
+										if (!bRadialMenuVisibleAtSwap)
+										{
+											HideRadialMenu( VREditorInteractor );
+										}
 									}
+								}
+							}
+						}
+					}
+				}
+
+				if(!bWasHandled)
+				{
+					FVector LaserPointerStart, LaserPointerEnd;
+					if ( VREditorInteractor->GetLaserPointer( LaserPointerStart, LaserPointerEnd ) )
+					{
+						FHitResult HitResult = VREditorInteractor->GetHitResultFromLaserPointer();
+						if ( HitResult.Actor.IsValid() )
+						{
+							// Only allow clicks to our own widget components
+							UWidgetComponent* WidgetComponent = Cast<UWidgetComponent>( HitResult.GetComponent() );
+							if ( WidgetComponent != nullptr && IsWidgetAnEditorUIWidget( WidgetComponent ) )
+							{
+								// Always mark the event as handled so that the editor doesn't try to select the widget component
+								bWasHandled = true;
+
+								if ( Action.Event != IE_Repeat )
+								{
+									// If the Modifier button is held down, treat this like a right click instead of a left click
+									const bool bIsRightClicking =
+										( Action.Event == IE_Pressed && VREditorInteractor->IsModifierPressed() ) ||
+										( Action.Event == IE_Released && VREditorInteractor->IsRightClickingOnUI() );
+
+									FVector2D LastLocalHitLocation = WidgetComponent->GetLastLocalHitLocation();
 
-									FReply Reply = FReply::Unhandled();
-									if ( Action.Event == IE_Pressed )
+									FVector2D LocalHitLocation;
+									WidgetComponent->GetLocalHitLocation( HitResult.ImpactPoint, LocalHitLocation );
+
+									// If we weren't already hovering over this widget, then we'll reset the last hit location
+									if ( WidgetComponent != VREditorInteractor->GetHoveringOverWidgetComponent() )
 									{
-										const double CurrentTime = FPlatformTime::Seconds();
-										if (CurrentTime - VREditorInteractor->GetLastClickReleaseTime() <= VREd::DoubleClickTime->GetFloat())
+										LastLocalHitLocation = LocalHitLocation;
+
+										if ( UVREditorWidgetComponent* VRWidgetComponent = Cast<UVREditorWidgetComponent>( VREditorInteractor->GetHoveringOverWidgetComponent() ) )
 										{
-											// Trigger a double click event!
-											Reply = FSlateApplication::Get().RoutePointerDoubleClickEvent(WidgetPathUnderFinger, PointerEvent);
+											VRWidgetComponent->SetIsHovering( false );
 										}
-										else
+									}
+
+									FWidgetPath WidgetPathUnderFinger = FWidgetPath( WidgetComponent->GetHitWidgetPath( HitResult.ImpactPoint, /*bIgnoreEnabledStatus*/ false ) );
+									if ( WidgetPathUnderFinger.IsValid() )
+									{
+										TSet<FKey> PressedButtons;
+										if ( Action.Event == IE_Pressed )
 										{
-											Reply = FSlateApplication::Get().RoutePointerDownEvent(WidgetPathUnderFinger, PointerEvent);
+											PressedButtons.Add( bIsRightClicking ? EKeys::RightMouseButton : EKeys::LeftMouseButton );
 										}
 
-										// In case of selecting a level in the content browser the VREditormode is closed, this makes sure nothing happens after that.
-										if( !IVREditorModule::Get().IsVREditorModeActive() )
+										FPointerEvent PointerEvent(
+											1 + (uint8)VREditorInteractor->GetControllerSide(),
+											LocalHitLocation,
+											LastLocalHitLocation,
+											PressedButtons,
+											bIsRightClicking ? EKeys::RightMouseButton : EKeys::LeftMouseButton,
+											0.0f,	// Wheel delta
+											FModifierKeysState() );
+
+										VREditorInteractor->SetHoveringOverWidgetComponent( WidgetComponent );
+
+										if ( UVREditorWidgetComponent* VRWidgetComponent = Cast<UVREditorWidgetComponent>( VREditorInteractor->GetHoveringOverWidgetComponent() ) )
 										{
-											bWasHandled = true;
-											return;
+											VRWidgetComponent->SetIsHovering( true );
 										}
+
+										FReply Reply = FReply::Unhandled();
+										if ( Action.Event == IE_Pressed )
+										{
+											const double CurrentTime = FPlatformTime::Seconds();
+											if (CurrentTime - VREditorInteractor->GetLastClickReleaseTime() <= VREd::DoubleClickTime->GetFloat())
+											{
+												// Trigger a double click event!
+												Reply = FSlateApplication::Get().RoutePointerDoubleClickEvent(WidgetPathUnderFinger, PointerEvent);
+											}
+											else
+											{
+												// If we are clicking on an editable text field and the radial menu is not a numpad, show the numpad
+												if (WidgetPathUnderFinger.Widgets.Last().Widget->GetTypeAsString() == TEXT("SEditableText") && (bRadialMenuIsNumpad == false))
+												{
+													if (!QuickRadialMenu->bHidden)
+													{
+														bRadialMenuVisibleAtSwap = true;
+													}
+													else
+													{
+														bRadialMenuVisibleAtSwap = false;
+														// Force the radial menu to spawn even if the laser is over UI
+														const bool bForceOverUI = true;
+														TryToSpawnRadialMenu( VREditorInteractor, bForceOverUI );
+													}
+													SwapRadialMenu();
+												}
+												Reply = FSlateApplication::Get().RoutePointerDownEvent(WidgetPathUnderFinger, PointerEvent);
+											}
+
+											// In case of selecting a level in the content browser the VREditormode is closed, this makes sure nothing happens after that.
+											if( !IVREditorModule::Get().IsVREditorModeActive() )
+											{
+												bWasHandled = true;
+												return;
+											}
 										
-										VREditorInteractor->SetIsClickingOnUI( true );
-										VREditorInteractor->SetIsRightClickingOnUI( bIsRightClicking );
-										VREditorInteractor->SetLastClickReleaseTime( CurrentTime );
-										bOutIsInputCaptured = true;
+											VREditorInteractor->SetIsClickingOnUI( true );
+											VREditorInteractor->SetIsRightClickingOnUI( bIsRightClicking );
+											VREditorInteractor->SetLastClickReleaseTime( CurrentTime );
+											bOutIsInputCaptured = true;
 
-										// Play a haptic effect on press
-										VREditorInteractor->PlayHapticEffect( VREd::UIPressHapticFeedbackStrength->GetFloat() );
-									}
-									else if ( Action.Event == IE_Released )
-									{
-										Reply = FSlateApplication::Get().RoutePointerUpEvent( WidgetPathUnderFinger, PointerEvent );
+											// Play a haptic effect on press
+											VREditorInteractor->PlayHapticEffect( VREd::UIPressHapticFeedbackStrength->GetFloat() );
+										}
+										else if ( Action.Event == IE_Released )
+										{
+											Reply = FSlateApplication::Get().RoutePointerUpEvent( WidgetPathUnderFinger, PointerEvent );
+										}
 									}
 								}
 							}
+							else
+							{
+								// If we didn't handle the input in any other way, send an empty mouse down event so Slate focus is handled correctly
+								const bool bIsRightClicking =
+									(Action.Event == IE_Pressed && VREditorInteractor->IsModifierPressed()) ||
+									(Action.Event == IE_Released && VREditorInteractor->IsRightClickingOnUI());
+								TSet<FKey> PressedButtons;
+								FPointerEvent PointerEvent(
+									1 + (uint8)VREditorInteractor->GetControllerSide(),
+									FVector2D::ZeroVector,
+									FVector2D::ZeroVector,
+									PressedButtons,
+									bIsRightClicking ? EKeys::RightMouseButton : EKeys::LeftMouseButton,
+									0.0f,	// Wheel delta
+									FModifierKeysState() );
+
+								FWidgetPath EmptyWidgetPath;
+								FReply Reply = FSlateApplication::Get().RoutePointerDownEvent( EmptyWidgetPath, PointerEvent );
+							}
 						}
 					}
 				}
-
 				if ( Action.Event == IE_Released )
 				{
 					bool bWasRightClicking = false;
@@ -394,6 +546,16 @@ void UVREditorUISystem::OnVRHoverUpdate( FEditorViewportClient& ViewportClient,
 				FHitResult HitResult = Interactor->GetHitResultFromLaserPointer();
 				if( HitResult.Actor.IsValid() )
 				{
+					// If a button is currently highlighted on the radial menu, don't highlight another button with the laser
+					if ((HitResult.Actor == QuickRadialMenu) && (RadialWidget.Get() != nullptr))
+					{
+						const TSharedPtr<SRadialBox>& RadialBox = StaticCastSharedPtr<SRadialBox>(RadialWidget);
+						if (RadialBox->GetCurrentlyHoveredButton().Get() != nullptr)
+						{
+							return;
+						}
+					}
+
 					// Only allow clicks to our own widget components
 					UWidgetComponent* WidgetComponent = Cast<UWidgetComponent>( HitResult.GetComponent() );
 					if( WidgetComponent != nullptr && IsWidgetAnEditorUIWidget( WidgetComponent ) )
@@ -428,7 +590,7 @@ void UVREditorUISystem::OnVRHoverUpdate( FEditorViewportClient& ViewportClient,
 								FModifierKeysState());
 
 							FSlateApplication::Get().RoutePointerMoveEvent(WidgetPathUnderFinger, PointerEvent, false);
-
+							
 							bWasHandled = true;
 							HoverImpactPoint = HitResult.ImpactPoint;
 							VREditorInteractor->SetHoveringOverWidgetComponent( WidgetComponent );
@@ -445,7 +607,8 @@ void UVREditorUISystem::OnVRHoverUpdate( FEditorViewportClient& ViewportClient,
 								const bool bIsAbsolute = ( GetOwner().GetHMDDeviceType() == EHMDDeviceType::DT_SteamVR );
 
 								float ScrollDelta = 0.0f;
-								if( VREditorInteractor->IsTouchingTrackpad() || !bIsAbsolute )
+								// Don't scroll if the radial menu is a number pad
+								if(( VREditorInteractor->IsTouchingTrackpad() || !bIsAbsolute) && !bRadialMenuIsNumpad)
 								{
 									if( bIsAbsolute )
 									{
@@ -584,14 +747,14 @@ void UVREditorUISystem::Tick( FEditorViewportClient* ViewportClient, const float
 			{
 				// @todo vreditor tweak: Weird to hard code this here. Probably should be an accessor on the hand itself, and based on the actual device type
 				const FTransform UICapsuleTransform = OtherInteractor->GetTransform();
-				const FVector UICapsuleStart = FVector( -9.0f, 0.0f, 0.0f ) * WorldScaleFactor;
+				const FVector UICapsuleStart = FVector( -3.0f, 0.0f, 0.0f ) * WorldScaleFactor;
 				const FVector UICapsuleEnd = FVector( -18.0f, 0.0f, 0.0f ) * WorldScaleFactor;
-				const float UICapsuleLocalRadius = 6.0f * WorldScaleFactor;
+				const float UICapsuleLocalRadius = 8.0f * WorldScaleFactor;
 				const float MinDistanceToUICapsule = 8.0f * WorldScaleFactor; // @todo vreditor tweak
 				const FVector UIForwardVector = FVector::UpVector;
 				const float MinDotForAimingAtOtherHand = 0.25f;	// @todo vreditor tweak
 
-				if( GetOwner().IsHandAimingTowardsCapsule( Interactor, UICapsuleTransform, UICapsuleStart, UICapsuleEnd, UICapsuleLocalRadius, MinDistanceToUICapsule, UIForwardVector, MinDotForAimingAtOtherHand ) )
+				if (GetOwner().IsHandAimingTowardsCapsule(Interactor, UICapsuleTransform, UICapsuleStart, UICapsuleEnd, UICapsuleLocalRadius, MinDistanceToUICapsule, UIForwardVector, MinDotForAimingAtOtherHand) && (QuickRadialMenu->bHidden))
 				{
 					bShowQuickMenuOnArm = true;
 				}
@@ -609,9 +772,26 @@ void UVREditorUISystem::Tick( FEditorViewportClient* ViewportClient, const float
 		{
 			// Despawn
 			HandInteractorWithQuickMenu->SetHasUIOnForearm( false );
+			HandInteractorWithQuickMenu->SetHasUIInFront( false );
 			QuickMenuUI->ShowUI( false );
+			// Reset all the currently animating buttons to their minimum scale and stop animating them
+			for (FVRButton& VRButton : VRButtons)
+			{
+				if(VRButton.AnimationDirection != EVREditorAnimationState::None)
+				{
+					
+					TSharedRef<SWidget> PlaceholderWidget = (VRButton.BaseSlot->DetachWidget()).ToSharedRef();
+					TSharedRef<SWidget> WidgetToTransfer = (VRButton.HoverSlot->DetachWidget()).ToSharedRef();
+					VRButton.BaseSlot->AttachWidget(WidgetToTransfer);
+					VRButton.HoverSlot->AttachWidget(PlaceholderWidget);
+					VRButton.CurrentScale = VRButton.MinScale;
+					VRButton.ButtonBorder->SetRenderTransform(FSlateRenderTransform::FTransform2D(VRButton.CurrentScale));
+					VRButton.AnimationDirection = EVREditorAnimationState::None;
+		}
+			}
 		}
 
+
 		if( HandInteractorThatNeedsQuickMenu != nullptr && !QuickMenuUI->IsUIVisible() )
 		{
 			UVREditorMotionControllerInteractor* MotionControllerHandInteractorThatNeedsQuickMenu = Cast<UVREditorMotionControllerInteractor>( HandInteractorThatNeedsQuickMenu );
@@ -620,17 +800,49 @@ void UVREditorUISystem::Tick( FEditorViewportClient* ViewportClient, const float
 				const AVREditorFloatingUI::EDockedTo DockTo = ( MotionControllerHandInteractorThatNeedsQuickMenu->GetControllerSide() == EControllerHand::Left ) ? AVREditorFloatingUI::EDockedTo::LeftArm : AVREditorFloatingUI::EDockedTo::RightArm;
 				QuickMenuUI->SetDockedTo( DockTo );
 				QuickMenuUI->ShowUI( true );
+				MotionControllerHandInteractorThatNeedsQuickMenu->SetHasUIInFront( true );
 				MotionControllerHandInteractorThatNeedsQuickMenu->SetHasUIOnForearm( true );
 			}
 		}
 	}
-
-	// Close the radial menu if it was not updated for a while
-	RadialMenuHideDelayTime += DeltaTime;
-	if( RadialMenuHideDelayTime >= VREd::RadialMenuFadeDelay->GetFloat() )
+	
+	// Iterate through all quick menu and radial menu buttons and animate any that need it
+	for (FVRButton& VRButton : VRButtons)
 	{
-		UVREditorInteractor* Interactor = GetOwner().GetHandInteractor( QuickRadialMenu->GetDockedTo() == AVREditorFloatingUI::EDockedTo::LeftHand ? EControllerHand::Left : EControllerHand::Right );
-		HideRadialMenu( Interactor );
+		switch (VRButton.AnimationDirection)
+		{
+			case EVREditorAnimationState::Forward:
+			{
+				if (VRButton.CurrentScale < VRButton.MaxScale)
+				{
+					VRButton.CurrentScale += DeltaTime*VRButton.ScaleRate;
+				}
+				else
+				{
+					VRButton.CurrentScale = VRButton.MaxScale;
+					VRButton.AnimationDirection = EVREditorAnimationState::None;
+				}
+				break;
+			}
+			case EVREditorAnimationState::Backward:
+			{
+				if (VRButton.CurrentScale > VRButton.MinScale)
+				{
+					VRButton.CurrentScale -= DeltaTime*VRButton.ScaleRate;
+				}
+				else
+				{
+					VRButton.CurrentScale = VRButton.MinScale;
+					VRButton.AnimationDirection = EVREditorAnimationState::None;	
+				}
+				break;
+			}
+			case EVREditorAnimationState::None:
+			{
+				break;
+			}
+		}
+		VRButton.ButtonBorder->SetRenderTransform(FSlateRenderTransform::FTransform2D(VRButton.CurrentScale));	
 	}
 	
 	// Tick all of our floating UIs
@@ -653,37 +865,38 @@ void UVREditorUISystem::CreateUIs()
 	const FIntPoint DefaultResolution( VREd::DefaultEditorUIResolutionX->GetInt(), VREd::DefaultEditorUIResolutionY->GetInt() );
 
 	const bool bIsVREditorDemo = FParse::Param( FCommandLine::Get(), TEXT( "VREditorDemo" ) );	// @todo vreditor: Remove this when no longer needed
-																								
-	// @todo vreditor: Tweak
-	if ( QuickMenuWidgetClass != nullptr )
 	{
-		const FIntPoint Resolution( VREd::QuickMenuUIResolutionX->GetInt(), VREd::QuickMenuUIResolutionY->GetInt() );
-
+		const FIntPoint Resolution(VREd::QuickMenuUIResolutionX->GetInt(), VREd::QuickMenuUIResolutionY->GetInt());
 		const bool bWithSceneComponent = false;
-		QuickMenuUI = GetOwner().SpawnTransientSceneActor< AVREditorFloatingUI >(TEXT("QuickMenu"), bWithSceneComponent);
-		QuickMenuUI->SetUMGWidget( *this, QuickMenuWidgetClass, Resolution, 30.0f, AVREditorFloatingUI::EDockedTo::Nothing );
-		QuickMenuUI->ShowUI( false );
-		QuickMenuUI->SetRelativeOffset( FVector( -11.0f, 0.0f, 3.0f ) );
-		FloatingUIs.Add( QuickMenuUI );
-	}
+		// @todo vreditor: Tweak
+		{
+			QuickMenuUI = GetOwner().SpawnTransientSceneActor< AVREditorFloatingUI >(TEXT("QuickMenu"), bWithSceneComponent);
+			QuickMenuUI->SetSlateWidget(*this, BuildQuickMenuWidget(), Resolution, 40.0f, AVREditorFloatingUI::EDockedTo::Nothing);
+			QuickMenuUI->ShowUI(false);
+			QuickMenuUI->SetRelativeOffset(FVector(-12.0f, 0.0f, 5.0f));
+			QuickMenuUI->GetWidgetComponent()->SetOpacityFromTexture(1.0f);
+			QuickMenuUI->GetWidgetComponent()->SetBackgroundColor(FLinearColor::Transparent);
+			QuickMenuUI->GetWidgetComponent()->SetBlendMode(EWidgetBlendMode::Masked);
+			FloatingUIs.Add(QuickMenuUI);
 
-	// Create the radial UI
-	if ( QuickRadialWidgetClass != nullptr )
-	{
-		TSharedPtr<SWidget> WidgetToDraw;
-		QuickRadialMenu = GetOwner().SpawnTransientSceneActor< AVREditorFloatingUI >(TEXT("QuickRadialmenu"), false);
-		QuickRadialMenu->SetUMGWidget( *this, QuickRadialWidgetClass, DefaultResolution, 40.0f, AVREditorFloatingUI::EDockedTo::Nothing );
-		QuickRadialMenu->ShowUI( false );
-		QuickRadialMenu->SetActorEnableCollision( false );
-		QuickRadialMenu->SetRelativeOffset( FVector( 10.0f, 0.0f, 10.f ) );
-		QuickRadialMenu->SetCollisionOnShow( false );
-		FloatingUIs.Add(QuickRadialMenu);
+		}
+		// Create the radial UI
+		{
+			QuickRadialMenu = GetOwner().SpawnTransientSceneActor< AVREditorFloatingUI >(TEXT("QuickRadialmenu"), bWithSceneComponent);
+			QuickRadialMenu->SetSlateWidget(*this, BuildRadialMenuWidget(), DefaultResolution, 40.0f, AVREditorFloatingUI::EDockedTo::Nothing);
+			QuickRadialMenu->ShowUI(false);
+			QuickRadialMenu->SetActorEnableCollision(false);
+			QuickRadialMenu->SetRelativeOffset(FVector(0.0f, 0.0f, 5.0f));
+			QuickRadialMenu->GetWidgetComponent()->SetOpacityFromTexture(1.0f);
+			QuickRadialMenu->GetWidgetComponent()->SetBackgroundColor(FLinearColor::Transparent);
+			QuickRadialMenu->GetWidgetComponent()->SetBlendMode(EWidgetBlendMode::Masked);
+			FloatingUIs.Add(QuickRadialMenu);
+		}
 	}
-
 	// Make some editor UIs!
 	{
 		{
-			const FIntPoint Resolution( VREd::ContentBrowserUIResolutionX->GetInt(), VREd::ContentBrowserUIResolutionY->GetInt() );
+			const FIntPoint Resolution = FIntPoint( VREd::ContentBrowserUIResolutionX->GetInt(), VREd::ContentBrowserUIResolutionY->GetInt() );
 
 			IContentBrowserSingleton& ContentBrowserSingleton = FModuleManager::Get().LoadModuleChecked<FContentBrowserModule>( "ContentBrowser" ).Get();;
 
@@ -821,6 +1034,7 @@ void UVREditorUISystem::CreateUIs()
 
 			AVREditorFloatingUI* TutorialUI = GetOwner().SpawnTransientSceneActor< AVREditorDockableWindow >( TEXT( "TutorialUI" ), false );
 			TutorialUI->SetUMGWidget( *this, TutorialWidgetClass, FIntPoint( VREd::TutorialUIResolutionX->GetFloat(), VREd::TutorialUIResolutionY->GetFloat() ), VREd::TutorialUISize->GetFloat(), bShowAtStart ? AVREditorFloatingUI::EDockedTo::Room : AVREditorFloatingUI::EDockedTo::Nothing );
+
 			if( bShowAtStart )
 			{
 				TutorialUI->SetRelativeOffset( FVector( VREd::TutorialUILocationX->GetFloat(), VREd::TutorialUILocationY->GetFloat(), VREd::TutorialUILocationZ->GetFloat() ) );
@@ -835,7 +1049,7 @@ void UVREditorUISystem::CreateUIs()
 		}
 
 		{
-			const FIntPoint Resolution(VREd::AssetEditorUIResolutionX->GetInt(), VREd::AssetEditorUIResolutionY->GetInt());
+			const FIntPoint Resolution = FIntPoint(VREd::AssetEditorUIResolutionX->GetInt(), VREd::AssetEditorUIResolutionY->GetInt());
 
 			const bool bWithSceneComponent = false;
 			AVREditorDockableWindow* TabManagerUI = GetOwner().SpawnTransientSceneActor< AVREditorDockableWindow >( TEXT( "AssetEditor" ), bWithSceneComponent );
@@ -1035,38 +1249,22 @@ void UVREditorUISystem::ShowEditorUIPanel( AVREditorFloatingUI* Panel, UVREditor
 			UGameplayStatics::PlaySound2D( GetOwner().GetWorld(), bShouldShow ? ShowUISound : HideUISound );
 		}
 
-		if ( bRefreshQuickMenu && QuickMenuUI )
-		{
-			QuickMenuUI->GetUserWidget<UVREditorQuickMenu>()->RefreshUI();
-		}
 	}
 }
 
 
-bool UVREditorUISystem::IsShowingRadialMenu( UVREditorInteractor* Interactor ) const
+bool UVREditorUISystem::IsShowingRadialMenu(const UVREditorInteractor* Interactor ) const
 {
-	const EControllerHand DockedToHand = QuickRadialMenu->GetDockedTo() == AVREditorFloatingUI::EDockedTo::LeftHand ? EControllerHand::Left : EControllerHand::Right;
+	const EControllerHand DockedToHand = QuickRadialMenu->GetDockedTo() == AVREditorFloatingUI::EDockedTo::LeftArm ? EControllerHand::Left : EControllerHand::Right;
 	const UVREditorInteractor* DockedToInteractor = GetOwner().GetHandInteractor( DockedToHand );
 	return DockedToInteractor == Interactor && !QuickRadialMenu->bHidden;
 }
 
 
-void UVREditorUISystem::UpdateRadialMenu( UVREditorInteractor* Interactor )
-{
-	if( QuickRadialMenu->bHidden )
-	{
-		QuickRadialMenu->ShowUI( true, false );
-	}
 
-	if( !QuickRadialMenu->bHidden )
-	{
-		RadialMenuHideDelayTime = 0.0f;
-		QuickRadialMenu->GetUserWidget<UVREditorRadialMenu>()->Update( Interactor );
-	}
-}
 
 
-void UVREditorUISystem::TryToSpawnRadialMenu( UVREditorInteractor* Interactor )
+void UVREditorUISystem::TryToSpawnRadialMenu( UVREditorInteractor* Interactor, const bool bForceOverUI )
 {
 	if( Interactor != nullptr )
 	{
@@ -1076,7 +1274,7 @@ void UVREditorUISystem::TryToSpawnRadialMenu( UVREditorInteractor* Interactor )
 			UVREditorInteractor* DockedToInteractor = nullptr;
 			if( !QuickRadialMenu->bHidden )
 			{
-				const EControllerHand ControllerHand = QuickRadialMenu->GetDockedTo() == AVREditorFloatingUI::EDockedTo::LeftHand ? EControllerHand::Left : EControllerHand::Right;
+				const EControllerHand ControllerHand = QuickRadialMenu->GetDockedTo() == AVREditorFloatingUI::EDockedTo::LeftArm ? EControllerHand::Left : EControllerHand::Right;
 				DockedToInteractor = GetOwner().GetHandInteractor( ControllerHand );
 			}
 
@@ -1090,21 +1288,13 @@ void UVREditorUISystem::TryToSpawnRadialMenu( UVREditorInteractor* Interactor )
 				DraggingMode != EViewportInteractionDraggingMode::World &&
 				DraggingMode != EViewportInteractionDraggingMode::AssistingDrag &&
 				( InteractorDraggingUI == nullptr || InteractorDraggingUI != Interactor ) &&
-				!Interactor->IsHoveringOverUI();	// Don't show radial menu when aiming at a UI  (too much clutter)
+				(!Interactor->IsHoveringOverUI() || bForceOverUI);	// Don't show radial menu when aiming at a UI unless forcing it (too much clutter)
 
-			UVREditorRadialMenu* RadialMenu = QuickRadialMenu->GetUserWidget<UVREditorRadialMenu>();
-			// We need to update the trackpad position in the radialmenu before checking if it can be used
-			RadialMenu->Update( Interactor );
-			if( RadialMenu && RadialMenu->IsInMenuRadius() )
+			if (bNeedsSpawn)
 			{
-				bNeedsSpawn = false;
-			}
-
-			if( bNeedsSpawn )
-			{
-				const AVREditorFloatingUI::EDockedTo DockedTo = MotionControllerInteractor->GetControllerSide() == EControllerHand::Left ? AVREditorFloatingUI::EDockedTo::LeftHand : AVREditorFloatingUI::EDockedTo::RightHand;
-				QuickRadialMenu->SetDockedTo( DockedTo );
-				QuickRadialMenu->ShowUI( true );
+				const AVREditorFloatingUI::EDockedTo DockedTo = MotionControllerInteractor->GetControllerSide() == EControllerHand::Left ? AVREditorFloatingUI::EDockedTo::LeftArm : AVREditorFloatingUI::EDockedTo::RightArm;
+				QuickRadialMenu->SetDockedTo(DockedTo);
+				QuickRadialMenu->ShowUI(true);
 			}
 		}
 	}
@@ -1351,11 +1541,7 @@ void UVREditorUISystem::TogglePanelsVisibility()
 		VREditorInteractor->SetHasUIInFront( false );
 	}
 
-	if (QuickMenuUI)
-	{
-		QuickMenuUI->GetUserWidget<UVREditorQuickMenu>()->RefreshUI();
 	}
-}
 
 float UVREditorUISystem::GetMaxDockWindowSize() const
 {
@@ -1546,3 +1732,856 @@ void UVREditorUISystem::DestroyVRColorPicker()
 	const bool bRefreshQuickMenu = true;
 	ShowEditorUIPanel(EEditorUIPanel::ColorPicker, VREditorInteractor, bShouldShow, bShowOnHand, bRefreshQuickMenu);
 }
+
+void UVREditorUISystem::MakeUniformGridMenu(const TSharedRef<FMultiBox>& MultiBox, const TSharedRef<SMultiBoxWidget>& MultiBoxWidget, int32 Columns)
+{
+	// Grab the list of blocks, early out if there's nothing to fill the widget with
+	const TArray< TSharedRef<const FMultiBlock > >& Blocks = MultiBox->GetBlocks();
+	const uint32 NumItems = Blocks.Num();
+	if( NumItems == 0 )
+	{
+		return;
+	}
+
+	TSharedRef<SUniformGridPanel> GridPanel = SNew( SUniformGridPanel )
+		.SlotPadding(3.0f);
+	MultiBox.Get().SetStyle(&FVREditorStyle::Get(), FVREditorStyle::GetStyleSetName());
+	int32 Index = 0;
+	int32 Row = 0;
+	for( const TSharedRef<const FMultiBlock>& MultiBlock : Blocks )
+	{
+		const TSharedRef<const FMultiBlock>& Block = MultiBlock;
+		if(Block->GetType() == EMultiBlockType::MenuEntry)
+		{
+			TSharedRef<SWidget> BlockWidget = Block->MakeWidget(MultiBoxWidget, EMultiBlockLocation::Middle, true)->AsWidget();
+		
+			int32 Column = FMath::Fmod(Index, Columns);
+			Row = FMath::DivideAndRoundDown(Index, Columns);
+			TSharedRef<SOverlay> TestOverlay = SNew(SOverlay);
+			GridPanel->AddSlot(Column, Row)
+				.VAlign(VAlign_Fill)
+				.HAlign(HAlign_Fill)
+				[
+					SAssignNew(TestOverlay, SOverlay)
+				];
+
+			AddHoverableButton(BlockWidget, TEXT("SMenuEntryButton"), TestOverlay);
+			SetButtonTextWrap(BlockWidget, 75.0);
+			Index++;
+		}
+	}
+
+	MultiBoxWidget->SetContent(
+		SNew(SScrollBox)
+		+ SScrollBox::Slot()
+		.Padding(0, 0, 0, 0)
+		.VAlign(VAlign_Fill)
+		.HAlign(HAlign_Fill)
+		[
+			GridPanel
+		]
+	);
+}
+
+TSharedRef<SWidget> UVREditorUISystem::AddHoverableButton(TSharedRef<SWidget>& BlockWidget, FName ButtonType, TSharedRef<SOverlay>& TestOverlay)
+{
+	TSharedPtr<SBorder> ButtonBorder;
+	SOverlay::FOverlaySlot* BaseSlot;
+	SOverlay::FOverlaySlot* HoverSlot;
+	TestOverlay->AddSlot(0)
+	.Expose(BaseSlot)
+		[
+			SAssignNew(ButtonBorder, SBorder)
+			.BorderImage(FEditorStyle::GetBrush("ToolPanel.GroupBorder"))
+		[
+			BlockWidget
+		]
+		];
+	// Placeholder slot to put the button into when it is hovered
+	TestOverlay->AddSlot(10)
+	.Expose(HoverSlot)
+		[
+			SNew(SOverlay)
+			.Visibility(EVisibility::SelfHitTestInvisible)
+		];
+
+	TSharedRef<SWidget> TestWidget = FindWidgetOfType(BlockWidget, ButtonType);
+	if (TestWidget != SNullWidget::NullWidget)
+	{
+		TSharedRef<SButton> Button = StaticCastSharedRef<SButton>(TestWidget);
+		Button->SetRenderTransformPivot(FVector2D(0.5, 0.5));
+		Button->SetOnHovered(FSimpleDelegate::CreateUObject(this, &UVREditorUISystem::OnHoverBeginEffect, Button));
+		Button->SetOnUnhovered(FSimpleDelegate::CreateUObject(this, &UVREditorUISystem::OnHoverEndEffect, Button));
+		FSlateSound SlateButtonPressSound = FSlateSound();
+		SlateButtonPressSound.SetResourceObject(ButtonPressSound);
+		Button->SetPressedSound(SlateButtonPressSound);
+		ButtonBorder->SetRenderTransformPivot(FVector2D(0.5, 0.5));
+		FVRButton ButtonToAdd = FVRButton(Button, BaseSlot, HoverSlot, ButtonBorder);
+		VRButtons.Add(ButtonToAdd);
+	}
+	return BlockWidget;
+}
+
+TSharedRef<SWidget> UVREditorUISystem::SetButtonTextWrap(TSharedRef<SWidget>& BlockWidget, float WrapSize)
+{
+	TSharedRef<SWidget> TestWidget = FindWidgetOfType(BlockWidget, FName(TEXT("STextBlock")));
+	if (TestWidget != SNullWidget::NullWidget)
+	{
+		TSharedRef<STextBlock> TextBlock = StaticCastSharedRef<STextBlock>(TestWidget);
+		TextBlock->SetWrapTextAt(WrapSize);
+		TextBlock->SetJustification(ETextJustify::Center);
+	}			
+	return BlockWidget;
+}
+
+void UVREditorUISystem::MakeRadialBoxMenu(const TSharedRef<FMultiBox>& MultiBox, const TSharedRef<SMultiBoxWidget>& MultiBoxWidget, float RadiusRatioOverride)
+{
+	// Grab the list of blocks, early out if there's nothing to fill the widget with
+	const TArray< TSharedRef< const FMultiBlock > >& Blocks = MultiBox->GetBlocks();
+	const uint32 NumItems = Blocks.Num();
+	if (NumItems == 0)
+	{
+		return;
+	}
+	MultiBox.Get().SetStyle(&FVREditorStyle::Get(), FVREditorStyle::GetStyleSetName());
+	TSharedRef<SRadialBox> GridPanel = SNew(SRadialBox)
+		.RadiusRatio(RadiusRatioOverride);
+	GridPanel->NumEntries = Blocks.Num();
+	for (const TSharedRef<const FMultiBlock>& MultiBlock : Blocks)
+	{
+		const TSharedRef<const FMultiBlock>& Block = MultiBlock;
+		TSharedRef<SWidget> BlockWidget = Block->MakeWidget(MultiBoxWidget, EMultiBlockLocation::Middle, true)->AsWidget();
+		TSharedRef<SOverlay> TestOverlay = SNew(SOverlay);
+		
+		GridPanel->AddSlot()
+		.VAlign(VAlign_Center)
+		.HAlign(HAlign_Center)
+		[
+			SAssignNew(TestOverlay, SOverlay)
+		];
+		
+		BlockWidget = AddHoverableButton(BlockWidget, TEXT("SButton"), TestOverlay);
+		BlockWidget = SetButtonTextWrap(BlockWidget, 75.0);
+
+	}
+	RadialWidget = GridPanel;
+	MultiBoxWidget->SetContent(
+			SNew(SDPIScaler)
+			.DPIScale(3)
+			[
+				SNew(SVerticalBox)
+				+ SVerticalBox::Slot()
+				.Padding(10, 10, 10, 10)
+				.VAlign(VAlign_Fill)
+				.HAlign(HAlign_Fill)
+				[
+					GridPanel
+				]
+			]
+		);
+}
+
+
+TSharedRef<SWidget> UVREditorUISystem::BuildQuickMenuWidget()
+{
+	FVREditorActionCallbacks::UpdateGizmoCoordinateSystemText(VRMode);
+	FVREditorActionCallbacks::UpdateGizmoModeText(VRMode);
+
+	FMultiBox::FOnMakeMultiBoxBuilderOverride VREditorMenuBuilderOverride;
+	VREditorMenuBuilderOverride.BindUObject(this, &UVREditorUISystem::MakeUniformGridMenu, 3);
+	const TSharedRef<FExtender> MenuExtender = MakeShareable(new FExtender());
+	TSharedPtr<FUICommandList> CommandList = MakeShareable(new FUICommandList());
+	FMenuBuilder SnapMenuBuilder(false, CommandList, MenuExtender);
+
+	// Snap menu
+	SnapMenuBuilder.AddMenuEntry(
+		LOCTEXT("ToggleTranslationSnap", "Translate"),
+		LOCTEXT("ToggleTranslationSnapTooltip", "Toggle Translation Snap"),
+		FSlateIcon(FEditorStyle::GetStyleSetName(), "EditorViewport.TranslateMode.Small"),
+		FUIAction
+		(
+			FExecuteAction::CreateStatic(&FLevelEditorActionCallbacks::LocationGridSnap_Clicked),
+			FCanExecuteAction::CreateStatic(&FLevelEditorActionCallbacks::DefaultCanExecuteAction),
+			FGetActionCheckState::CreateStatic(&FVREditorActionCallbacks::GetTranslationSnapState)
+		),
+		NAME_None,
+		EUserInterfaceActionType::Check
+		);
+	SnapMenuBuilder.AddMenuEntry(
+		LOCTEXT("ToggleRotationSnap", "Rotate"),
+		LOCTEXT("ToggleRotationSnapTooltip", "Toggle Rotation Snap"),
+		FSlateIcon(FEditorStyle::GetStyleSetName(), "EditorViewport.RotateMode.Small"),
+		FUIAction
+		(
+			FExecuteAction::CreateStatic(&FLevelEditorActionCallbacks::RotationGridSnap_Clicked),
+			FCanExecuteAction::CreateStatic(&FLevelEditorActionCallbacks::DefaultCanExecuteAction),
+			FGetActionCheckState::CreateStatic(&FVREditorActionCallbacks::GetRotationSnapState)
+			),
+		NAME_None,
+		EUserInterfaceActionType::Check
+		);
+	SnapMenuBuilder.AddMenuEntry(
+		LOCTEXT("ToggleScaleSnap", "Scale"),
+		LOCTEXT("ToggleScaleSnapTooltip", "Toggle Scale Snap"),
+		FSlateIcon(FEditorStyle::GetStyleSetName(), "EditorViewport.ScaleMode.Small"),
+		FUIAction
+		(
+			FExecuteAction::CreateStatic(&FLevelEditorActionCallbacks::ScaleGridSnap_Clicked),
+			FCanExecuteAction::CreateStatic(&FLevelEditorActionCallbacks::DefaultCanExecuteAction),
+			FGetActionCheckState::CreateStatic(&FVREditorActionCallbacks::GetScaleSnapState)
+			),
+		NAME_None,
+		EUserInterfaceActionType::Check
+		);
+	TAttribute<FText> DynamicTranslationSizeLabel;
+	DynamicTranslationSizeLabel.BindStatic(&FVREditorActionCallbacks::GetTranslationSnapSizeText);
+	SnapMenuBuilder.AddMenuEntry(
+		DynamicTranslationSizeLabel,
+		LOCTEXT("ToggleTranslationSnapSizeTooltip", "Toggle Translation Snap Size"),
+		FSlateIcon(),
+		FUIAction
+		(
+			FExecuteAction::CreateStatic(&FVREditorActionCallbacks::OnTranslationSnapSizeButtonClicked),
+			FCanExecuteAction::CreateStatic(&FLevelEditorActionCallbacks::DefaultCanExecuteAction)
+			)
+		);
+	TAttribute<FText> DynamicRotationSizeLabel;
+	DynamicRotationSizeLabel.BindStatic(&FVREditorActionCallbacks::GetRotationSnapSizeText);
+	SnapMenuBuilder.AddMenuEntry(
+		DynamicRotationSizeLabel,
+		LOCTEXT("ToggleRotationSnapSizeTooltip", "Toggle Rotation Snap Size"),
+		FSlateIcon(),
+		FUIAction
+		(
+			FExecuteAction::CreateStatic(&FVREditorActionCallbacks::OnRotationSnapSizeButtonClicked),
+			FCanExecuteAction::CreateStatic(&FLevelEditorActionCallbacks::DefaultCanExecuteAction)
+			)
+		);
+	TAttribute<FText> DynamicScaleSizeLabel;
+	DynamicScaleSizeLabel.BindStatic(&FVREditorActionCallbacks::GetScaleSnapSizeText);
+	SnapMenuBuilder.AddMenuEntry(
+		DynamicScaleSizeLabel,
+		LOCTEXT("ToggleScaleSnapSizeTooltip", "Toggle Scale Snap Size"),
+		FSlateIcon(),
+		FUIAction
+		(
+			FExecuteAction::CreateStatic(&FVREditorActionCallbacks::OnScaleSnapSizeButtonClicked),
+			FCanExecuteAction::CreateStatic(&FLevelEditorActionCallbacks::DefaultCanExecuteAction)
+			)
+		);
+
+	TSharedRef<SWidget> SnapMenuWidget = SnapMenuBuilder.MakeWidget(&VREditorMenuBuilderOverride);
+
+	// Next Menu section: Interface
+	CommandList = MakeShareable(new FUICommandList());
+	FMenuBuilder InterfaceMenuBuilder(false, CommandList, MenuExtender);
+
+	InterfaceMenuBuilder.AddMenuEntry(
+		LOCTEXT("ActorDetails", "Details"),
+		LOCTEXT("ActorDetailsTooltip", "Details"),
+		FSlateIcon(FEditorStyle::GetStyleSetName(), "LevelEditor.Tabs.Details"),
+		FUIAction
+		(
+			FExecuteAction::CreateStatic(&FVREditorActionCallbacks::OnUIToggleButtonClicked, VRMode, UVREditorUISystem::EEditorUIPanel::ActorDetails),
+			FCanExecuteAction::CreateStatic(&FLevelEditorActionCallbacks::DefaultCanExecuteAction),
+			FGetActionCheckState::CreateStatic(&FVREditorActionCallbacks::GetUIToggledState, VRMode, UVREditorUISystem::EEditorUIPanel::ActorDetails)
+			),
+		NAME_None,
+		EUserInterfaceActionType::Check
+		);
+	InterfaceMenuBuilder.AddMenuEntry(
+		LOCTEXT("ContentBrowser", "Content Browser"),
+		LOCTEXT("ContentBrowserTooltip", "Content Browser"),
+		FSlateIcon(FEditorStyle::GetStyleSetName(), "LevelEditor.Tabs.ContentBrowser"),
+		FUIAction
+		(
+			FExecuteAction::CreateStatic(&FVREditorActionCallbacks::OnUIToggleButtonClicked, VRMode, UVREditorUISystem::EEditorUIPanel::ContentBrowser),
+			FCanExecuteAction::CreateStatic(&FLevelEditorActionCallbacks::DefaultCanExecuteAction),
+			FGetActionCheckState::CreateStatic(&FVREditorActionCallbacks::GetUIToggledState, VRMode, UVREditorUISystem::EEditorUIPanel::ContentBrowser)
+			),
+		NAME_None,
+		EUserInterfaceActionType::Check
+		);
+	InterfaceMenuBuilder.AddMenuEntry(
+		LOCTEXT("Modes", "Modes"),
+		LOCTEXT("ModesTooltip", "Modes"),
+		FSlateIcon(FEditorStyle::GetStyleSetName(), "LevelEditor.Tabs.Modes"),
+		FUIAction
+		(
+			FExecuteAction::CreateStatic(&FVREditorActionCallbacks::OnUIToggleButtonClicked, VRMode, UVREditorUISystem::EEditorUIPanel::Modes),
+			FCanExecuteAction::CreateStatic(&FLevelEditorActionCallbacks::DefaultCanExecuteAction),
+			FGetActionCheckState::CreateStatic(&FVREditorActionCallbacks::GetUIToggledState, VRMode, UVREditorUISystem::EEditorUIPanel::Modes)
+			),
+		NAME_None,
+		EUserInterfaceActionType::Check
+		);
+	InterfaceMenuBuilder.AddMenuEntry(
+		LOCTEXT("WorldOutliner", "World Outliner"),
+		LOCTEXT("WorldOutlinerTooltip", "World Outliner"),
+		FSlateIcon(FEditorStyle::GetStyleSetName(), "LevelEditor.Tabs.Outliner"),
+		FUIAction
+		(
+			FExecuteAction::CreateStatic(&FVREditorActionCallbacks::OnUIToggleButtonClicked, VRMode, UVREditorUISystem::EEditorUIPanel::WorldOutliner),
+			FCanExecuteAction::CreateStatic(&FLevelEditorActionCallbacks::DefaultCanExecuteAction),
+			FGetActionCheckState::CreateStatic(&FVREditorActionCallbacks::GetUIToggledState, VRMode, UVREditorUISystem::EEditorUIPanel::WorldOutliner)
+			),
+		NAME_None,
+		EUserInterfaceActionType::Check
+		);
+	InterfaceMenuBuilder.AddMenuEntry(
+		LOCTEXT("AssetEditor", "Asset Editor"),
+		LOCTEXT("AssetEditorTooltip", "Asset Editor"),
+		FSlateIcon(),
+		FUIAction
+		(
+			FExecuteAction::CreateStatic(&FVREditorActionCallbacks::OnUIToggleButtonClicked, VRMode, UVREditorUISystem::EEditorUIPanel::AssetEditor),
+			FCanExecuteAction::CreateStatic(&FLevelEditorActionCallbacks::DefaultCanExecuteAction),
+			FGetActionCheckState::CreateStatic(&FVREditorActionCallbacks::GetUIToggledState, VRMode, UVREditorUISystem::EEditorUIPanel::AssetEditor)
+			),
+		NAME_None,
+		EUserInterfaceActionType::Check
+		);
+	InterfaceMenuBuilder.AddMenuEntry(
+		LOCTEXT("Tutorial", "Tutorial"),
+		LOCTEXT("TutorialTooltip", "Tutorial"),
+		FSlateIcon(FEditorStyle::GetStyleSetName(), "Tutorials.Browser.DefaultTutorialIcon"),
+		FUIAction
+		(
+			FExecuteAction::CreateStatic(&FVREditorActionCallbacks::OnUIToggleButtonClicked, VRMode, UVREditorUISystem::EEditorUIPanel::Tutorial),
+			FCanExecuteAction::CreateStatic(&FLevelEditorActionCallbacks::DefaultCanExecuteAction),
+			FGetActionCheckState::CreateStatic(&FVREditorActionCallbacks::GetUIToggledState, VRMode, UVREditorUISystem::EEditorUIPanel::Tutorial)
+			),
+		NAME_None,
+		EUserInterfaceActionType::Check
+		);
+	InterfaceMenuBuilder.AddMenuEntry(
+		LOCTEXT("WorldSettings", "World Settings"),
+		LOCTEXT("WorldSettingsTooltip", "World Settings"),
+		FSlateIcon(FEditorStyle::GetStyleSetName(), "LevelEditor.WorldProperties.Tab"),
+		FUIAction
+		(
+			FExecuteAction::CreateStatic(&FVREditorActionCallbacks::OnUIToggleButtonClicked, VRMode, UVREditorUISystem::EEditorUIPanel::WorldSettings),
+			FCanExecuteAction::CreateStatic(&FLevelEditorActionCallbacks::DefaultCanExecuteAction),
+			FGetActionCheckState::CreateStatic(&FVREditorActionCallbacks::GetUIToggledState, VRMode, UVREditorUISystem::EEditorUIPanel::WorldSettings)
+			),
+		NAME_None,
+		EUserInterfaceActionType::Check
+		);
+	InterfaceMenuBuilder.AddMenuEntry(
+		LOCTEXT("ToggleAllUI", "Toggle All UI"),
+		LOCTEXT("ToggleAllUITooltip", "ToggleAllUI"),
+		FSlateIcon(),
+		FUIAction
+		(
+			FExecuteAction::CreateUObject(this, &UVREditorUISystem::TogglePanelsVisibility),
+			FCanExecuteAction::CreateStatic(&FLevelEditorActionCallbacks::DefaultCanExecuteAction)
+			)
+		);
+	
+
+	TSharedRef<SWidget> InterfaceMenuWidget = InterfaceMenuBuilder.MakeWidget(&VREditorMenuBuilderOverride);
+
+	// Next Menu section: Gizmo
+	CommandList = MakeShareable(new FUICommandList());
+	FMenuBuilder GizmoMenuBuilder(false, CommandList, MenuExtender);
+
+	TAttribute<FText> DynamicGizmoCoordinateSystemLabel;
+	DynamicGizmoCoordinateSystemLabel.BindStatic(&FVREditorActionCallbacks::GetGizmoCoordinateSystemText);
+	GizmoMenuBuilder.AddMenuEntry(
+		DynamicGizmoCoordinateSystemLabel,
+		LOCTEXT("GizmoCoordinateTooltip", "Toggle Gizmo Coordinates"),
+		FSlateIcon(),
+		FUIAction
+		(
+			FExecuteAction::CreateStatic(&FVREditorActionCallbacks::OnGizmoCoordinateSystemButtonClicked, VRMode),
+			FCanExecuteAction::CreateStatic(&FLevelEditorActionCallbacks::DefaultCanExecuteAction)
+			)
+		);
+	TAttribute<FText> DynamicGizmoModeLabel;
+	DynamicGizmoModeLabel.BindStatic(&FVREditorActionCallbacks::GetGizmoModeText);
+	GizmoMenuBuilder.AddMenuEntry(
+		DynamicGizmoModeLabel,
+		LOCTEXT("GizmoModeTooltip", "Toggle Gizmo Mode"),
+		FSlateIcon(),
+		FUIAction
+		(
+			FExecuteAction::CreateStatic(&FVREditorActionCallbacks::OnGizmoModeButtonClicked, VRMode),
+			FCanExecuteAction::CreateStatic(&FLevelEditorActionCallbacks::DefaultCanExecuteAction)
+			)
+		);
+	GizmoMenuBuilder.AddMenuEntry(
+		LOCTEXT("CycleGizmo", "Cycle Gizmo"),
+		LOCTEXT("CycleGizmoTooltip", "Cycle Gizmo"),
+		FSlateIcon(),
+		FUIAction
+		(
+			FExecuteAction::CreateStatic(&FVREditorActionCallbacks::OnGizmoCycle, VRMode),
+			FCanExecuteAction::CreateStatic(&FLevelEditorActionCallbacks::DefaultCanExecuteAction)
+			)
+		);
+
+	TSharedRef<SWidget> GizmoMenuWidget = GizmoMenuBuilder.MakeWidget(&VREditorMenuBuilderOverride);
+
+	// Next Menu section: Extras
+	CommandList = MakeShareable(new FUICommandList());
+	FMenuBuilder ExtrasMenuBuilder(false, CommandList, MenuExtender);
+
+
+	ExtrasMenuBuilder.AddMenuEntry(
+		LOCTEXT("Screenshot", "Screenshot"),
+		LOCTEXT("ScreenshotTooltip", "Screenshot"),
+		FSlateIcon(),
+		FUIAction
+		(
+			FExecuteAction::CreateStatic(&FVREditorActionCallbacks::OnScreenshotButtonClicked, VRMode),
+			FCanExecuteAction::CreateStatic(&FLevelEditorActionCallbacks::DefaultCanExecuteAction)
+			)
+		);
+	ExtrasMenuBuilder.AddMenuEntry(
+		LOCTEXT("Flashlight", "Flashlight"),
+		LOCTEXT("FlashlightTooltip", "Flashlight"),
+		FSlateIcon(),
+		FUIAction
+		(
+			FExecuteAction::CreateStatic(&FVREditorActionCallbacks::OnLightButtonClicked, VRMode),
+			FCanExecuteAction::CreateStatic(&FLevelEditorActionCallbacks::DefaultCanExecuteAction),
+			FGetActionCheckState::CreateStatic(&FVREditorActionCallbacks::GetFlashlightState, VRMode)
+			),
+		NAME_None,
+		EUserInterfaceActionType::Check
+		);
+	ExtrasMenuBuilder.AddMenuEntry(
+		LOCTEXT("PlayInEditor", "Play"),
+		LOCTEXT("PlayTooltip", "Play"),
+		FSlateIcon(),
+		FUIAction
+		(
+			FExecuteAction::CreateStatic(&FVREditorActionCallbacks::OnPlayButtonClicked, VRMode),
+			FCanExecuteAction::CreateStatic(&FLevelEditorActionCallbacks::DefaultCanExecuteAction)
+			)
+		);
+	TSharedRef<SWidget> ExtrasMenuWidget = ExtrasMenuBuilder.MakeWidget(&VREditorMenuBuilderOverride);
+	const FSlateBrush* NoBorderBrush = FStyleDefaults::GetNoBrush();
+	TSharedRef<SWidget> MenuWidget = SNew(SBox)
+		.VAlign(VAlign_Top)
+		.HAlign(HAlign_Center)
+		[
+			SNew(SDPIScaler)
+			.DPIScale(3)
+			[
+				SNew(SVerticalBox)
+				+SVerticalBox::Slot()
+				.HAlign(HAlign_Fill)
+				.VAlign(VAlign_Bottom)
+				[
+					SNew(STextBlock)
+					.Justification(ETextJustify::Center)
+					.Text(LOCTEXT("GridSnap", "Grid Snap"))
+					.TextStyle(FVREditorStyle::Get(),"VREditorStyle.Heading")
+				]
+				+ SVerticalBox::Slot()
+				.HAlign(HAlign_Fill)
+				.VAlign(VAlign_Top)
+				.AutoHeight()
+				[
+					SnapMenuWidget
+				]
+				+ SVerticalBox::Slot()
+				.HAlign(HAlign_Fill)
+				.VAlign(VAlign_Bottom)
+				[
+					SNew(STextBlock)
+					.Justification(ETextJustify::Center)
+					.Text(LOCTEXT("Interface", "Interface"))
+					.TextStyle(FVREditorStyle::Get(), "VREditorStyle.Heading")
+				]
+				+ SVerticalBox::Slot()
+				.HAlign(HAlign_Fill)
+				.VAlign(VAlign_Top)
+				.AutoHeight()
+				[
+					InterfaceMenuWidget
+				]
+				+ SVerticalBox::Slot()
+				.HAlign(HAlign_Fill)
+				.VAlign(VAlign_Bottom)
+				[
+					SNew(STextBlock)
+					.Justification(ETextJustify::Center)
+					.Text(LOCTEXT("VRTools", "Tools"))
+					.TextStyle(FVREditorStyle::Get(), "VREditorStyle.Heading")
+				]
+				+ SVerticalBox::Slot()
+				.HAlign(HAlign_Fill)
+				.VAlign(VAlign_Top)
+				.AutoHeight()
+				[
+					ExtrasMenuWidget
+				]
+				+ SVerticalBox::Slot()
+				.HAlign(HAlign_Fill)
+				.VAlign(VAlign_Bottom)
+				[
+					SNew(STextBlock)
+					.Justification(ETextJustify::Center)
+					.Text(LOCTEXT("Gizmo", "Gizmo"))
+					.TextStyle(FVREditorStyle::Get(), "VREditorStyle.Heading")
+				]
+				+ SVerticalBox::Slot()
+				.HAlign(HAlign_Fill)
+				.VAlign(VAlign_Top)
+				.AutoHeight()
+				[
+					GizmoMenuWidget
+				]
+			]
+		]
+	;
+
+
+	return MenuWidget;
+}
+
+
+
+TSharedRef<SWidget> UVREditorUISystem::BuildRadialMenuWidget()
+{
+
+	FMultiBox::FOnMakeMultiBoxBuilderOverride VREditorMenuBuilderOverride;
+	VREditorMenuBuilderOverride.BindUObject(this, &UVREditorUISystem::MakeRadialBoxMenu, 1.0f);
+	const TSharedRef<FExtender> MenuExtender = MakeShareable(new FExtender());
+	TSharedPtr<FUICommandList> CommandList = MakeShareable(new FUICommandList());
+	FMenuBarBuilder MenuBuilder(CommandList, MenuExtender);
+
+	FVREditorActionCallbacks::UpdateGizmoCoordinateSystemText(VRMode);
+	FVREditorActionCallbacks::UpdateGizmoModeText(VRMode);
+
+	// First menu entry is at 90 degrees 
+	MenuBuilder.AddMenuEntry(
+		LOCTEXT( "Delete", "Delete" ),
+		LOCTEXT( "DeleteTooltip", "Delete" ),
+		FSlateIcon(),
+		FUIAction
+		(
+			FExecuteAction::CreateStatic( &FLevelEditorActionCallbacks::ExecuteExecCommand, FString( TEXT( "DELETE" ) ) ),
+			FCanExecuteAction::CreateStatic( &FLevelEditorActionCallbacks::Delete_CanExecute )
+		)
+	);
+	MenuBuilder.AddMenuEntry(
+		LOCTEXT( "Cut", "Cut" ),
+		LOCTEXT( "CutTooltip", "Cut" ),
+		FSlateIcon(),
+		FUIAction
+		(
+			FExecuteAction::CreateStatic( &FLevelEditorActionCallbacks::ExecuteExecCommand, FString( TEXT( "EDIT CUT" ) ) ),
+			FCanExecuteAction::CreateStatic( &FLevelEditorActionCallbacks::Cut_CanExecute )
+		)
+	);
+	MenuBuilder.AddMenuEntry(
+		LOCTEXT( "Copy", "Copy" ),
+		LOCTEXT( "CopyTooltip", "Copy" ),
+		FSlateIcon(),
+		FUIAction
+		(
+			FExecuteAction::CreateStatic( &FLevelEditorActionCallbacks::ExecuteExecCommand, FString( TEXT( "EDIT COPY" ) ) ),
+			FCanExecuteAction::CreateStatic( &FLevelEditorActionCallbacks::Copy_CanExecute )
+		)
+	);
+	MenuBuilder.AddMenuEntry(
+		LOCTEXT( "Duplicate", "Duplicate" ),
+		LOCTEXT( "DuplicateTooltip", "Duplicate" ),
+		FSlateIcon(),
+		FUIAction
+		(
+			FExecuteAction::CreateStatic( &FLevelEditorActionCallbacks::ExecuteExecCommand, FString( TEXT( "DUPLICATE" ) ) ),
+			FCanExecuteAction::CreateStatic( &FLevelEditorActionCallbacks::Duplicate_CanExecute )
+		)
+	);
+	MenuBuilder.AddMenuEntry(
+		LOCTEXT( "Paste", "Paste" ),
+		LOCTEXT( "PasteTooltip", "Paste" ),
+		FSlateIcon(),
+		FUIAction
+		(
+			FExecuteAction::CreateStatic( &FLevelEditorActionCallbacks::ExecuteExecCommand, FString( TEXT( "EDIT PASTE" ) ) ),
+			FCanExecuteAction::CreateStatic( &FLevelEditorActionCallbacks::Paste_CanExecute )
+		)
+	);
+	MenuBuilder.AddMenuEntry(
+		LOCTEXT( "Undo", "Undo" ),
+		LOCTEXT( "UndoTooltip", "Undo" ),
+		FSlateIcon(),
+		FUIAction
+		(
+			FExecuteAction::CreateStatic( &FLevelEditorActionCallbacks::ExecuteExecCommand, FString( TEXT( "TRANSACTION UNDO" ) ) ),
+			FCanExecuteAction::CreateStatic( &FLevelEditorActionCallbacks::DefaultCanExecuteAction )
+		)
+	);
+	MenuBuilder.AddMenuEntry(
+		LOCTEXT( "SnapActorsToGround", "Snap Actors To Ground" ),
+		LOCTEXT( "SnapActorsToGroundTooltip", "Snap Actors To Ground" ),
+		FSlateIcon(),
+		FUIAction
+		(
+			FExecuteAction::CreateStatic( &FVREditorActionCallbacks::OnSnapActorsToGroundClicked, VRMode ),
+			FCanExecuteAction::CreateStatic( &FLevelEditorActionCallbacks::Copy_CanExecute )
+		)
+	);
+	MenuBuilder.AddMenuEntry(
+		LOCTEXT( "Redo", "Redo" ),
+		LOCTEXT( "RedoTooltip", "Redo" ),
+		FSlateIcon(),
+		FUIAction
+		(
+			FExecuteAction::CreateStatic( &FLevelEditorActionCallbacks::ExecuteExecCommand, FString( TEXT( "TRANSACTION REDO" ) ) ),
+			FCanExecuteAction::CreateStatic( &FLevelEditorActionCallbacks::DefaultCanExecuteAction )
+		)
+	);
+
+	// Create the menu widget
+	TSharedRef<SWidget> MenuWidget = MenuBuilder.MakeWidget(&VREditorMenuBuilderOverride);
+	return MenuWidget;
+}
+
+TSharedRef<SWidget> UVREditorUISystem::BuildNumPadWidget()
+{
+
+	FMultiBox::FOnMakeMultiBoxBuilderOverride VREditorMenuBuilderOverride;
+	VREditorMenuBuilderOverride.BindUObject(this, &UVREditorUISystem::MakeRadialBoxMenu, 0.5f);
+	const TSharedRef<FExtender> MenuExtender = MakeShareable(new FExtender());
+	TSharedPtr<FUICommandList> CommandList = MakeShareable(new FUICommandList());
+	FMenuBarBuilder MenuBuilder(CommandList, MenuExtender);
+
+	// First menu entry is at 90 degrees 
+
+	MenuBuilder.AddMenuEntry(
+		LOCTEXT("NumThree", "3"),
+		LOCTEXT("", ""),
+		FSlateIcon(),
+		FUIAction
+		(
+			FExecuteAction::CreateStatic(&FVREditorActionCallbacks::SimulateCharacterEntry, FString::FromInt(3))
+		)
+		);
+	MenuBuilder.AddMenuEntry(
+		LOCTEXT("NumFour", "4"),
+		LOCTEXT("", ""),
+		FSlateIcon(),
+		FUIAction
+		(
+			FExecuteAction::CreateStatic(&FVREditorActionCallbacks::SimulateCharacterEntry, FString::FromInt(4))
+		)
+		);
+	MenuBuilder.AddMenuEntry(
+		LOCTEXT("NumFive", "5"),
+		LOCTEXT("", ""),
+		FSlateIcon(),
+		FUIAction
+		(
+			FExecuteAction::CreateStatic(&FVREditorActionCallbacks::SimulateCharacterEntry, FString::FromInt(5))
+		)
+		);
+	MenuBuilder.AddMenuEntry(
+		LOCTEXT("NumSix", "6"),
+		LOCTEXT("", ""),
+		FSlateIcon(),
+		FUIAction
+		(
+			FExecuteAction::CreateStatic(&FVREditorActionCallbacks::SimulateCharacterEntry, FString::FromInt(6))
+		)
+		);
+	MenuBuilder.AddMenuEntry(
+		LOCTEXT("NumSeven", "7"),
+		LOCTEXT("", ""),
+		FSlateIcon(),
+		FUIAction
+		(
+			FExecuteAction::CreateStatic(&FVREditorActionCallbacks::SimulateCharacterEntry, FString::FromInt(7))
+			)
+		);
+	MenuBuilder.AddMenuEntry(
+		LOCTEXT("NumEight", "8"),
+		LOCTEXT("", ""),
+		FSlateIcon(),
+		FUIAction
+		(
+			FExecuteAction::CreateStatic(&FVREditorActionCallbacks::SimulateCharacterEntry, FString::FromInt(8))
+			)
+		);
+	MenuBuilder.AddMenuEntry(
+		LOCTEXT("NumNine", "9"),
+		LOCTEXT("", ""),
+		FSlateIcon(),
+		FUIAction
+		(
+			FExecuteAction::CreateStatic(&FVREditorActionCallbacks::SimulateCharacterEntry, FString::FromInt(9))
+		)
+		);
+	MenuBuilder.AddMenuEntry(
+		LOCTEXT( "Hyphen", "-" ),
+		LOCTEXT( "", "" ),
+		FSlateIcon(),
+		FUIAction
+		(
+			FExecuteAction::CreateStatic( &FVREditorActionCallbacks::SimulateCharacterEntry, FString::FString( TEXT( "-" ) ) )
+		)
+	);
+	MenuBuilder.AddMenuEntry(
+		LOCTEXT( "Decimal", "." ),
+		LOCTEXT( "", "" ),
+		FSlateIcon(),
+		FUIAction
+		(
+			FExecuteAction::CreateStatic( &FVREditorActionCallbacks::SimulateCharacterEntry, FString::FString( TEXT( "." ) ) )
+		)
+	);
+	MenuBuilder.AddMenuEntry(
+		LOCTEXT( "NumZero", "0" ),
+		LOCTEXT( "", "" ),
+		FSlateIcon(),
+		FUIAction
+		(
+			FExecuteAction::CreateStatic( &FVREditorActionCallbacks::SimulateCharacterEntry, FString::FromInt( 0 ) )
+		)
+	);
+	MenuBuilder.AddMenuEntry(
+		LOCTEXT( "NumOne", "1" ),
+		LOCTEXT( "", "" ),
+		FSlateIcon(),
+		FUIAction
+		(
+			FExecuteAction::CreateStatic( &FVREditorActionCallbacks::SimulateCharacterEntry, FString::FromInt( 1 ) )
+		)
+	);
+	MenuBuilder.AddMenuEntry(
+		LOCTEXT( "NumTwo", "2" ),
+		LOCTEXT( "", "" ),
+		FSlateIcon(),
+		FUIAction
+		(
+			FExecuteAction::CreateStatic( &FVREditorActionCallbacks::SimulateCharacterEntry, FString::FromInt( 2 ) )
+		)
+	);
+
+	// Create the menu widget
+	TSharedRef<SWidget> MenuWidget = MenuBuilder.MakeWidget(&VREditorMenuBuilderOverride);
+
+	TSharedRef<SWidget> FinalWidget = SNew( SOverlay )
+		+ SOverlay::Slot()
+		[
+			MenuWidget
+		]
+		+ SOverlay::Slot()
+		.HAlign( HAlign_Center )
+		.VAlign( VAlign_Center )
+		[
+			SNew(SVerticalBox)
+			+ SVerticalBox::Slot()
+			.HAlign( HAlign_Fill )
+			.VAlign( VAlign_Center )
+			.AutoHeight()
+			[
+				SNew( STextBlock )
+				.Justification( ETextJustify::Center )
+				.Text( LOCTEXT( "BackspaceHelp", "Modifier: Backspace" ) )
+				.TextStyle( FVREditorStyle::Get(), "VREditorStyle.HelperText" )
+			]
+		+ SVerticalBox::Slot()
+			.HAlign( HAlign_Fill )
+			.VAlign( VAlign_Center )
+			.AutoHeight()
+			[
+				SNew( STextBlock )
+				.Justification( ETextJustify::Center )
+				.Text( LOCTEXT( "EnterHelp", "Grip: Enter" ) )
+			.	TextStyle( FVREditorStyle::Get(), "VREditorStyle.HelperText" )
+			]
+		]
+		;
+
+	return FinalWidget;
+}
+
+const TSharedRef<SWidget>& UVREditorUISystem::FindWidgetOfType(const TSharedRef<SWidget>& Content, FName WidgetType)
+{
+	if (Content->GetType() == WidgetType)
+	{
+		return Content;
+	}
+
+	FChildren* Children = Content->GetChildren();
+	int32 NumChildren = Children->Num();
+
+	for (int32 Index = 0; Index < NumChildren; ++Index)
+	{
+		const TSharedRef<SWidget>& Found = FindWidgetOfType((Children->GetChildAt(Index)), WidgetType);
+		if (Found != SNullWidget::NullWidget)
+		{
+			return Found;
+		}
+	}
+	return SNullWidget::NullWidget;
+}
+
+void UVREditorUISystem::OnHoverBeginEffect(TSharedRef<SButton> Button)
+{
+	FVRButton* ButtonToAnimate = VRButtons.FindByPredicate([&Button](const FVRButton& ButtonData)
+	{
+		return Button == ButtonData.Button;
+	});
+
+	check(ButtonToAnimate);
+
+	//Set the newly hovered button's animation state to Forward
+	ButtonToAnimate->AnimationDirection = EVREditorAnimationState::Forward;
+
+	// Put the button in the higher z-order overlay slot
+	const TSharedRef<SWidget>& PlaceholderWidget = ButtonToAnimate->HoverSlot->DetachWidget().ToSharedRef();
+	const TSharedRef<SWidget>& WidgetToTransfer = ButtonToAnimate->BaseSlot->DetachWidget().ToSharedRef();
+	ButtonToAnimate->HoverSlot->AttachWidget(WidgetToTransfer);
+	ButtonToAnimate->BaseSlot->AttachWidget(PlaceholderWidget);
+}
+
+void UVREditorUISystem::OnHoverEndEffect(TSharedRef<SButton> Button)
+{
+	FVRButton* ButtonToAnimate = VRButtons.FindByPredicate([&Button](const FVRButton& ButtonData)
+	{
+		return Button == ButtonData.Button;
+	});
+
+	check(ButtonToAnimate);
+
+	// Set the unhovered button's animation state to Backward
+	ButtonToAnimate->AnimationDirection = EVREditorAnimationState::Backward;
+	// Put the button back in the lower z-order slot
+	TSharedRef<SWidget> PlaceholderWidget = (ButtonToAnimate->BaseSlot->DetachWidget()).ToSharedRef();
+	TSharedRef<SWidget> WidgetToTransfer = (ButtonToAnimate->HoverSlot->DetachWidget()).ToSharedRef();
+	ButtonToAnimate->BaseSlot->AttachWidget(WidgetToTransfer);
+	ButtonToAnimate->HoverSlot->AttachWidget(PlaceholderWidget);
+}
+
+const TSharedPtr<SWidget>& UVREditorUISystem::GetRadialWidget()
+{
+	return RadialWidget;
+}
+
+void UVREditorUISystem::SwapRadialMenu()
+{
+	const AVREditorFloatingUI::EDockedTo RadialMenuInteractor = QuickRadialMenu->GetDockedTo();
+	const FIntPoint DefaultResolution( VREd::DefaultEditorUIResolutionX->GetInt(), VREd::DefaultEditorUIResolutionY->GetInt() );
+	if (bRadialMenuIsNumpad)
+	{
+		QuickRadialMenu->SetSlateWidget( *this, BuildRadialMenuWidget(), DefaultResolution, 40.0f, RadialMenuInteractor );
+		bRadialMenuIsNumpad = false;
+	}
+	else
+	{
+		QuickRadialMenu->SetSlateWidget( *this, BuildNumPadWidget(), DefaultResolution, 40.0f, RadialMenuInteractor );
+		bRadialMenuIsNumpad = true;
+	}
+	QuickRadialMenu->GetWidgetComponent()->SetOpacityFromTexture( 1.0f );
+	QuickRadialMenu->GetWidgetComponent()->SetBackgroundColor( FLinearColor::Transparent );
+	QuickRadialMenu->GetWidgetComponent()->SetBlendMode( EWidgetBlendMode::Masked );
+}
+
+#undef LOCTEXT_NAMESPACE
\ No newline at end of file
diff --git a/Engine/Source/Editor/VREditor/UI/VREditorUISystem.h b/Engine/Source/Editor/VREditor/UI/VREditorUISystem.h
--- a/Engine/Source/Editor/VREditor/UI/VREditorUISystem.h
+++ b/Engine/Source/Editor/VREditor/UI/VREditorUISystem.h
@@ -15,6 +15,69 @@ class SColorPicker;
 class UViewportInteractor;
 class UVREditorInteractor;
 
+
+/** Stores the animation playback state of a VR UI element */
+enum class EVREditorAnimationState : uint8
+{
+	None,
+	Forward,
+	Backward
+};
+
+/** Structure to keep track of all relevant interaction and animation elements of a VR Button */
+struct FVRButton
+{
+	/** Pointer to button */
+	TSharedPtr<SButton> Button;
+
+	/** Pointer to slot in overlay with lower z-order */
+	SOverlay::FOverlaySlot* BaseSlot;
+
+	/** Pointer to slot in overlay with higher z-order */
+	SOverlay::FOverlaySlot* HoverSlot;
+
+	/** Pointer to button border */
+	TSharedPtr<SBorder> ButtonBorder;
+
+	/** Animation playback state of the button */
+	EVREditorAnimationState AnimationDirection;
+
+	/** Current scale of the button element */
+	float CurrentScale;
+
+	/** Minimum scale of the button element */
+	float MinScale;
+
+	/** Maximum scale of the button element */
+	float MaxScale;
+
+	/** Rate at which the button changes scale. Currently the same for scaling up and scaling down. */
+	float ScaleRate;
+
+	FVRButton()
+		: AnimationDirection(EVREditorAnimationState::None),
+		CurrentScale(1.0f),
+		MinScale(1.0f),
+		MaxScale(1.10f),
+		ScaleRate(2.0f)
+		{}
+
+	FVRButton(TSharedPtr<SButton> InButton, SOverlay::FOverlaySlot* InBaseSlot, SOverlay::FOverlaySlot* InHoverSlot, TSharedPtr<SBorder> InButtonBorder,
+		EVREditorAnimationState InAnimationDirection = EVREditorAnimationState::None, float InCurrentScale = 1.0f, float InMinScale = 1.0f, float InMaxScale = 1.25f, float InScaleRate = 2.0f)
+		: Button(InButton),
+		BaseSlot(InBaseSlot),
+		HoverSlot(InHoverSlot),
+		ButtonBorder(InButtonBorder),
+		AnimationDirection(InAnimationDirection),
+		CurrentScale(InCurrentScale),
+		MinScale(InMinScale),
+		MaxScale(InMaxScale),
+		ScaleRate(InScaleRate)
+		{}
+
+};
+
+
 /**
  * VR Editor user interface manager
  */
@@ -86,13 +149,10 @@ public:
 	void ShowEditorUIPanel( class AVREditorFloatingUI* Panel, UVREditorInteractor* Interactor, const bool bShouldShow, const bool OnHand = false, const bool bRefreshQuickMenu = true, const bool bPlaySound = true );
 
 	/** Returns true if the radial menu is visible on this hand */
-	bool IsShowingRadialMenu( UVREditorInteractor* Interactor ) const;
+	bool IsShowingRadialMenu(const UVREditorInteractor* Interactor ) const;
 
 	/** Tries to spawn the radial menu (if the specified hand isn't doing anything else) */
-	void TryToSpawnRadialMenu( UVREditorInteractor* Interactor );
-
-	/** Updates the radial menu with the current hand trackpad location */
-	void UpdateRadialMenu( UVREditorInteractor* Interactor );
+	void TryToSpawnRadialMenu( UVREditorInteractor* Interactor, const bool bForceOverUI );
 
 	/** Hides the radial menu if the specified hand is showing it */
 	void HideRadialMenu( UVREditorInteractor* Interactor );
@@ -129,6 +189,17 @@ public:
 	/** Toggles the visibility of the panel, if the panel is in room space it will be hidden and docked to nothing */
 	void TogglePanelVisibility( const EEditorUIPanel EditorUIPanel );
 
+	/** Returns the radial widget so other classes, like the interactors, can access its functionality */
+	const TSharedPtr<SWidget>& GetRadialWidget();
+
+	/** 
+	 * Finds a widget with a given name inside the Content argument 
+	 * @param Content The widget to begin searching in
+	 * @param The FName of the widget type to find. 
+	 */
+	static const TSharedRef<SWidget>& FindWidgetOfType(const TSharedRef<SWidget>& Content, FName WidgetType);
+
+
 protected:
 
 	/** Called when the user presses a button on their motion controller device */
@@ -166,6 +237,30 @@ protected:
 	/** Hides the VR-specific color picker. Gets bound to SColorPicker's destruction override delegate */
 	void DestroyVRColorPicker();
 
+	/**  Makes a uniform grid widget from the menu information contained in a MultiBox and MultiBoxWidget */
+	void MakeUniformGridMenu( const TSharedRef<FMultiBox>& MultiBox, const TSharedRef<SMultiBoxWidget>& MultiBoxWidget, int32 Columns);
+	/**  Makes a radial box widget from the menu information contained in a MultiBox and MultiBoxWidget */
+	void MakeRadialBoxMenu(const TSharedRef<FMultiBox>& MultiBox, const TSharedRef<SMultiBoxWidget>& MultiBoxWidget, float RadiusRatioOverride);
+
+	/** Adds a hoverable button of a given type to an overlay, using menu data from a BlockWidget */
+	TSharedRef<SWidget> AddHoverableButton(TSharedRef<SWidget>& BlockWidget, FName ButtonType, TSharedRef<SOverlay>& TestOverlay);
+	/** Sets the text wrap size of the text block element nested in a BlockWidget */
+	TSharedRef<SWidget> SetButtonTextWrap(TSharedRef<SWidget>& BlockWidget, float WrapSize);
+	
+	/** Builds the quick menu Slate widget */
+	TSharedRef<SWidget> BuildQuickMenuWidget();
+	/** Builds the radial menu Slate widget */
+	TSharedRef<SWidget> BuildRadialMenuWidget();
+	/** Builds the numpad Slate widget */
+	TSharedRef<SWidget> BuildNumPadWidget();
+	/** Swaps the content of the radial menu between the radial menu and the numpad */
+	void SwapRadialMenu();
+
+	/** Called when a laser or simulated mouse hover enters a button */
+	void OnHoverBeginEffect(TSharedRef<SButton> Button);
+	/** Called when a laser or simulated mouse hover leaves a button */
+	void OnHoverEndEffect(TSharedRef<SButton> Button);
+
 
 protected:
 
@@ -192,14 +287,17 @@ protected:
 	/** The time since the radial menu was updated */
 	float RadialMenuHideDelayTime;
 
-	/** Quick menu widget class */
-	UClass* QuickMenuWidgetClass;
-
-	/** Quick radial menu widget class */
-	UClass* QuickRadialWidgetClass;
-
 	/** Tutorial menu widget class */
 	UClass* TutorialWidgetClass;
+
+	/** Pointer to the radial menu widget */
+	TSharedPtr<SWidget> RadialWidget;
+
+	/** True if the radial menu was visible when the content was swapped */
+	bool bRadialMenuVisibleAtSwap;
+
+	/** True if the radial menu is currently displaying the numpad */
+	bool bRadialMenuIsNumpad;
 	
 	//
 	// Dragging UI
@@ -250,6 +348,10 @@ protected:
 	UPROPERTY()
 	class USoundCue* ShowUISound;
 
+	/** Button press sound */
+	UPROPERTY()
+	class USoundCue* ButtonPressSound;
+
 	/** If the current dragged dock passed a certain distance if dragged from a hand */
 	bool bDraggedDockFromHandPassedThreshold;
 
@@ -259,6 +361,9 @@ protected:
 	/** Default transforms */
 	TArray<FTransform> DefaultWindowTransforms;
 
+	/** All buttons created for the radial and quick menus */
+	TArray<FVRButton> VRButtons;
+
 	/** If this is the first time using TogglePanelsVisibility */
 	bool bSetDefaultLayout;
 };
diff --git a/Engine/Source/Editor/VREditor/VREditorActions.cpp b/Engine/Source/Editor/VREditor/VREditorActions.cpp
new file mode 100644
--- /dev/null
+++ b/Engine/Source/Editor/VREditor/VREditorActions.cpp
@@ -0,0 +1,305 @@
+// Copyright 1998-2017 Epic Games, Inc. All Rights Reserved.
+
+#include "VREditorActions.h"
+#include "VREditorModule.h"
+#include "VREditorUISystem.h"
+#include "VREditorMode.h"
+#include "ViewportWorldInteraction.h"
+#include "VREditorInteractor.h"
+#include "EditorWorldManager.h"
+#include "VREditorFloatingUI.h"
+#include "VREditorTransformGizmo.h"
+#include "SLevelViewport.h"
+#include "ImageUtils.h"
+#include "FileHelper.h"
+
+
+#define LOCTEXT_NAMESPACE "VREditorActions"
+
+ECheckBoxState FVREditorActionCallbacks::GetTranslationSnapState()
+{
+	return GetDefault<ULevelEditorViewportSettings>()->GridEnabled ? ECheckBoxState::Checked : ECheckBoxState::Unchecked;
+}
+
+void FVREditorActionCallbacks::OnTranslationSnapSizeButtonClicked()
+{
+	const TArray<float>& GridSizes = GEditor->GetCurrentPositionGridArray();
+	const int32 CurrentGridSize = GetDefault<ULevelEditorViewportSettings>()->CurrentPosGridSize;
+
+	int32 NewGridSize = CurrentGridSize + 1;
+	if (NewGridSize >= GridSizes.Num())
+	{
+		NewGridSize = 0;
+	}
+
+	GEditor->SetGridSize(NewGridSize);
+
+	GetTranslationSnapSizeText();
+
+}
+
+FText FVREditorActionCallbacks::GetTranslationSnapSizeText() 
+{
+	return FText::AsNumber(GEditor->GetGridSize());
+
+}
+
+ECheckBoxState FVREditorActionCallbacks::GetRotationSnapState()
+{
+	return GetDefault<ULevelEditorViewportSettings>()->RotGridEnabled ? ECheckBoxState::Checked : ECheckBoxState::Unchecked;
+}
+
+void FVREditorActionCallbacks::OnRotationSnapSizeButtonClicked()
+{
+	const TArray<float>& GridSizes = GEditor->GetCurrentRotationGridArray();
+	const int32 CurrentGridSize = GetDefault<ULevelEditorViewportSettings>()->CurrentRotGridSize;
+	const ERotationGridMode CurrentGridMode = GetDefault<ULevelEditorViewportSettings>()->CurrentRotGridMode;
+
+	int32 NewGridSize = CurrentGridSize + 1;
+	if (NewGridSize >= GridSizes.Num())
+	{
+		NewGridSize = 0;
+	}
+
+	GEditor->SetRotGridSize(NewGridSize, CurrentGridMode);
+
+}
+
+FText FVREditorActionCallbacks::GetRotationSnapSizeText() 
+{
+	return FText::AsNumber(GEditor->GetRotGridSize().Yaw);
+}
+
+ECheckBoxState FVREditorActionCallbacks::GetScaleSnapState()
+{
+	return GetDefault<ULevelEditorViewportSettings>()->SnapScaleEnabled ? ECheckBoxState::Checked : ECheckBoxState::Unchecked;
+}
+
+void FVREditorActionCallbacks::OnScaleSnapSizeButtonClicked()
+{
+	const ULevelEditorViewportSettings* ViewportSettings = GetMutableDefault<ULevelEditorViewportSettings>();
+	const int32 NumGridSizes = ViewportSettings->ScalingGridSizes.Num();
+
+	const int32 CurrentGridSize = GetDefault<ULevelEditorViewportSettings>()->CurrentScalingGridSize;
+
+	int32 NewGridSize = CurrentGridSize + 1;
+	if (NewGridSize >= NumGridSizes)
+	{
+		NewGridSize = 0;
+	}
+
+	GEditor->SetScaleGridSize(NewGridSize);
+
+	
+}
+
+FText FVREditorActionCallbacks::GetScaleSnapSizeText() 
+{
+	return FText::AsNumber(GEditor->GetScaleGridSize());
+}
+
+void FVREditorActionCallbacks::OnGizmoCoordinateSystemButtonClicked(class UVREditorMode* InVRMode)
+{
+	InVRMode->GetWorldInteraction().CycleTransformGizmoCoordinateSpace();
+	UpdateGizmoModeText(InVRMode);
+	UpdateGizmoCoordinateSystemText(InVRMode);
+
+}
+
+FText FVREditorActionCallbacks::GetGizmoCoordinateSystemText()
+{
+	return EVREditorActions::GizmoCoordinateSystemText;
+}
+
+void FVREditorActionCallbacks::UpdateGizmoCoordinateSystemText(class UVREditorMode* InVRMode) 
+{
+	const ECoordSystem CurrentCoordSystem = InVRMode->GetWorldInteraction().GetTransformGizmoCoordinateSpace(); //@todo VREditor
+	EVREditorActions::GizmoCoordinateSystemText = (CurrentCoordSystem == COORD_World ? LOCTEXT("WorldCoordinateSystem", "World Space") : LOCTEXT("LocalCoordinateSystem", "Local Space"));
+}
+
+void FVREditorActionCallbacks::OnGizmoModeButtonClicked(class UVREditorMode* InVRMode)
+{
+	const EGizmoHandleTypes CurrentGizmoMode = InVRMode->GetWorldInteraction().GetCurrentGizmoType();
+	InVRMode->CycleTransformGizmoHandleType();
+	UpdateGizmoModeText(InVRMode);
+	UpdateGizmoCoordinateSystemText(InVRMode);
+
+}
+
+FText FVREditorActionCallbacks::GetGizmoModeText()
+{
+	return EVREditorActions::GizmoModeText;
+}
+
+void FVREditorActionCallbacks::UpdateGizmoModeText(class UVREditorMode* InVRMode) 
+{
+	const EGizmoHandleTypes CurrentGizmoType = InVRMode->GetWorldInteraction().GetCurrentGizmoType();
+	FText GizmoTypeText;
+	switch (CurrentGizmoType)
+	{
+	case EGizmoHandleTypes::All:
+		GizmoTypeText = LOCTEXT("AllGizmoType", "Universal");
+		break;
+
+	case EGizmoHandleTypes::Translate:
+		GizmoTypeText = LOCTEXT("TranslateGizmoType", "Translate");
+		break;
+
+	case EGizmoHandleTypes::Rotate:
+		GizmoTypeText = LOCTEXT("RotationGizmoType", "Rotation");
+		break;
+
+	case EGizmoHandleTypes::Scale:
+		GizmoTypeText = LOCTEXT("ScaleGizmoType", "Scale");
+		break;
+
+	default:
+		check(0);	// Unrecognized type
+		break;
+	}
+
+	EVREditorActions::GizmoModeText = GizmoTypeText;
+}
+
+void FVREditorActionCallbacks::OnGizmoCycle(class UVREditorMode* InVRMode)
+{
+	InVRMode->CycleTransformGizmoHandleType();
+	UpdateGizmoModeText(InVRMode);
+	UpdateGizmoCoordinateSystemText(InVRMode);
+}
+
+void FVREditorActionCallbacks::OnUIToggleButtonClicked(class UVREditorMode* InVRMode, const UVREditorUISystem::EEditorUIPanel PanelToToggle)
+{
+	InVRMode->GetUISystem().TogglePanelVisibility(PanelToToggle);
+}
+
+ECheckBoxState FVREditorActionCallbacks::GetUIToggledState(class UVREditorMode* InVRMode, const UVREditorUISystem::EEditorUIPanel PanelToCheck)
+{
+	return InVRMode->GetUISystem().IsShowingEditorUIPanel(PanelToCheck) ? ECheckBoxState::Checked : ECheckBoxState::Unchecked;
+}
+
+
+void FVREditorActionCallbacks::OnLightButtonClicked(class UVREditorMode* InVRMode)
+{
+	// Always spawn the flashlight on the hand clicking on the UI
+	UVREditorInteractor* VREditorInteractor = InVRMode->GetHandInteractor(EControllerHand::Left);
+	if (!VREditorInteractor->IsHoveringOverUI())
+	{
+		VREditorInteractor = InVRMode->GetHandInteractor(EControllerHand::Right);
+
+	}
+	InVRMode->ToggleFlashlight(VREditorInteractor);
+
+}
+
+ECheckBoxState FVREditorActionCallbacks::GetFlashlightState(class UVREditorMode* InVRMode)
+{
+	return InVRMode->IsFlashlightOn() ? ECheckBoxState::Checked : ECheckBoxState::Unchecked;
+}
+
+
+void FVREditorActionCallbacks::OnScreenshotButtonClicked(class UVREditorMode* InVRMode)
+{
+	// @todo vreditor: update after direct buffer grab changes
+
+	FString GeneratedFilename = TEXT("");
+	FString Filename = TEXT("");
+	FScreenshotRequest::CreateViewportScreenShotFilename(GeneratedFilename);
+	const bool bRemovePath = false;
+	GeneratedFilename = FPaths::GetBaseFilename(GeneratedFilename, bRemovePath);
+	FFileHelper::GenerateNextBitmapFilename(GeneratedFilename, TEXT("png"), Filename);
+
+	TSharedRef<SWidget> WindowRef = InVRMode->GetLevelViewportPossessedForVR().AsWidget();
+
+
+	TArray<FColor> OutImageData;
+	FIntVector OutImageSize;
+
+	if (FSlateApplication::Get().TakeScreenshot(WindowRef, OutImageData, OutImageSize))
+	{
+		int32 Width = OutImageSize.X;
+		int32 Height = OutImageSize.Y;
+
+		TArray<FColor> ScaledBitmap;
+
+
+		ScaledBitmap = OutImageData;
+		//Clear the alpha channel before saving
+		for (int32 Index = 0; Index < Width*Height; Index++)
+		{
+			ScaledBitmap[Index].A = 255;
+		}
+
+
+		TArray<uint8> CompressedBitmap;
+		FImageUtils::CompressImageArray(Width, Height, ScaledBitmap, CompressedBitmap);
+
+
+		//Save locally
+		const bool bTree = true;
+		const FString FinalFileName = Filename;
+		IFileManager::Get().MakeDirectory(*FPaths::GetPath(FinalFileName), bTree);
+		FFileHelper::SaveArrayToFile(CompressedBitmap, *FinalFileName);
+
+	}
+
+
+}
+
+void FVREditorActionCallbacks::OnPlayButtonClicked(class UVREditorMode* InVRMode)
+{
+	InVRMode->StartExitingVRMode(EVREditorExitType::PIE_VR);
+}
+
+void FVREditorActionCallbacks::OnSnapActorsToGroundClicked(class UVREditorMode* InVRMode)
+{
+	InVRMode->SnapSelectedActorsToGround();
+}
+
+
+void FVREditorActionCallbacks::SimulateCharacterEntry(const FString InChar)
+{
+
+	for (int32 CharIndex = 0; CharIndex < InChar.Len(); CharIndex++)
+	{
+		TCHAR CharKey = InChar[CharIndex];
+		const bool bRepeat = false;
+		FCharacterEvent CharacterEvent(CharKey, FModifierKeysState::FModifierKeysState(), 0, bRepeat);
+		FSlateApplication::Get().ProcessKeyCharEvent(CharacterEvent);
+	}
+
+}
+
+void FVREditorActionCallbacks::SimulateKeyDown(FKey Key, bool bRepeat)
+{
+	const uint32* KeyCodePtr;
+	const uint32* CharCodePtr;
+	FInputKeyManager::Get().GetCodesFromKey( Key, KeyCodePtr, CharCodePtr );
+
+	uint32 KeyCode = KeyCodePtr ? *KeyCodePtr : 0;
+	uint32 CharCode = CharCodePtr ? *CharCodePtr : 0;
+
+	FKeyEvent KeyEvent( Key, FModifierKeysState::FModifierKeysState(), 0, bRepeat, KeyCode, CharCode );
+	bool DownResult = FSlateApplication::Get().ProcessKeyDownEvent( KeyEvent );
+
+	if (CharCodePtr)
+	{
+		FCharacterEvent CharacterEvent( CharCode, FModifierKeysState::FModifierKeysState(), 0, bRepeat );
+		FSlateApplication::Get().ProcessKeyCharEvent( CharacterEvent );
+	}
+}
+
+void FVREditorActionCallbacks::SimulateKeyUp(FKey Key)
+{
+	const uint32* KeyCodePtr;
+	const uint32* CharCodePtr;
+	FInputKeyManager::Get().GetCodesFromKey( Key, KeyCodePtr, CharCodePtr );
+
+	uint32 KeyCode = KeyCodePtr ? *KeyCodePtr : 0;
+	uint32 CharCode = CharCodePtr ? *CharCodePtr : 0;
+
+	FKeyEvent KeyEvent( Key, FModifierKeysState::FModifierKeysState(), 0, false, KeyCode, CharCode );
+	FSlateApplication::Get().ProcessKeyUpEvent( KeyEvent );
+}
+
+#undef LOCTEXT_NAMESPACE
diff --git a/Engine/Source/Editor/VREditor/VREditorActions.h b/Engine/Source/Editor/VREditor/VREditorActions.h
new file mode 100644
--- /dev/null
+++ b/Engine/Source/Editor/VREditor/VREditorActions.h
@@ -0,0 +1,152 @@
+// Copyright 1998-2017 Epic Games, Inc. All Rights Reserved.
+#pragma once
+
+#include "VREditorUISystem.h"
+
+namespace EVREditorActions
+{
+	//Gizmo coordinate button text starts blank before it is initialized
+	static FText GizmoCoordinateSystemText = NSLOCTEXT("GizmoCoordinateDefault", "", "");
+
+	//Gizmo mode button text starts blank before it is initialized
+	static FText GizmoModeText = NSLOCTEXT("GizmoModeDefault", "", "");;
+}
+
+/**
+* Implementation of various VR editor action callback functions
+*/
+class FVREditorActionCallbacks
+{
+public:
+
+	/** Returns the checked state of the translation snap enable/disable button */
+	static ECheckBoxState GetTranslationSnapState();
+
+	/** Rotates through the available translation snap sizes */
+	static void OnTranslationSnapSizeButtonClicked();
+
+	/** Returns the translation snap size as text to use as the button display text */
+	static FText GetTranslationSnapSizeText();
+
+	/** Returns the checked state of the rotation snap enable/disable button */
+	static ECheckBoxState GetRotationSnapState();
+
+	/** Rotates through the available rotation snap sizes */
+	static void OnRotationSnapSizeButtonClicked();
+
+	/** Returns the rotation snap size as text to use as the button display text */
+	static FText GetRotationSnapSizeText();
+
+	/** Returns the checked state of the scale snap enable/disable button */
+	static ECheckBoxState GetScaleSnapState();
+
+	/** Rotates through the available scale snap sizes */
+	static void OnScaleSnapSizeButtonClicked();
+
+	/** Returns the scale snap size as text to use as the button display text */
+	static FText GetScaleSnapSizeText();
+
+	/**
+	 * Toggles the gizmo coordinate system between local and world space
+	 * @param InVRMode Currently active VRMode
+	 */
+	static void OnGizmoCoordinateSystemButtonClicked(class UVREditorMode* InVRMode);
+
+	/** Returns the gizmo coordinate system as text to use as the button display text */
+	static FText GetGizmoCoordinateSystemText();
+
+	/**
+	 * Updates the gizmo coordinate system text if the coordinate system or gizmo type is changed
+	 * @param InVRMode Currently active VRMode
+	 */
+	static void UpdateGizmoCoordinateSystemText(class UVREditorMode* InVRMode);
+
+	/**
+	 * Rotates the gizmo type through universal, translate, rotate, and scale
+	 * @param InVRMode Currently active VRMode
+	 */
+	static void OnGizmoModeButtonClicked(class UVREditorMode* InVRMode);
+
+	/** Returns the gizmo type as text to use as the button display text */
+	static FText GetGizmoModeText();
+
+	/**
+	 * Updates the gizmo coordinate type text if the coordinate system or gizmo type is changed
+	 * @param InVRMode Currently active VRMode
+	 */
+	static void UpdateGizmoModeText(class UVREditorMode* InVRMode);
+
+	/**
+	* Cycles the gizmo through all coordinate spaces and Actor manipulation types
+	* @param InVRMode Currently active VRMode
+	*/
+	static void OnGizmoCycle(class UVREditorMode* InVRMode);
+
+	/**
+	 * Toggles a VR UI panel's state between visible and invisible
+	 * @param InVRMode Currently active VRMode
+	 * @param PanelToToggle UI panel to change the state of
+	 */
+	static void OnUIToggleButtonClicked(class UVREditorMode* InVRMode, const UVREditorUISystem::EEditorUIPanel PanelToToggle);
+
+	/**
+	 * Returns a VR UI panel's visibility - used for check boxes on the menu button
+	 * @param InVRMode Currently active VRMode
+	 * @param PanelToToggle UI panel to read the state of
+	 */
+	static ECheckBoxState GetUIToggledState(class UVREditorMode* InVRMode, const UVREditorUISystem::EEditorUIPanel PanelToCheck);
+
+	/**
+	 * Toggles a flashlight on and off on the interactor that clicked on the UI button
+	 * @param InVRMode Currently active VRMode
+	 */
+	static void OnLightButtonClicked(class UVREditorMode* InVRMode);
+
+	/**
+	 * Returns whether or not the flashlight is enabled - used for check box on the flashlight button
+	 * @param InVRMode Currently active VRMode
+	 */
+	static ECheckBoxState GetFlashlightState(class UVREditorMode* InVRMode);
+
+	/**
+	 * Takes a screenshot of the mirror viewport
+	 * @param InVRMode Currently active VRMode
+	 */
+	static void OnScreenshotButtonClicked(class UVREditorMode* InVRMode);
+
+	/**
+	 * Enters Play In Editor mode for testing of gameplay
+	 * @param InVRMode Currently active VRMode
+	 */
+	static void OnPlayButtonClicked(class UVREditorMode* InVRMode);
+
+	/**
+	 * Snaps currently selected Actors to the ground
+	 * @param InVRMode Currently active VRMode
+	 */
+	static void OnSnapActorsToGroundClicked(class UVREditorMode* InVRMode);
+
+	/**
+	 * Simulates the user entering characters with a keyboard for data entry
+	 * @param InChar String of characters to enter
+	 */
+	static void SimulateCharacterEntry(const FString InChar);
+
+	/**
+	 * Simulates the user pressing a key down
+	 * @param Key Key to press
+	 * @param bRepeat Whether or not to repeat
+	 */
+	static void SimulateKeyDown( FKey Key, bool bRepeat );
+
+	/**
+	 * Simulates the user releasing a key
+	 * @param Key Key to release
+	 */
+	static void SimulateKeyUp( FKey Key );
+
+}
+;
+
+
+
diff --git a/Engine/Source/Editor/VREditor/VREditorBaseActor.cpp b/Engine/Source/Editor/VREditor/VREditorBaseActor.cpp
--- a/Engine/Source/Editor/VREditor/VREditorBaseActor.cpp
+++ b/Engine/Source/Editor/VREditor/VREditorBaseActor.cpp
@@ -7,6 +7,7 @@
 namespace VREd
 {
 	static FAutoConsoleVariable UIOnHandRotationOffset( TEXT( "VREd.UIOnHandRotationOffset" ), 45.0f, TEXT( "Rotation offset for UI that's docked to your hand, to make it more comfortable to hold" ) );
+	static FAutoConsoleVariable UIOnArmRotationOffset( TEXT( "VREd.UIOnArmRotationOffset" ), 15.0f, TEXT( "Rotation offset for UI that's docked to your arm, so it aligns with the controllers" ) );
 }
 
 AVREditorBaseActor::AVREditorBaseActor() 
@@ -89,10 +90,16 @@ FTransform AVREditorBaseActor::MakeUITransformLockedToHand( UViewportInteractor*
 	const float WorldScaleFactor = VRMode->GetWorldScaleFactor();
 
 	FTransform UIToHandTransform( InLocalRotation, InRelativeOffset * WorldScaleFactor );
+	// If the UI is on the hand
 	if ( !bOnArm )
 	{
 		UIToHandTransform *= FTransform( FRotator( VREd::UIOnHandRotationOffset->GetFloat(), 0.0f, 0.0f ).Quaternion(), FVector::ZeroVector );
 	}
+	// If the UI is on the arm
+	else
+	{
+		UIToHandTransform *= FTransform( FRotator( VREd::UIOnArmRotationOffset->GetFloat(), 0.0f, 0.0f ).Quaternion(), FVector::ZeroVector );
+	}
 
 	const FTransform HandToWorldTransform = Interactor->GetTransform();
 	FTransform UIToWorldTransform = UIToHandTransform * HandToWorldTransform;
diff --git a/Engine/Source/Editor/VREditor/VREditorMode.cpp b/Engine/Source/Editor/VREditor/VREditorMode.cpp
--- a/Engine/Source/Editor/VREditor/VREditorMode.cpp
+++ b/Engine/Source/Editor/VREditor/VREditorMode.cpp
@@ -40,6 +40,7 @@ namespace VREd
 {
 	static FAutoConsoleVariable UseMouseAsHandInForcedVRMode( TEXT( "VREd.UseMouseAsHandInForcedVRMode" ), 1, TEXT( "When in forced VR mode, enabling this setting uses the mouse cursor as a virtual hand instead of motion controllers" ) );
 	static FAutoConsoleVariable ForceOculusMirrorMode( TEXT( "VREd.ForceOculusMirrorMode" ), 3, TEXT( "Which Oculus display mirroring mode to use (see MirrorWindowModeType in OculusRiftHMD.h)" ) );
+	static FAutoConsoleVariable VRNearClipPlane(TEXT("VREd.VRNearClipPlane"), 1.0f, TEXT("The near clip plane to use for VR"));
 }
 
 // @todo vreditor: Hacky that we have to import these this way. (Plugin has them in a .cpp, not exported)
@@ -276,9 +277,8 @@ void UVREditorMode::Enter(const bool bReenteringVREditing)
 			VRViewportClient.bAlwaysShowModeWidgetAfterSelectionChanges = false;
 
 			// Force tiny near clip plane distance, because user can scale themselves to be very small.
-			// @todo vreditor: Make this automatically change based on WorldToMetersScale?
 			SavedEditorState.NearClipPlane = GNearClippingPlane;
-			GNearClippingPlane = 1.0f;	// Normally defaults to 10cm (NOTE: If we go too small, skyboxes become affected)
+			GNearClippingPlane = GetDefaultVRNearClipPlane();	
 
 			SavedEditorState.bOnScreenMessages = GAreScreenMessagesEnabled;
 			GAreScreenMessagesEnabled = false;
@@ -793,6 +793,11 @@ void UVREditorMode::DestroyTransientActor( AActor* Actor ) const
 	WorldInteraction->DestroyTransientActor( Actor );
 }
 
+bool UVREditorMode::IsShowingRadialMenu(const UVREditorInteractor* Interactor) const
+{
+	return UISystem->IsShowingRadialMenu(Interactor);
+}
+
 const SLevelViewport& UVREditorMode::GetLevelViewportPossessedForVR() const
 {
 	return *VREditorLevelViewportWeakPtr.Pin();
@@ -899,6 +904,11 @@ FLinearColor UVREditorMode::GetColor( const EColors Color ) const
 	return Colors[ (int32)Color ];
 }
 
+float UVREditorMode::GetDefaultVRNearClipPlane() const 
+{
+	return VREd::VRNearClipPlane->GetFloat();
+}
+
 bool UVREditorMode::IsHandAimingTowardsCapsule( UViewportInteractor* Interactor, const FTransform& CapsuleTransform, FVector CapsuleStart, const FVector CapsuleEnd, const float CapsuleRadius, const float MinDistanceToCapsule, const FVector CapsuleFrontDirection, const float MinDotForAimingAtCapsule ) const
 {
 	bool bIsAimingTowards = false;
diff --git a/Engine/Source/Editor/VREditor/VREditorModule.cpp b/Engine/Source/Editor/VREditor/VREditorModule.cpp
--- a/Engine/Source/Editor/VREditor/VREditorModule.cpp
+++ b/Engine/Source/Editor/VREditor/VREditorModule.cpp
@@ -7,6 +7,8 @@
 #include "IVREditorModule.h"
 #include "TickableEditorObject.h"
 #include "VREditorModeManager.h"
+#include "VREditorStyle.h"
+#include "Kismet/HeadMountedDisplayFunctionLibrary.h"	// For EHMDWornState::Type
 
 class FVREditorModule : public IVREditorModule, public FTickableEditorObject
 {
@@ -57,10 +59,20 @@ namespace VREd
 
 void FVREditorModule::StartupModule()
 {
+	if (GIsEditor)
+	{
+		FVREditorStyle::Initialize();
+	}
+
 }
 
 void FVREditorModule::ShutdownModule()
 {
+
+	if (GIsEditor)
+	{
+		FVREditorStyle::Shutdown();
+	}
 }
 
 void FVREditorModule::PostLoadCallback()
diff --git a/Engine/Source/Editor/VREditor/VREditorMotionControllerInteractor.cpp b/Engine/Source/Editor/VREditor/VREditorMotionControllerInteractor.cpp
--- a/Engine/Source/Editor/VREditor/VREditorMotionControllerInteractor.cpp
+++ b/Engine/Source/Editor/VREditor/VREditorMotionControllerInteractor.cpp
@@ -25,6 +25,9 @@
 #include "MotionControllerComponent.h"
 #include "Features/IModularFeatures.h"
 
+#include "SVRRadialPanel.h"
+#include "VREditorWidgetComponent.h"
+
 namespace VREd
 {
 	//Laser
@@ -40,8 +43,10 @@ namespace VREd
 	static FAutoConsoleVariable LaserPointerLightRadius( TEXT( "VREd.LaserPointLightRadius" ), 20.0f, TEXT( "How big our hover light is" ) );
 
 	//Trigger
-	static FAutoConsoleVariable TriggerLightlyPressedThreshold( TEXT( "VI.TriggerLightlyPressedThreshold" ), 0.03f, TEXT( "Minimum trigger threshold before we consider the trigger at least 'lightly pressed'" ) );
-	static FAutoConsoleVariable TriggerDeadZone( TEXT( "VI.TriggerDeadZone" ), 0.01f, TEXT( "Trigger dead zone.  The trigger must be fully released before we'll trigger a new 'light press'" ) );
+	static FAutoConsoleVariable TriggerLightlyPressedThreshold_Vive( TEXT( "VI.TriggerLightlyPressedThreshold_Vive" ), 0.03f, TEXT( "Minimum trigger threshold before we consider the trigger at least 'lightly pressed'" ) );
+	static FAutoConsoleVariable TriggerDeadZone_Vive( TEXT( "VI.TriggerDeadZone_Vive" ), 0.01f, TEXT( "Trigger dead zone.  The trigger must be fully released before we'll trigger a new 'light press'" ) );
+	static FAutoConsoleVariable TriggerLightlyPressedThreshold_Rift( TEXT( "VI.TriggerLightlyPressedThreshold_Rift" ), 0.2f, TEXT( "Minimum trigger threshold before we consider the trigger at least 'lightly pressed'" ) );
+	static FAutoConsoleVariable TriggerDeadZone_Rift( TEXT( "VI.TriggerDeadZone_Rift" ), 0.15f, TEXT( "Trigger dead zone.  The trigger must be fully released before we'll trigger a new 'light press'" ) );
 	static FAutoConsoleVariable TriggerFullyPressedThreshold( TEXT( "VI.TriggerFullyPressedThreshold" ), 0.95f, TEXT( "Minimum trigger threshold before we consider the trigger 'fully pressed'" ) );
 	static FAutoConsoleVariable TriggerFullyPressedReleaseThreshold( TEXT( "VI.TriggerFullyPressedReleaseThreshold" ), 0.8f, TEXT( "After fully pressing the trigger, if the axis falls below this threshold we no longer consider it fully pressed" ) );
 
@@ -54,7 +59,6 @@ namespace VREd
 	static FAutoConsoleVariable InvertTrackpadVertical( TEXT( "VREd.InvertTrackpadVertical" ), 1, TEXT( "Toggles inverting the touch pad vertical axis" ) );
 	static FAutoConsoleVariable TriggerLightlyPressedLockTime( TEXT( "VREd.TriggerLightlyPressedLockTime" ), 0.15f, TEXT( "If the trigger remains lightly pressed for longer than this, we'll continue to treat it as a light press in some cases" ) );
 	static FAutoConsoleVariable MinVelocityForInertia( TEXT( "VREd.MinVelocityForMotionControllerInertia" ), 1.0f, TEXT( "Minimum velocity (in cm/frame in unscaled room space) before inertia will kick in when releasing objects (or the world)" ) );
-	static FAutoConsoleVariable MinJoystickOffsetBeforeRadialMenu( TEXT( "VREd.MinJoystickOffsetBeforeRadialMenu" ), 0.15f, TEXT( "Toggles inverting the touch pad vertical axis" ) );
 
 	static FAutoConsoleVariable HelpLabelFadeDuration( TEXT( "VREd.HelpLabelFadeDuration" ), 0.4f, TEXT( "Duration to fade controller help labels in and out" ) );
 	static FAutoConsoleVariable HelpLabelFadeDistance( TEXT( "VREd.HelpLabelFadeDistance" ), 30.0f, TEXT( "Distance at which controller help labels should appear (in cm)" ) );
@@ -132,8 +136,6 @@ void UVREditorMotionControllerInteractor::Init( class UVREditorMode* InVRMode )
 		AddKeyAction( EKeys::MotionController_Left_Grip1, FViewportActionKeyInput( ViewportWorldActionTypes::WorldMovement ) );
 		AddKeyAction( UVREditorMotionControllerInteractor::MotionController_Left_FullyPressedTriggerAxis, FViewportActionKeyInput( ViewportWorldActionTypes::SelectAndMove ) );
 		AddKeyAction( UVREditorMotionControllerInteractor::MotionController_Left_LightlyPressedTriggerAxis, FViewportActionKeyInput( ViewportWorldActionTypes::SelectAndMove_LightlyPressed ) );
-		AddKeyAction( EKeys::MotionController_Left_Thumbstick, FViewportActionKeyInput( VRActionTypes::ConfirmRadialSelection ) );
-
 		AddKeyAction( SteamVRControllerKeyNames::Touch0, FViewportActionKeyInput( VRActionTypes::Touch ) );
 		AddKeyAction( EKeys::MotionController_Left_TriggerAxis, FViewportActionKeyInput( UVREditorMotionControllerInteractor::TriggerAxis ) );
 		AddKeyAction( EKeys::MotionController_Left_Thumbstick_X, FViewportActionKeyInput( UVREditorMotionControllerInteractor::TrackpadPositionX ) );
@@ -335,7 +337,7 @@ void UVREditorMotionControllerInteractor::Tick( const float DeltaTime )
 		}
 	}
 
-	UpdateRadialMenuInput( DeltaTime );
+	UpdateRadialMenuInput(DeltaTime);
 
 	{
 		const float WorldScaleFactor = WorldInteraction->GetWorldScaleFactor();
@@ -629,6 +631,32 @@ void UVREditorMotionControllerInteractor::HandleInputKey( FViewportActionKeyInpu
 		}
 	}
 
+	if( Action.ActionType == VRActionTypes::ConfirmRadialSelection )
+	{
+		bOutWasHandled = true;
+		const EViewportInteractionDraggingMode DraggingMode = GetDraggingMode();
+
+		if( Event == IE_Pressed )
+		{ 
+			// Start dragging at laser impact when already dragging actors freely
+			if( DraggingMode == EViewportInteractionDraggingMode::ActorsFreely )
+			{
+				FVector PlaceAt = GetHoverLocation();
+				const bool bIsPlacingActors = true;
+				const FViewportActionKeyInput FakeSelectAndMoveAction( ViewportWorldActionTypes::SelectAndMove_LightlyPressed );
+				WorldInteraction->StartDraggingActors( this, FakeSelectAndMoveAction, WorldInteraction->GetTransformGizmoActor()->GetRootComponent(), PlaceAt, bIsPlacingActors );
+			}
+		}
+		else if( Event == IE_Released )
+		{
+			// Disable dragging at laser impact when releasing
+			if( DraggingMode == EViewportInteractionDraggingMode::ActorsAtLaserImpact )
+			{
+				SetDraggingMode( EViewportInteractionDraggingMode::ActorsFreely );
+			}
+		}
+	}
+
 	ApplyButtonPressColors( Action );
 }
 
@@ -636,6 +664,9 @@ void UVREditorMotionControllerInteractor::HandleInputAxis( FViewportActionKeyInp
 {
 	if ( Action.ActionType == TriggerAxis )
 	{
+		const float TriggerLightlyPressedThreshold = ( GetHMDDeviceType() == EHMDDeviceType::DT_OculusRift ) ? VREd::TriggerLightlyPressedThreshold_Rift->GetFloat() : VREd::TriggerLightlyPressedThreshold_Vive->GetFloat();
+		const float TriggerDeadZone = ( GetHMDDeviceType() == EHMDDeviceType::DT_OculusRift ) ? VREd::TriggerDeadZone_Rift->GetFloat() : VREd::TriggerDeadZone_Vive->GetFloat();
+
 		// Synthesize "lightly pressed" events for the trigger
 		{
 			// Store latest trigger value amount
@@ -647,7 +678,7 @@ void UVREditorMotionControllerInteractor::HandleInputAxis( FViewportActionKeyInp
 			if ( !bIsFullPressAlreadyCapturing &&		// Don't fire if we're already capturing input for a full press
 				!bIsTriggerAtLeastLightlyPressed &&	// Don't fire unless we are already pressed
 				bHasTriggerBeenReleasedSinceLastPress &&	// Only if we've been fully released since the last time we fired
-				Delta >= VREd::TriggerLightlyPressedThreshold->GetFloat() )
+				Delta >= TriggerLightlyPressedThreshold )
 			{
 				bIsTriggerAtLeastLightlyPressed = true;
 				SetAllowTriggerLightPressLocking( true );
@@ -659,7 +690,7 @@ void UVREditorMotionControllerInteractor::HandleInputAxis( FViewportActionKeyInp
 				const EInputEvent InputEvent = IE_Pressed;
 				const bool bWasLightPressHandled = UViewportInteractor::HandleInputKey( ControllerHandSide == EControllerHand::Left ? MotionController_Left_LightlyPressedTriggerAxis : MotionController_Right_LightlyPressedTriggerAxis, InputEvent );
 			}
-			else if ( bIsTriggerAtLeastLightlyPressed && Delta < VREd::TriggerLightlyPressedThreshold->GetFloat() )
+			else if ( bIsTriggerAtLeastLightlyPressed && Delta < TriggerLightlyPressedThreshold )
 			{
 				bIsTriggerAtLeastLightlyPressed = false;
 
@@ -669,7 +700,7 @@ void UVREditorMotionControllerInteractor::HandleInputAxis( FViewportActionKeyInp
 			}
 		}
 
-		if ( Delta < VREd::TriggerDeadZone->GetFloat() )
+		if ( !bHasTriggerBeenReleasedSinceLastPress && Delta < TriggerDeadZone )
 		{
 			bHasTriggerBeenReleasedSinceLastPress = true;
 		}
@@ -1278,7 +1309,7 @@ void UVREditorMotionControllerInteractor::SetLaserVisuals( const FLinearColor& N
 void UVREditorMotionControllerInteractor::UpdateRadialMenuInput( const float DeltaTime )
 {
 	UVREditorUISystem& UISystem = GetVRMode().GetUISystem();
-
+	const EHMDDeviceType::Type HMDDeviceType = GetVRMode().GetHMDDeviceType();
 	//Update the radial menu
 	EViewportInteractionDraggingMode DraggingMode = GetDraggingMode();
 	if ( !HasUIInFront() &&
@@ -1290,24 +1321,21 @@ void UVREditorMotionControllerInteractor::UpdateRadialMenuInput( const float Del
 		!UISystem.IsInteractorDraggingDockUI( this ) &&
 		!UISystem.IsShowingRadialMenu( Cast<UVREditorInteractor>( OtherInteractor ) ) )
 	{
-		const EHMDDeviceType::Type HMDDevice = GetHMDDeviceType();
-
-		// Spawn the radial menu if we are using the touchpad for steamvr or the analog stick for the oculus
-		if ( ( HMDDevice == EHMDDeviceType::DT_SteamVR && bIsTouchingTrackpad ) ||
-			( HMDDevice == EHMDDeviceType::DT_OculusRift && TrackpadPosition.Size() > VREd::MinJoystickOffsetBeforeRadialMenu->GetFloat() ) )
-		{
-			UISystem.TryToSpawnRadialMenu( this );
-		}
-		else
-		{
-			// Hide it if we are not using the touchpad or analog stick
-			UISystem.HideRadialMenu( this );
-		}
-
 		// Update the radial menu if we are already showing the radial menu
 		if ( UISystem.IsShowingRadialMenu( this ) )
 		{
-			UISystem.UpdateRadialMenu( this );
+			const TSharedPtr<SRadialBox>& RadialBox = StaticCastSharedPtr<SRadialBox>(UISystem.GetRadialWidget());
+			RadialBox->HighlightSlot(TrackpadPosition);
+		}
+	}
+	// If we are not currently touching the Vive touchpad, reset the highlighted button
+	else if (HMDDeviceType == EHMDDeviceType::DT_SteamVR && !bIsTouchingTrackpad)
+	{
+		if (UISystem.IsShowingRadialMenu( this ))
+		{
+			const TSharedPtr<SRadialBox>& RadialBox = StaticCastSharedPtr<SRadialBox>( UISystem.GetRadialWidget() );
+			FVector2D ReturnToCenter = FVector2D::ZeroVector;
+			RadialBox->HighlightSlot( ReturnToCenter );
 		}
 	}
 }
diff --git a/Engine/Source/Editor/VREditor/VREditorStyle.cpp b/Engine/Source/Editor/VREditor/VREditorStyle.cpp
new file mode 100644
--- /dev/null
+++ b/Engine/Source/Editor/VREditor/VREditorStyle.cpp
@@ -0,0 +1,112 @@
+// Copyright 1998-2017 Epic Games, Inc. All Rights Reserved.
+
+#include "VREditorStyle.h"
+#include "VREditorModule.h"
+#include "SlateTypes.h"
+#include "SlateGameResources.h"
+#include "Styling/SlateStyleRegistry.h"
+#include "SUniformGridPanel.h"
+
+TSharedPtr< FSlateStyleSet > FVREditorStyle::VREditorStyleInstance = NULL;
+
+void FVREditorStyle::Initialize()
+{
+	if ( !VREditorStyleInstance.IsValid() )
+	{
+		VREditorStyleInstance = Create();
+		FSlateStyleRegistry::RegisterSlateStyle( *VREditorStyleInstance);
+	}
+}
+
+void FVREditorStyle::Shutdown()
+{
+	FSlateStyleRegistry::UnRegisterSlateStyle( *VREditorStyleInstance);
+	ensure(VREditorStyleInstance.IsUnique() );
+	VREditorStyleInstance.Reset();
+}
+
+FName FVREditorStyle::GetStyleSetName()
+{
+	static FName StyleSetName(TEXT("VREditorStyle"));
+	return StyleSetName;
+}
+
+#define IMAGE_BRUSH( RelativePath, ... ) FSlateImageBrush( Style->RootToContentDir( RelativePath, TEXT(".png") ), __VA_ARGS__ )
+#define BOX_BRUSH( RelativePath, ... ) FSlateBoxBrush( Style->RootToContentDir( RelativePath, TEXT(".png") ), __VA_ARGS__ )
+#define BORDER_BRUSH( RelativePath, ... ) FSlateBorderBrush( Style->RootToContentDir( RelativePath, TEXT(".png") ), __VA_ARGS__ )
+#define TTF_FONT( RelativePath, ... ) FSlateFontInfo( Style->RootToContentDir( RelativePath, TEXT(".ttf") ), __VA_ARGS__ )
+#define OTF_FONT( RelativePath, ... ) FSlateFontInfo( Style->RootToContentDir( RelativePath, TEXT(".otf") ), __VA_ARGS__ )
+#define TTF_CORE_FONT(RelativePath, ...) FSlateFontInfo(Style->RootToCoreContentDir(RelativePath, TEXT(".ttf") ), __VA_ARGS__)
+
+const FVector2D Icon16x16(16.0f, 16.0f);
+const FVector2D Icon20x20(20.0f, 20.0f);
+const FVector2D Icon40x40(40.0f, 40.0f);
+
+TSharedRef< FSlateStyleSet > FVREditorStyle::Create()
+{
+	TSharedRef< FSlateStyleSet > Style = MakeShareable(new FSlateStyleSet(FVREditorStyle::GetStyleSetName()));
+	// Use the default menu button style, but set the background to dark grey.
+	const FButtonStyle NormalButton = FEditorStyle::GetWidgetStyle<FButtonStyle>("Menu.Button");
+	Style->Set("VREditorStyle.Button", FButtonStyle(NormalButton)
+		.SetNormal(BOX_BRUSH("Common/RoundedSelection_16x", 4.0f / 16.0f, FLinearColor(0.1f,0.1f,0.1f))));
+
+	const FTextBlockStyle NormalText = FEditorStyle::GetWidgetStyle<FTextBlockStyle>("NormalText");
+	Style->Set("VREditorStyle.Label", FTextBlockStyle(NormalText)
+		.SetFont(TTF_CORE_FONT("Slate/Fonts/Roboto-Regular", 9)));
+
+	// Headings will have a font outline
+	FFontOutlineSettings HeadingOutline;
+	HeadingOutline.OutlineColor = FLinearColor(0.2f, 0.2f, 0.2f, 0.5f);
+	HeadingOutline.OutlineSize = 2.0f;
+	FSlateFontInfo HeadlineFont = TTF_CORE_FONT("Slate/Fonts/Roboto-Regular", 18);
+	HeadlineFont.OutlineSettings = HeadingOutline;
+
+	Style->Set("VREditorStyle.Heading", FTextBlockStyle(NormalText)
+			.SetFont(HeadlineFont)
+			.SetColorAndOpacity(FLinearColor::White));
+
+	// Headings will have a font outline
+	FFontOutlineSettings HelperOutline;
+	HelperOutline.OutlineColor = FLinearColor( 0.2f, 0.2f, 0.2f, 0.5f );
+	HelperOutline.OutlineSize = 3.0f;
+	FSlateFontInfo HelperFont = TTF_CORE_FONT( "Slate/Fonts/Roboto-Regular", 24 );
+	HelperFont.OutlineSettings = HelperOutline;
+
+	Style->Set( "VREditorStyle.HelperText", FTextBlockStyle( NormalText )
+		.SetFont( HelperFont )
+		.SetColorAndOpacity( FLinearColor::White ) );
+	
+	const FCheckBoxStyle CheckStyle = FEditorStyle::GetWidgetStyle<FCheckBoxStyle>("Menu.Check");
+	Style->Set("VREditorStyle.Check", FCheckBoxStyle(CheckStyle));
+
+	const FCheckBoxStyle CheckboxStyle = FEditorStyle::GetWidgetStyle<FCheckBoxStyle>("Menu.CheckBox");
+	Style->Set("VREditorStyle.CheckBox", FCheckBoxStyle(CheckboxStyle));
+
+	const FCheckBoxStyle RadioButtonStyle = FEditorStyle::GetWidgetStyle<FCheckBoxStyle>("Menu.RadioButton");
+	Style->Set("VREditorStyle.RadioButton", FCheckBoxStyle(RadioButtonStyle));
+
+	const FTextBlockStyle KeybindingStyle = FEditorStyle::GetWidgetStyle<FTextBlockStyle>("Menu.Keybinding");
+	Style->Set("VREditorStyle.Keybinding", FTextBlockStyle(KeybindingStyle));
+
+	const FTextBlockStyle LabelStyle = FEditorStyle::GetWidgetStyle<FTextBlockStyle>("Menu.Label");
+	Style->Set("VREditorStyle.Label", FTextBlockStyle(LabelStyle));
+
+	return Style;
+}
+
+#undef IMAGE_BRUSH
+#undef BOX_BRUSH
+#undef BORDER_BRUSH
+#undef TTF_FONT
+#undef OTF_FONT
+#undef TTF_CORE_FONT
+
+void FVREditorStyle::ReloadTextures()
+{
+	FSlateApplication::Get().GetRenderer()->ReloadTextureResources();
+}
+
+const ISlateStyle& FVREditorStyle::Get()
+{
+	return *VREditorStyleInstance;
+}
diff --git a/Engine/Source/Editor/VREditor/VREditorStyle.h b/Engine/Source/Editor/VREditor/VREditorStyle.h
new file mode 100644
--- /dev/null
+++ b/Engine/Source/Editor/VREditor/VREditorStyle.h
@@ -0,0 +1,34 @@
+// Copyright 1998-2017 Epic Games, Inc. All Rights Reserved.
+
+#pragma once
+
+
+/**  */
+class FVREditorStyle
+{
+public:
+
+	static void Initialize();
+
+	static void Shutdown();
+
+	/** reloads textures used by slate renderer */
+	static void ReloadTextures();
+
+	/** @return The Slate style set for the UMG Style */
+	static const ISlateStyle& Get();
+
+	static FName GetStyleSetName();
+
+	
+
+private:
+
+	static TSharedRef< class FSlateStyleSet > Create();
+
+private:
+
+	static TSharedPtr< class FSlateStyleSet > VREditorStyleInstance;
+
+	
+};
\ No newline at end of file
diff --git a/Engine/Source/Editor/VREditor/VREditorWorldInteraction.cpp b/Engine/Source/Editor/VREditor/VREditorWorldInteraction.cpp
--- a/Engine/Source/Editor/VREditor/VREditorWorldInteraction.cpp
+++ b/Engine/Source/Editor/VREditor/VREditorWorldInteraction.cpp
@@ -66,12 +66,14 @@ void UVREditorWorldInteraction::Init( UVREditorMode* InOwner, UViewportWorldInte
 	FEditorDelegates::OnAssetDragStarted.AddUObject( this, &UVREditorWorldInteraction::OnAssetDragStartedFromContentBrowser );
 
 	ViewportWorldInteraction->OnStopDragging().AddUObject( this, &UVREditorWorldInteraction::StopDragging );
+	ViewportWorldInteraction->OnWorldScaleChanged().AddUObject( this, &UVREditorWorldInteraction::UpdateNearClipPlaneOnScaleChange );
 }
 
 void UVREditorWorldInteraction::Shutdown()
 {
 	FEditorDelegates::OnAssetDragStarted.RemoveAll( this );
 	ViewportWorldInteraction->OnStopDragging().RemoveAll( this );
+	ViewportWorldInteraction->OnWorldScaleChanged().RemoveAll( this );
 
 	PlacingMaterialOrTextureAsset = nullptr;
 	FloatingUIAssetDraggedFrom = nullptr;
@@ -121,6 +123,13 @@ void UVREditorWorldInteraction::StopDragging( UViewportInteractor* Interactor )
 	}
 }
 
+void UVREditorWorldInteraction::UpdateNearClipPlaneOnScaleChange(const float NewWorldToMetersScale)
+{
+	// Adjust the clipping plane for the user's scale, but don't let it be larger than the engine default
+	const float DefaultWorldToMetersScale = Owner->GetSavedEditorState().WorldToMetersScale;
+	GNearClippingPlane = FMath::Min((Owner->GetDefaultVRNearClipPlane()) * (NewWorldToMetersScale / DefaultWorldToMetersScale), Owner->GetSavedEditorState().NearClipPlane);
+}
+
 void UVREditorWorldInteraction::StartDraggingMaterialOrTexture( UViewportInteractor* Interactor, const FViewportActionKeyInput& Action, const FVector HitLocation, UObject* MaterialOrTextureAsset )
 {
 	check( MaterialOrTextureAsset != nullptr );
diff --git a/Engine/Source/Editor/VREditor/VREditorWorldInteraction.h b/Engine/Source/Editor/VREditor/VREditorWorldInteraction.h
--- a/Engine/Source/Editor/VREditor/VREditorWorldInteraction.h
+++ b/Engine/Source/Editor/VREditor/VREditorWorldInteraction.h
@@ -40,6 +40,9 @@ protected:
 	/** When an interactor stops dragging */
 	void StopDragging( UViewportInteractor* Interactor );
 
+	/** When the world scale changes, update the near clip plane */
+	void UpdateNearClipPlaneOnScaleChange(const float NewWorldToMetersScale);
+
 	/** Starts dragging a material, allowing the user to drop it on an object in the scene to place it */
 	void StartDraggingMaterialOrTexture( UViewportInteractor* Interactor, const FViewportActionKeyInput& Action, const FVector HitLocation, UObject* MaterialOrTextureAsset );
 
diff --git a/Engine/Source/Editor/ViewportInteraction/Public/ViewportWorldInteraction.h b/Engine/Source/Editor/ViewportInteraction/Public/ViewportWorldInteraction.h
--- a/Engine/Source/Editor/ViewportInteraction/Public/ViewportWorldInteraction.h
+++ b/Engine/Source/Editor/ViewportInteraction/Public/ViewportWorldInteraction.h
@@ -174,6 +174,8 @@ public:
 
 	virtual FOnStopDragging& OnStopDragging() override { return OnStopDraggingEvent; };
 
+	virtual FOnWorldScaleChanged& OnWorldScaleChanged() override { return OnWorldScaleChangedEvent; };
+
 	/** Starts dragging selected actors around, changing selection if needed.  Called when clicking in the world, or when placing new objects */
 	void StartDraggingActors( UViewportInteractor* Interactor, const FViewportActionKeyInput& Action, UActorComponent* ClickedComponent, const FVector& HitLocation, const bool bIsPlacingActors );
 	
@@ -380,6 +382,9 @@ private:
 	/** Spawns the grid actor */
 	void SpawnGridMeshActor();
 
+	/** Average location of all the current transformables */
+	FVector CalculateAverageLocationOfTransformables();
+
 	//
 	// Colors
 	//
@@ -540,9 +545,12 @@ private:
 	/** Event that is fired when new axis input is received by the InputProcessor */
 	FOnHandleInputAxis OnAxisInputEvent;
 
-	/** Event this is fired when a interactor stopped dragging */
+	/** Event that is fired when a interactor stopped dragging */
 	FOnStopDragging OnStopDraggingEvent;
 
+	/** Event that is fired when the world scale changes */
+	FOnWorldScaleChanged OnWorldScaleChangedEvent;
+
 	/** If the world interaction is active and running the tick function */
 	bool bActive;
 };
diff --git a/Engine/Source/Editor/ViewportInteraction/Public/ViewportWorldInteractionInterface.h b/Engine/Source/Editor/ViewportInteraction/Public/ViewportWorldInteractionInterface.h
--- a/Engine/Source/Editor/ViewportInteraction/Public/ViewportWorldInteractionInterface.h
+++ b/Engine/Source/Editor/ViewportInteraction/Public/ViewportWorldInteractionInterface.h
@@ -55,4 +55,8 @@ public:
 	/** Gets the event for when an interactor stops dragging */
 	DECLARE_EVENT_OneParam( IViewportWorldInteractionInterface, FOnStopDragging, class UViewportInteractor* /** Interactor */ );
 	virtual FOnStopDragging& OnStopDragging() = 0;
-};
+
+	/** Gets the event for when the world scale changes */
+	DECLARE_EVENT_OneParam(IViewportWorldInteractionInterface, FOnWorldScaleChanged, const float /* NewWorldToMetersScale */);
+	virtual FOnWorldScaleChanged& OnWorldScaleChanged() = 0;
+};
\ No newline at end of file
diff --git a/Engine/Source/Editor/ViewportInteraction/ViewportInteractorData.h b/Engine/Source/Editor/ViewportInteraction/ViewportInteractorData.h
--- a/Engine/Source/Editor/ViewportInteraction/ViewportInteractorData.h
+++ b/Engine/Source/Editor/ViewportInteraction/ViewportInteractorData.h
@@ -138,6 +138,10 @@ struct FViewportInteractorData
 	/** If the latest hitresult is hovering over a priority type */
 	bool bHitResultIsPriorityType;
 
+	/** The offset between the hitlocation of the object selected to start dragging and its center. 
+		This is used to offset the objects when dragging to the end of the laser */
+	FVector StartHitLocationToTransformableCenter;
+
 	/** Default constructor for FVirtualHand that initializes safe defaults */
 	FViewportInteractorData()
 	{
@@ -173,6 +177,8 @@ struct FViewportInteractorData
 		bAllowTriggerLightPressLocking = true;
 		bAllowTriggerFullPress = true;
 		bHitResultIsPriorityType = false;
+
+		StartHitLocationToTransformableCenter = FVector::ZeroVector;
 	}
 };
 
diff --git a/Engine/Source/Editor/ViewportInteraction/ViewportWorldInteraction.cpp b/Engine/Source/Editor/ViewportInteraction/ViewportWorldInteraction.cpp
--- a/Engine/Source/Editor/ViewportInteraction/ViewportWorldInteraction.cpp
+++ b/Engine/Source/Editor/ViewportInteraction/ViewportWorldInteraction.cpp
@@ -2182,8 +2182,7 @@ void UViewportWorldInteraction::StartDraggingActors( UViewportInteractor* Intera
 		const bool bHaveLaserPointer = Interactor->GetLaserPointer( /* Out */ LaserPointerStart, /* Out */ LaserPointerEnd );
 		if ( bHaveLaserPointer )
 		{
-			AActor* Actor = ClickedComponent->GetOwner();
-
+			AActor* Actor = ClickedComponent->GetOwner(); 
 			FViewportInteractorData& InteractorData = Interactor->GetInteractorData();
 
 			// Capture undo state
@@ -2261,6 +2260,9 @@ void UViewportWorldInteraction::StartDraggingActors( UViewportInteractor* Intera
 
 			SetupTransformablesForSelectedActors();
 
+			// Calculate the offset between the average location and the hit location after setting the transformables for the new selected objects
+			InteractorData.StartHitLocationToTransformableCenter = CalculateAverageLocationOfTransformables() - HitLocation;
+
 			// If we're placing actors, start interpolating to their actual location.  This helps smooth everything out when
 			// using the laser impact point as the target transform
 			if ( bIsPlacingActors )
@@ -2430,7 +2432,9 @@ bool UViewportWorldInteraction::FindPlacementPointUnderLaser( UViewportInteracto
 			}
 
 			const FVector WorldSpaceExtremePointOnBox = InteractorData.GizmoStartTransform.TransformVectorNoScale( ExtremePointOnBox );
+
 			HitLocation -= WorldSpaceExtremePointOnBox;
+			HitLocation -= InteractorData.StartHitLocationToTransformableCenter;
 		}
 
 		OutHitLocation = HitLocation;
@@ -2800,6 +2804,7 @@ void UViewportWorldInteraction::SetWorldToMetersScale( const float NewWorldToMet
 
 	ENGINE_API extern float GNewWorldToMetersScale;
 	GNewWorldToMetersScale = NewWorldToMetersScale;
+	OnWorldScaleChangedEvent.Broadcast(NewWorldToMetersScale);
 }
 
 UViewportInteractor* UViewportWorldInteraction::GetOtherInteractorIntertiaContribute( UViewportInteractor* Interactor )
@@ -2975,6 +2980,20 @@ void UViewportWorldInteraction::SpawnGridMeshActor()
 	}
 }
 
+FVector UViewportWorldInteraction::CalculateAverageLocationOfTransformables()
+{
+	FVector Result = FVector::ZeroVector;
+
+	for( TUniquePtr<FViewportTransformable>& TransformablePtr : Transformables )
+	{
+		Result += ((*TransformablePtr).LastTransform.GetLocation());
+	}
+
+	Result /= Transformables.Num();
+
+	return Result;
+}
+
 FLinearColor UViewportWorldInteraction::GetColor( const EColors Color ) const
 {
 	return Colors[ (int32)Color ];
diff --git a/Engine/Source/Runtime/Engine/Classes/Components/SkeletalMeshComponent.h b/Engine/Source/Runtime/Engine/Classes/Components/SkeletalMeshComponent.h
--- a/Engine/Source/Runtime/Engine/Classes/Components/SkeletalMeshComponent.h
+++ b/Engine/Source/Runtime/Engine/Classes/Components/SkeletalMeshComponent.h
@@ -613,6 +613,15 @@ public:
 	UPROPERTY()
 	uint32 bEnableLineCheckWithBounds:1;
 
+protected:
+#if WITH_EDITORONLY_DATA
+	/** If true, this will Tick until disabled */
+	UPROPERTY(AdvancedDisplay, EditInstanceOnly, transient, Category = SkeletalMesh)
+	uint32 bUpdateAnimationInEditor : 1;
+#endif
+
+public:
+
 	/** Cache AnimCurveUidVersion from Skeleton and this will be used to identify if it needs to be updated */
 	UPROPERTY(transient)
 	uint16 CachedAnimCurveUidVersion;
@@ -829,6 +838,13 @@ public:
 	UFUNCTION(BlueprintCallable, Category="Components|SkeletalMesh", meta=(UnsafeDuringActorConstruction="true"))
 	void UnbindClothFromMasterPoseComponent(bool bRestoreSimulationSpace = true);
 
+	/**
+	* Sets whether or not to force tick component in order to update animation and refresh transform for this component
+	* This is supported only in the editor
+	*/
+	UFUNCTION(BlueprintCallable, Category = "Components|SkeletalMesh", meta = (DevelopmentOnly, UnsafeDuringActorConstruction = "true"))
+	void SetUpdateAnimationInEditor(const bool NewUpdateState);
+
 	/** We detach the Component once we are done playing it.
 	 *
 	 * @param	ParticleSystemComponent that finished
diff --git a/Engine/Source/Runtime/Engine/Private/Components/SkeletalMeshComponent.cpp b/Engine/Source/Runtime/Engine/Private/Components/SkeletalMeshComponent.cpp
--- a/Engine/Source/Runtime/Engine/Private/Components/SkeletalMeshComponent.cpp
+++ b/Engine/Source/Runtime/Engine/Private/Components/SkeletalMeshComponent.cpp
@@ -840,12 +840,21 @@ bool USkeletalMeshComponent::UpdateLODStatus()
 
 bool USkeletalMeshComponent::ShouldUpdateTransform(bool bLODHasChanged) const
 {
+
 #if WITH_EDITOR
+	
 	// If we're in an editor world (Non running, WorldType will be PIE when simulating or in PIE) then we only want transform updates on LOD changes as the
 	// animation isn't running so it would just waste CPU time
-	if(GetWorld()->WorldType == EWorldType::Editor && !bLODHasChanged)
+	if(GetWorld()->WorldType == EWorldType::Editor)
 	{
-		return false;
+		if( bUpdateAnimationInEditor )
+		{
+			return true;
+		}
+		if( !bLODHasChanged )
+		{
+			return false;
+		}
 	}
 #endif
 
@@ -862,9 +871,18 @@ bool USkeletalMeshComponent::ShouldTickPose() const
 	// So only enforce a single tick per frame.
 	const bool bAlreadyTickedThisFrame = PoseTickedThisFrame();
 
+#if WITH_EDITOR
+	if (GetWorld()->WorldType == EWorldType::Editor)
+	{
+		if (bUpdateAnimationInEditor)
+		{
+			return true;
+		}
+	}
+#endif 
+
 	// Autonomous Ticking is allowed to occur multiple times per frame, as we can receive and process multiple networking updates the same frame.
 	const bool bShouldTickBasedOnAutonomousCheck = bIsAutonomousTickPose || (!bOnlyAllowAutonomousTickPose && !bAlreadyTickedThisFrame);
-
 	const bool bIsPlayingNetworkedRootMotionMontage = (GetAnimInstance() != nullptr)
 		? (GetAnimInstance()->RootMotionMode == ERootMotionMode::RootMotionFromMontagesOnly) && (GetAnimInstance()->GetRootMotionMontageInstance() != nullptr)
 		: false;
@@ -2993,4 +3011,14 @@ void USkeletalMeshComponent::SnapshotPose(FPoseSnapshot& Snapshot)
 	Snapshot.bIsValid = true;
 }
 
+void USkeletalMeshComponent::SetUpdateAnimationInEditor(const bool NewUpdateState)
+{
+	#if WITH_EDITOR
+	if (IsRegistered())
+	{
+		bUpdateAnimationInEditor = NewUpdateState;
+	}
+	#endif
+}
+
 #undef LOCTEXT_NAMESPACE
diff --git a/Engine/Source/Runtime/Engine/Public/PhysicsPublic.h b/Engine/Source/Runtime/Engine/Public/PhysicsPublic.h
--- a/Engine/Source/Runtime/Engine/Public/PhysicsPublic.h
+++ b/Engine/Source/Runtime/Engine/Public/PhysicsPublic.h
@@ -10,6 +10,7 @@
 #include "CoreMinimal.h"
 #include "Stats/Stats.h"
 #include "Engine/EngineTypes.h"
+#include "Engine/World.h"
 #include "Misc/CoreMisc.h"
 #include "EngineDefines.h"
 #include "RenderResource.h"
@@ -542,6 +543,13 @@ public:
 	/** @return Whether physics scene is using substepping */
 	bool IsSubstepping(uint32 SceneType) const
 	{
+		// Substepping relies on interpolating transforms over frames, but only game worlds will be ticked,
+		// so we disallow this feature in non-game worlds.
+		if( !GetOwningWorld()->IsGameWorld() )
+		{
+			return false;
+		}
+
 		if (SceneType == PST_Sync) return bSubstepping;
 		if (SceneType == PST_Async) return bSubsteppingAsync;
 		return false;
diff --git a/Engine/Source/Runtime/Slate/Private/Framework/MultiBox/MultiBox.cpp b/Engine/Source/Runtime/Slate/Private/Framework/MultiBox/MultiBox.cpp
--- a/Engine/Source/Runtime/Slate/Private/Framework/MultiBox/MultiBox.cpp
+++ b/Engine/Source/Runtime/Slate/Private/Framework/MultiBox/MultiBox.cpp
@@ -257,7 +257,7 @@ void FMultiBox::InsertCustomMultiBlock( TSharedRef<const FMultiBlock> InBlock, i
  *
  * @return  MultiBox widget object
  */
-TSharedRef< SMultiBoxWidget > FMultiBox::MakeWidget( bool bSearchable )
+TSharedRef< SMultiBoxWidget > FMultiBox::MakeWidget( bool bSearchable, FOnMakeMultiBoxBuilderOverride* InMakeMultiBoxBuilderOverride /* = nullptr */ )
 {	
 	ApplyCustomizedBlocks();
 
@@ -270,8 +270,16 @@ TSharedRef< SMultiBoxWidget > FMultiBox::MakeWidget( bool bSearchable )
 	// Assign ourselves to the MultiBox widget
 	NewMultiBoxWidget->SetMultiBox( AsShared() );
 
-	// Build up the widget
-	NewMultiBoxWidget->BuildMultiBoxWidget();
+	if( (InMakeMultiBoxBuilderOverride != nullptr) && (InMakeMultiBoxBuilderOverride->IsBound()) )
+	{
+		TSharedRef<FMultiBox> ThisMultiBox = AsShared();
+		InMakeMultiBoxBuilderOverride->Execute( ThisMultiBox, NewMultiBoxWidget );
+	}
+	else
+	{
+		// Build up the widget
+		NewMultiBoxWidget->BuildMultiBoxWidget();
+	}
 	
 #if PLATFORM_MAC
 	if(Type == EMultiBoxType::MenuBar)
diff --git a/Engine/Source/Runtime/Slate/Private/Framework/MultiBox/MultiBoxBuilder.cpp b/Engine/Source/Runtime/Slate/Private/Framework/MultiBox/MultiBoxBuilder.cpp
--- a/Engine/Source/Runtime/Slate/Private/Framework/MultiBox/MultiBoxBuilder.cpp
+++ b/Engine/Source/Runtime/Slate/Private/Framework/MultiBox/MultiBoxBuilder.cpp
@@ -83,9 +83,9 @@ FMultiBoxCustomization FMultiBoxBuilder::GetCustomization() const
 	return FMultiBoxCustomization( MultiBox->GetCustomizationName() ); 
 }
 
-TSharedRef< class SWidget > FMultiBoxBuilder::MakeWidget()
+TSharedRef< class SWidget > FMultiBoxBuilder::MakeWidget( FMultiBox::FOnMakeMultiBoxBuilderOverride* InMakeMultiBoxBuilderOverride /* = nullptr */ )
 {
-	return MultiBox->MakeWidget( false );
+	return MultiBox->MakeWidget( false, InMakeMultiBoxBuilderOverride );
 }
 
 TSharedRef< class FMultiBox > FMultiBoxBuilder::GetMultiBox()
@@ -166,10 +166,10 @@ void FBaseMenuBuilder::AddMenuEntry( const FUIAction& UIAction, const TSharedRef
 	ApplyHook(InExtensionHook, EExtensionHook::After);
 }
 
-TSharedRef< class SWidget > FMenuBuilder::MakeWidget()
+TSharedRef< class SWidget > FMenuBuilder::MakeWidget( FMultiBox::FOnMakeMultiBoxBuilderOverride* InMakeMultiBoxBuilderOverride /* = nullptr */ )
 {
 	// Make menu builders searchable
-	return MultiBox->MakeWidget( true );
+	return MultiBox->MakeWidget( true, InMakeMultiBoxBuilderOverride );
 }
 
 void FMenuBuilder::BeginSection( FName InExtensionHook, const TAttribute< FText >& InHeadingText )
diff --git a/Engine/Source/Runtime/Slate/Private/Widgets/Input/SButton.cpp b/Engine/Source/Runtime/Slate/Private/Widgets/Input/SButton.cpp
--- a/Engine/Source/Runtime/Slate/Private/Widgets/Input/SButton.cpp
+++ b/Engine/Source/Runtime/Slate/Private/Widgets/Input/SButton.cpp
@@ -417,6 +417,16 @@ void SButton::SetOnClicked(FOnClicked InOnClicked)
 	OnClicked = InOnClicked;
 }
 
+void SButton::SetOnHovered(FSimpleDelegate InOnHovered)
+{
+	OnHovered = InOnHovered;
+}
+
+void SButton::SetOnUnhovered(FSimpleDelegate InOnUnhovered)
+{
+	OnUnhovered = InOnUnhovered;
+}
+
 void SButton::SetButtonStyle(const FButtonStyle* ButtonStyle)
 {
 	/* Get pointer to the button style */
diff --git a/Engine/Source/Runtime/Slate/Public/Framework/MultiBox/MultiBox.h b/Engine/Source/Runtime/Slate/Public/Framework/MultiBox/MultiBox.h
--- a/Engine/Source/Runtime/Slate/Public/Framework/MultiBox/MultiBox.h
+++ b/Engine/Source/Runtime/Slate/Public/Framework/MultiBox/MultiBox.h
@@ -255,12 +255,14 @@ public:
 	 */
 	void InsertCustomMultiBlock( TSharedRef<const FMultiBlock> InBlock, int32 Index );
 
+	DECLARE_DELEGATE_TwoParams( FOnMakeMultiBoxBuilderOverride, const TSharedRef<FMultiBox>&, const TSharedRef<SMultiBoxWidget>& );
+
 	/**
 	 * Creates a MultiBox widget for this MultiBox
 	 *
 	 * @return  MultiBox widget object
 	 */
-	TSharedRef< class SMultiBoxWidget > MakeWidget( bool bSearchable );
+	TSharedRef< class SMultiBoxWidget > MakeWidget( bool bSearchable, FOnMakeMultiBoxBuilderOverride* InMakeMultiBoxBuilderOverride = nullptr );
 
 
 	/**
@@ -491,6 +493,18 @@ public:
 		MultiBox = InMultiBox;
 	}
 
+	/**
+	* Sets new content for the SMultiBoxWidget
+	*
+	* @param	InContent	The new content to place in the ChildSlot
+	*/
+	void SetContent( TSharedRef< SWidget > InContent )
+	{	
+		ChildSlot
+		[
+			InContent 
+		];
+	}
 
 	/**
 	 * Access the MultiBox associated with this widget
diff --git a/Engine/Source/Runtime/Slate/Public/Framework/MultiBox/MultiBoxBuilder.h b/Engine/Source/Runtime/Slate/Public/Framework/MultiBox/MultiBoxBuilder.h
--- a/Engine/Source/Runtime/Slate/Public/Framework/MultiBox/MultiBoxBuilder.h
+++ b/Engine/Source/Runtime/Slate/Public/Framework/MultiBox/MultiBoxBuilder.h
@@ -2,6 +2,8 @@
 
 #pragma once
 
+#include "MultiBoxExtender.h"
+#include "SlateDelegates.h"
 #include "CoreMinimal.h"
 #include "Misc/Attribute.h"
 #include "Layout/Visibility.h"
@@ -13,6 +15,12 @@
 #include "Framework/Commands/UICommandList.h"
 #include "Framework/MultiBox/MultiBoxDefs.h"
 #include "Framework/MultiBox/MultiBoxExtender.h"
+#include "MultiBox.h"
+
+class FUICommandInfo;
+class FUICommandList;
+struct FSlateIcon;
+struct FUIAction;
 
 /** Delegate used by multi-box to call a user function to populate a new menu.  Used for spawning sub-menus and pull-down menus. */
 DECLARE_DELEGATE_OneParam( FNewMenuDelegate, class FMenuBuilder& );
@@ -56,7 +64,7 @@ public:
 	 *
 	 * @return  New widget object
 	 */
-	virtual TSharedRef< class SWidget > MakeWidget();
+	virtual TSharedRef< class SWidget > MakeWidget( FMultiBox::FOnMakeMultiBoxBuilderOverride* InMakeMultiBoxBuilderOverride = nullptr );
 	
 
 	/** 
@@ -213,7 +221,7 @@ public:
 	*
 	* @return  New widget object
 	*/
-	virtual TSharedRef< class SWidget > MakeWidget() override;
+	virtual TSharedRef< class SWidget > MakeWidget( FMultiBox::FOnMakeMultiBoxBuilderOverride* InMakeMultiBoxBuilderOverride = nullptr ) override;
 
 	/**
 	 * Adds a menu separator
diff --git a/Engine/Source/Runtime/Slate/Public/Widgets/Input/SButton.h b/Engine/Source/Runtime/Slate/Public/Widgets/Input/SButton.h
--- a/Engine/Source/Runtime/Slate/Public/Widgets/Input/SButton.h
+++ b/Engine/Source/Runtime/Slate/Public/Widgets/Input/SButton.h
@@ -147,6 +147,12 @@ public:
 	/** See OnClicked event */
 	void SetOnClicked(FOnClicked InOnClicked);
 
+	/** Set OnHovered event */
+	void SetOnHovered(FSimpleDelegate InOnHovered);
+
+	/** Set OnUnhovered event */
+	void SetOnUnhovered(FSimpleDelegate InOnUnhovered);
+
 	/** See ButtonStyle attribute */
 	void SetButtonStyle(const FButtonStyle* ButtonStyle);
 
