Fixing VR Editor UI non-unity CIS errors

diff --git a/Engine/Source/Editor/VREditor/UI/VREditorUISystem.cpp b/Engine/Source/Editor/VREditor/UI/VREditorUISystem.cpp
--- a/Engine/Source/Editor/VREditor/UI/VREditorUISystem.cpp
+++ b/Engine/Source/Editor/VREditor/UI/VREditorUISystem.cpp
@@ -6,15 +6,23 @@
 #include "GenericPlatform/GenericApplication.h"
 #include "Modules/ModuleManager.h"
 #include "InputCoreTypes.h"
-#include "Widgets/SNullWidget.h"
 #include "Input/Events.h"
 #include "Input/Reply.h"
+
+#include "Widgets/SNullWidget.h"
 #include "Widgets/DeclarativeSyntaxSupport.h"
 #include "Widgets/SWidget.h"
 #include "Widgets/SWindow.h"
-#include "Layout/WidgetPath.h"
+#include "Widgets/Input/SButton.h"
+#include "Widgets/SOverlay.h"
+#include "Widgets/Layout/SUniformGridPanel.h"
+#include "Widgets/Layout/SBorder.h"
+
+#include "Framework/Multibox/Multibox.h"
 #include "Framework/Application/SlateApplication.h"
 #include "Framework/Docking/TabManager.h"
+#include "Layout/WidgetPath.h"
+
 #include "Engine/EngineTypes.h"
 #include "Editor/UnrealEdEngine.h"
 #include "Kismet/GameplayStatics.h"
@@ -23,6 +31,7 @@
 #include "EditorModeManager.h"
 #include "EditorModes.h"
 #include "UnrealEdGlobals.h"
+
 #include "VREditorMode.h"
 #include "VREditorBaseActor.h"
 #include "VREditorBaseUserWidget.h"
diff --git a/Engine/Source/Editor/VREditor/UI/VREditorUISystem.h b/Engine/Source/Editor/VREditor/UI/VREditorUISystem.h
--- a/Engine/Source/Editor/VREditor/UI/VREditorUISystem.h
+++ b/Engine/Source/Editor/VREditor/UI/VREditorUISystem.h
@@ -12,6 +12,12 @@ class AVREditorDockableWindow;
 class AVREditorFloatingUI;
 class FProxyTabmanager;
 class SColorPicker;
+class SBorder;
+class SButton;
+class SOverlay;
+class FMultibox;
+class SMultiBoxWidget;
+class SWidget;
 class UViewportInteractor;
 class UVREditorInteractor;
 
diff --git a/Engine/Source/Editor/VREditor/VREditorStyle.cpp b/Engine/Source/Editor/VREditor/VREditorStyle.cpp
--- a/Engine/Source/Editor/VREditor/VREditorStyle.cpp
+++ b/Engine/Source/Editor/VREditor/VREditorStyle.cpp
@@ -1,11 +1,15 @@
 // Copyright 1998-2017 Epic Games, Inc. All Rights Reserved.
 
 #include "VREditorStyle.h"
-#include "VREditorModule.h"
 #include "SlateTypes.h"
-#include "SlateGameResources.h"
+#include "EditorStyleSet.h"
 #include "Styling/SlateStyleRegistry.h"
 #include "SUniformGridPanel.h"
+#include "Brushes/SlateImageBrush.h"
+#include "Brushes/SlateBoxBrush.h"
+#include "Brushes/SlateBorderBrush.h"
+#include "Fonts/SlateFontInfo.h"
+#include "Framework/Application/SlateApplication.h"
 
 TSharedPtr< FSlateStyleSet > FVREditorStyle::VREditorStyleInstance = NULL;
 
diff --git a/Engine/Source/Editor/VREditor/VREditorStyle.h b/Engine/Source/Editor/VREditor/VREditorStyle.h
--- a/Engine/Source/Editor/VREditor/VREditorStyle.h
+++ b/Engine/Source/Editor/VREditor/VREditorStyle.h
@@ -2,6 +2,8 @@
 
 #pragma once
 
+#include "CoreMinimal.h"
+#include "Styling/SlateStyle.h"
 
 /**  */
 class FVREditorStyle
