Fixing Linux compile error for missing include

diff --git a/Engine/Source/Editor/VREditor/UI/VREditorUISystem.cpp b/Engine/Source/Editor/VREditor/UI/VREditorUISystem.cpp
--- a/Engine/Source/Editor/VREditor/UI/VREditorUISystem.cpp
+++ b/Engine/Source/Editor/VREditor/UI/VREditorUISystem.cpp
@@ -18,7 +18,7 @@
 #include "Widgets/Layout/SUniformGridPanel.h"
 #include "Widgets/Layout/SBorder.h"
 
-#include "Framework/Multibox/Multibox.h"
+#include "Framework/MultiBox/MultiBox.h"
 #include "Framework/Application/SlateApplication.h"
 #include "Framework/Docking/TabManager.h"
 #include "Layout/WidgetPath.h"
diff --git a/Engine/Source/Editor/VREditor/UI/VREditorUISystem.h b/Engine/Source/Editor/VREditor/UI/VREditorUISystem.h
--- a/Engine/Source/Editor/VREditor/UI/VREditorUISystem.h
+++ b/Engine/Source/Editor/VREditor/UI/VREditorUISystem.h
@@ -15,7 +15,7 @@ class SColorPicker;
 class SBorder;
 class SButton;
 class SOverlay;
-class FMultibox;
+class FMultiBox;
 class SMultiBoxWidget;
 class SWidget;
 class UViewportInteractor;
