Fixing some localization key conflicts that crept in from stream merges

diff --git a/Engine/Source/Editor/BlueprintGraph/Private/K2Node_MakeVariable.cpp b/Engine/Source/Editor/BlueprintGraph/Private/K2Node_MakeVariable.cpp
--- a/Engine/Source/Editor/BlueprintGraph/Private/K2Node_MakeVariable.cpp
+++ b/Engine/Source/Editor/BlueprintGraph/Private/K2Node_MakeVariable.cpp
@@ -194,5 +194,3 @@ UEdGraphPin* UK2Node_MakeVariable::GetOutputPin() const
 {
 	return FindPinChecked(MakeVariableOutputPinName);
 }
-
-#undef LOCTEXT_NAMESPACE
diff --git a/Engine/Source/Editor/DetailCustomizations/Private/WindowsTargetSettingsDetails.cpp b/Engine/Source/Editor/DetailCustomizations/Private/WindowsTargetSettingsDetails.cpp
--- a/Engine/Source/Editor/DetailCustomizations/Private/WindowsTargetSettingsDetails.cpp
+++ b/Engine/Source/Editor/DetailCustomizations/Private/WindowsTargetSettingsDetails.cpp
@@ -67,7 +67,7 @@ static FText GetFriendlyNameFromRHIName(const FString& InRHIName)
 	}
 	else if (InRHIName == TEXT("GLSL_150_ES31"))
 	{
-		FriendlyRHIName = LOCTEXT("OpenGL3ES2", "OpenGL 3 (ES3.1, Experimental)");
+		FriendlyRHIName = LOCTEXT("OpenGL3ES31", "OpenGL 3 (ES3.1, Experimental)");
 	}
 	else if (InRHIName == TEXT("GLSL_430"))
 	{
diff --git a/Engine/Source/Editor/Kismet/Private/BlueprintDetailsCustomization.cpp b/Engine/Source/Editor/Kismet/Private/BlueprintDetailsCustomization.cpp
--- a/Engine/Source/Editor/Kismet/Private/BlueprintDetailsCustomization.cpp
+++ b/Engine/Source/Editor/Kismet/Private/BlueprintDetailsCustomization.cpp
@@ -3272,7 +3272,7 @@ void FBlueprintGraphActionDetails::CustomizeDetails( IDetailLayoutBuilder& Detai
 						[
 							SNew(STextBlock)
 							.Font(IDetailLayoutBuilder::GetDetailFontBold())
-							.Text(LOCTEXT("FunctionNewInputArg", "New Parameter"))
+							.Text(LOCTEXT("FunctionNewParameterInputArg", "New Parameter"))
 							.Visibility(this, &FBlueprintGraphActionDetails::OnGetSectionTextVisibility, WeakInputsHeaderWidget)
 							.ShadowOffset(FVector2D(1, 1))
 						]
diff --git a/Engine/Source/Editor/LevelEditor/Private/LevelEditorContextMenu.cpp b/Engine/Source/Editor/LevelEditor/Private/LevelEditorContextMenu.cpp
--- a/Engine/Source/Editor/LevelEditor/Private/LevelEditorContextMenu.cpp
+++ b/Engine/Source/Editor/LevelEditor/Private/LevelEditorContextMenu.cpp
@@ -390,7 +390,7 @@ void FLevelEditorContextMenu::FillMenu( FMenuBuilder& MenuBuilder, TWeakPtr<SLev
 			MenuBuilder.AddMenuEntry(FLevelEditorCommands::Get().OpenMergeActor,
 					NAME_None,
 					LOCTEXT("OpenMergeActor", "Merge Actors"),
-					LOCTEXT("GoToDocsForActor_ToolTip", "Click to open the Merge Actor panel"));
+					LOCTEXT("OpenMergeActor_ToolTip", "Click to open the Merge Actor panel"));
 		}
 
 		if (GEditor->PlayWorld != NULL)
diff --git a/Engine/Source/Editor/UnrealEd/Private/Commandlets/GatherTextFromSourceCommandlet.cpp b/Engine/Source/Editor/UnrealEd/Private/Commandlets/GatherTextFromSourceCommandlet.cpp
--- a/Engine/Source/Editor/UnrealEd/Private/Commandlets/GatherTextFromSourceCommandlet.cpp
+++ b/Engine/Source/Editor/UnrealEd/Private/Commandlets/GatherTextFromSourceCommandlet.cpp
@@ -1361,7 +1361,7 @@ void UGatherTextFromSourceCommandlet::FStringMacroDescriptor::TryParse(const FSt
 				if ( Identifier.IsEmpty() )
 				{
 					//The command doesn't have an identifier so we can't gather it
-					UE_LOG(LogGatherTextFromSourceCommandlet, Warning, TEXT("Macro doesn't have unique identifier. %s"), *SourceLocation );
+					UE_LOG(LogGatherTextFromSourceCommandlet, Warning, TEXT("Localization macro has an empty identifier and cannot be gathered. %s"), *SourceLocation );
 					return;
 				}
 
diff --git a/Engine/Source/Editor/VREditor/UI/VREditorUISystem.cpp b/Engine/Source/Editor/VREditor/UI/VREditorUISystem.cpp
--- a/Engine/Source/Editor/VREditor/UI/VREditorUISystem.cpp
+++ b/Engine/Source/Editor/VREditor/UI/VREditorUISystem.cpp
@@ -2365,7 +2365,7 @@ TSharedRef<SWidget> UVREditorUISystem::BuildNumPadWidget()
 
 	MenuBuilder.AddMenuEntry(
 		LOCTEXT("NumThree", "3"),
-		LOCTEXT("", ""),
+		FText::GetEmpty(),
 		FSlateIcon(),
 		FUIAction
 		(
@@ -2374,7 +2374,7 @@ TSharedRef<SWidget> UVREditorUISystem::BuildNumPadWidget()
 		);
 	MenuBuilder.AddMenuEntry(
 		LOCTEXT("NumFour", "4"),
-		LOCTEXT("", ""),
+		FText::GetEmpty(),
 		FSlateIcon(),
 		FUIAction
 		(
@@ -2383,7 +2383,7 @@ TSharedRef<SWidget> UVREditorUISystem::BuildNumPadWidget()
 		);
 	MenuBuilder.AddMenuEntry(
 		LOCTEXT("NumFive", "5"),
-		LOCTEXT("", ""),
+		FText::GetEmpty(),
 		FSlateIcon(),
 		FUIAction
 		(
@@ -2392,7 +2392,7 @@ TSharedRef<SWidget> UVREditorUISystem::BuildNumPadWidget()
 		);
 	MenuBuilder.AddMenuEntry(
 		LOCTEXT("NumSix", "6"),
-		LOCTEXT("", ""),
+		FText::GetEmpty(),
 		FSlateIcon(),
 		FUIAction
 		(
@@ -2401,7 +2401,7 @@ TSharedRef<SWidget> UVREditorUISystem::BuildNumPadWidget()
 		);
 	MenuBuilder.AddMenuEntry(
 		LOCTEXT("NumSeven", "7"),
-		LOCTEXT("", ""),
+		FText::GetEmpty(),
 		FSlateIcon(),
 		FUIAction
 		(
@@ -2410,7 +2410,7 @@ TSharedRef<SWidget> UVREditorUISystem::BuildNumPadWidget()
 		);
 	MenuBuilder.AddMenuEntry(
 		LOCTEXT("NumEight", "8"),
-		LOCTEXT("", ""),
+		FText::GetEmpty(),
 		FSlateIcon(),
 		FUIAction
 		(
@@ -2419,7 +2419,7 @@ TSharedRef<SWidget> UVREditorUISystem::BuildNumPadWidget()
 		);
 	MenuBuilder.AddMenuEntry(
 		LOCTEXT("NumNine", "9"),
-		LOCTEXT("", ""),
+		FText::GetEmpty(),
 		FSlateIcon(),
 		FUIAction
 		(
@@ -2428,7 +2428,7 @@ TSharedRef<SWidget> UVREditorUISystem::BuildNumPadWidget()
 		);
 	MenuBuilder.AddMenuEntry(
 		LOCTEXT( "Hyphen", "-" ),
-		LOCTEXT( "", "" ),
+		FText::GetEmpty(),
 		FSlateIcon(),
 		FUIAction
 		(
@@ -2437,7 +2437,7 @@ TSharedRef<SWidget> UVREditorUISystem::BuildNumPadWidget()
 	);
 	MenuBuilder.AddMenuEntry(
 		LOCTEXT( "Decimal", "." ),
-		LOCTEXT( "", "" ),
+		FText::GetEmpty(),
 		FSlateIcon(),
 		FUIAction
 		(
@@ -2446,7 +2446,7 @@ TSharedRef<SWidget> UVREditorUISystem::BuildNumPadWidget()
 	);
 	MenuBuilder.AddMenuEntry(
 		LOCTEXT( "NumZero", "0" ),
-		LOCTEXT( "", "" ),
+		FText::GetEmpty(),
 		FSlateIcon(),
 		FUIAction
 		(
@@ -2455,7 +2455,7 @@ TSharedRef<SWidget> UVREditorUISystem::BuildNumPadWidget()
 	);
 	MenuBuilder.AddMenuEntry(
 		LOCTEXT( "NumOne", "1" ),
-		LOCTEXT( "", "" ),
+		FText::GetEmpty(),
 		FSlateIcon(),
 		FUIAction
 		(
@@ -2464,7 +2464,7 @@ TSharedRef<SWidget> UVREditorUISystem::BuildNumPadWidget()
 	);
 	MenuBuilder.AddMenuEntry(
 		LOCTEXT( "NumTwo", "2" ),
-		LOCTEXT( "", "" ),
+		FText::GetEmpty(),
 		FSlateIcon(),
 		FUIAction
 		(
