Fixing non-unity CIS build errors for VR Editor UI

diff --git a/Engine/Source/Editor/VREditor/SVRRadialPanel.h b/Engine/Source/Editor/VREditor/SVRRadialPanel.h
--- a/Engine/Source/Editor/VREditor/SVRRadialPanel.h
+++ b/Engine/Source/Editor/VREditor/SVRRadialPanel.h
@@ -2,7 +2,16 @@
 
 #pragma once
 
-#include "SPanel.h"
+#include "CoreMinimal.h"
+#include "Misc/Attribute.h"
+#include "Layout/Margin.h"
+#include "Layout/Visibility.h"
+#include "Types/SlateStructs.h"
+#include "Widgets/DeclarativeSyntaxSupport.h"
+#include "SlotBase.h"
+#include "Widgets/SWidget.h"
+#include "Layout/Children.h"
+#include "Widgets/SPanel.h"
 
 /**
  * A RadialPanel contains one child and describes how that child should be arranged on the screen.
diff --git a/Engine/Source/Editor/VREditor/UI/VREditorUISystem.cpp b/Engine/Source/Editor/VREditor/UI/VREditorUISystem.cpp
--- a/Engine/Source/Editor/VREditor/UI/VREditorUISystem.cpp
+++ b/Engine/Source/Editor/VREditor/UI/VREditorUISystem.cpp
@@ -14,11 +14,12 @@
 #include "Widgets/SWidget.h"
 #include "Widgets/SWindow.h"
 #include "Widgets/Input/SButton.h"
-#include "Widgets/SOverlay.h"
 #include "Widgets/Layout/SUniformGridPanel.h"
 #include "Widgets/Layout/SBorder.h"
+#include "Widgets/Layout/SBox.h"
 
 #include "Framework/MultiBox/MultiBox.h"
+#include "Framework/MultiBox/MultiBoxBuilder.h"
 #include "Framework/Application/SlateApplication.h"
 #include "Framework/Docking/TabManager.h"
 #include "Layout/WidgetPath.h"
diff --git a/Engine/Source/Editor/VREditor/UI/VREditorUISystem.h b/Engine/Source/Editor/VREditor/UI/VREditorUISystem.h
--- a/Engine/Source/Editor/VREditor/UI/VREditorUISystem.h
+++ b/Engine/Source/Editor/VREditor/UI/VREditorUISystem.h
@@ -6,6 +6,7 @@
 #include "SlateFwd.h"
 #include "UObject/ObjectMacros.h"
 #include "UObject/Object.h"
+#include "Widgets/SOverlay.h"
 #include "VREditorUISystem.generated.h"
 
 class AVREditorDockableWindow;
@@ -14,7 +15,6 @@ class FProxyTabmanager;
 class SColorPicker;
 class SBorder;
 class SButton;
-class SOverlay;
 class FMultiBox;
 class SMultiBoxWidget;
 class SWidget;
diff --git a/Engine/Source/Editor/VREditor/VREditorActions.cpp b/Engine/Source/Editor/VREditor/VREditorActions.cpp
--- a/Engine/Source/Editor/VREditor/VREditorActions.cpp
+++ b/Engine/Source/Editor/VREditor/VREditorActions.cpp
@@ -12,7 +12,7 @@
 #include "SLevelViewport.h"
 #include "ImageUtils.h"
 #include "FileHelper.h"
-
+#include "Framework/Application/SlateApplication.h"
 
 #define LOCTEXT_NAMESPACE "VREditorActions"
 
diff --git a/Engine/Source/Editor/VREditor/VREditorActions.h b/Engine/Source/Editor/VREditor/VREditorActions.h
--- a/Engine/Source/Editor/VREditor/VREditorActions.h
+++ b/Engine/Source/Editor/VREditor/VREditorActions.h
@@ -2,6 +2,8 @@
 #pragma once
 
 #include "VREditorUISystem.h"
+#include "SlateTypes.h"
+#include "InputCoreTypes.h"
 
 namespace EVREditorActions
 {
