Fixing some localization key conflicts that crept in from stream merges

diff --git a/Engine/Source/Editor/VREditor/VREditorActions.h b/Engine/Source/Editor/VREditor/VREditorActions.h
--- a/Engine/Source/Editor/VREditor/VREditorActions.h
+++ b/Engine/Source/Editor/VREditor/VREditorActions.h
@@ -8,10 +8,10 @@
 namespace EVREditorActions
 {
 	//Gizmo coordinate button text starts blank before it is initialized
-	static FText GizmoCoordinateSystemText = NSLOCTEXT("GizmoCoordinateDefault", "", "");
+	static FText GizmoCoordinateSystemText;
 
 	//Gizmo mode button text starts blank before it is initialized
-	static FText GizmoModeText = NSLOCTEXT("GizmoModeDefault", "", "");;
+	static FText GizmoModeText;
 }
 
 /**
