#vreditor Sequencer UI now spawns on the hand not hovering over UI

diff --git a/Engine/Source/Editor/VREditor/UI/VREditorUISystem.cpp b/Engine/Source/Editor/VREditor/UI/VREditorUISystem.cpp
--- a/Engine/Source/Editor/VREditor/UI/VREditorUISystem.cpp
+++ b/Engine/Source/Editor/VREditor/UI/VREditorUISystem.cpp
@@ -2669,7 +2669,7 @@ void UVREditorUISystem::UpdateSequencerUI()
 		const bool bWithSceneComponent = false;
 		SequencerPanel->SetSlateWidget(*this, WidgetToDraw, FIntPoint(VREd::ContentBrowserUIResolutionX->GetFloat(), VREd::ContentBrowserUIResolutionY->GetFloat()), VREd::ContentBrowserUISize->GetFloat(), AVREditorFloatingUI::EDockedTo::Nothing);
 		UVREditorInteractor* VREditorInteractor = VRMode->GetHandInteractor(EControllerHand::Left);
-		if (!VREditorInteractor->IsHoveringOverUI())
+		if (VREditorInteractor->IsHoveringOverUI())
 		{
 			VREditorInteractor = VRMode->GetHandInteractor(EControllerHand::Right);
 
