Fix for crash when getting the world when creating help menus on the motion controllers.

diff --git a/Engine/Source/Editor/VREditor/VREditorMotionControllerInteractor.cpp b/Engine/Source/Editor/VREditor/VREditorMotionControllerInteractor.cpp
--- a/Engine/Source/Editor/VREditor/VREditorMotionControllerInteractor.cpp
+++ b/Engine/Source/Editor/VREditor/VREditorMotionControllerInteractor.cpp
@@ -1104,10 +1104,10 @@ void UVREditorMotionControllerInteractor::ShowHelpForHand( const bool bShowIt )
 						LabelText = LOCTEXT( "ConfirmRadialSelectionHelp", "Radial Menu" );
 						ComponentName = TEXT( "ConfirmRadialSelectionHelp" );
 					}
-
+					
 					const bool bWithSceneComponent = false;	// Nope, we'll spawn our own inside AFloatingText
 					check( VRMode );
-					AFloatingText* FloatingText = UViewportWorldInteraction::SpawnTransientSceneActor<AFloatingText>( GetWorld(), ComponentName, bWithSceneComponent );
+					AFloatingText* FloatingText = UViewportWorldInteraction::SpawnTransientSceneActor<AFloatingText>( GetVRMode().GetWorld(), ComponentName, bWithSceneComponent );
 					FloatingText->SetText( LabelText );
 
 					HelpLabels.Add( Key, FloatingText );
@@ -1146,7 +1146,7 @@ void UVREditorMotionControllerInteractor::UpdateHelpLabels()
 		for ( auto& KeyAndValue : HelpLabels )
 		{
 			AFloatingText* FloatingText = KeyAndValue.Value;
-			UViewportWorldInteraction::DestroyTransientActor( GetWorld(), FloatingText );
+			UViewportWorldInteraction::DestroyTransientActor( GetVRMode().GetWorld(), FloatingText );
 		}
 		HelpLabels.Reset();
 	}
