Fixing another GetWorld() call to go through the VRMode on the VR motion controller interactor

diff --git a/Engine/Source/Editor/VREditor/VREditorInteractor.cpp b/Engine/Source/Editor/VREditor/VREditorInteractor.cpp
--- a/Engine/Source/Editor/VREditor/VREditorInteractor.cpp
+++ b/Engine/Source/Editor/VREditor/VREditorInteractor.cpp
@@ -47,7 +47,7 @@ void UVREditorInteractor::Shutdown()
 	for ( auto& KeyAndValue : HelpLabels )
 	{
 		AFloatingText* FloatingText = KeyAndValue.Value;
-		UViewportWorldInteraction::DestroyTransientActor( GetWorld(), FloatingText );
+		UViewportWorldInteraction::DestroyTransientActor( GetVRMode().GetWorld(), FloatingText );
 	}
 
 	HelpLabels.Empty();	
