Adjusting resolution of Sequencer window so keys are in the correct position

diff --git a/Engine/Source/Editor/VREditor/UI/VREditorUISystem.cpp b/Engine/Source/Editor/VREditor/UI/VREditorUISystem.cpp
--- a/Engine/Source/Editor/VREditor/UI/VREditorUISystem.cpp
+++ b/Engine/Source/Editor/VREditor/UI/VREditorUISystem.cpp
@@ -94,6 +94,8 @@ namespace VREd
 {
 	static FAutoConsoleVariable ContentBrowserUIResolutionX( TEXT( "VREd.ContentBrowserUIResolutionX" ), 1920, TEXT( "Horizontal resolution to use for content browser UI render targets" ) );
 	static FAutoConsoleVariable ContentBrowserUIResolutionY( TEXT( "VREd.ContentBrowserUIResolutionY" ), 1200, TEXT( "Vertical resolution to use for content browser UI render targets" ) );
+	static FAutoConsoleVariable SequencerUIResolutionX(TEXT("VREd.SequencerUIResolutionX"), 960, TEXT("Horizontal resolution to use for Sequencer UI render targets"));
+	static FAutoConsoleVariable SequencerUIResolutionY(TEXT("VREd.SequencerUIResolutionY"), 600, TEXT("Vertical resolution to use for Sequencer UI render targets"));
 	static FAutoConsoleVariable DefaultEditorUIResolutionX( TEXT( "VREd.DefaultEditorUIResolutionX" ), 1024, TEXT( "Horizontal resolution to use for VR editor UI render targets" ) );
 	static FAutoConsoleVariable DefaultEditorUIResolutionY( TEXT( "VREd.DefaultEditorUIResolutionY" ), 1024, TEXT( "Vertical resolution to use for VR editor UI render targets" ) );
 	static FAutoConsoleVariable QuickMenuUIResolutionX( TEXT( "VREd.QuickMenuUIResolutionX" ), 1440, TEXT( "Horizontal resolution to use for Quick Menu VR UI render targets" ) );
@@ -2652,7 +2654,6 @@ void UVREditorUISystem::UpdateSequencerUI()
 	if (EditorUIPanels[(int32)EEditorUIPanel::SequencerUI] != nullptr)
 	{
 		AVREditorFloatingUI* SequencerPanel = EditorUIPanels[(int32)EEditorUIPanel::SequencerUI];
-		const FIntPoint DefaultResolution(VREd::DefaultEditorUIResolutionX->GetInt(), VREd::DefaultEditorUIResolutionY->GetInt());
 		const TSharedRef< ILevelEditor >& LevelEditor = FModuleManager::GetModuleChecked<FLevelEditorModule>("LevelEditor").GetFirstLevelEditor().ToSharedRef();
 
 		const FName TabIdentifier = NAME_None;	// No tab for us!
@@ -2660,14 +2661,14 @@ void UVREditorUISystem::UpdateSequencerUI()
 
 		TSharedRef<SWidget> WidgetToDraw =
 			SNew(SDPIScaler)
-			.DPIScale(VREd::EditorUIScale->GetFloat())
+			.DPIScale(1.0f)
 			[
 				SequencerWidget
 			]
 		;
 
 		const bool bWithSceneComponent = false;
-		SequencerPanel->SetSlateWidget(*this, WidgetToDraw, FIntPoint(VREd::ContentBrowserUIResolutionX->GetFloat(), VREd::ContentBrowserUIResolutionY->GetFloat()), VREd::ContentBrowserUISize->GetFloat(), AVREditorFloatingUI::EDockedTo::Nothing);
+		SequencerPanel->SetSlateWidget(*this, WidgetToDraw, FIntPoint(VREd::SequencerUIResolutionX->GetFloat(), VREd::SequencerUIResolutionY->GetFloat()), VREd::EditorUISize->GetFloat(), AVREditorFloatingUI::EDockedTo::Nothing);
 		UVREditorInteractor* VREditorInteractor = VRMode->GetHandInteractor(EControllerHand::Left);
 		if (VREditorInteractor->IsHoveringOverUI())
 		{
