Changing drag-drop events in VR to synthetic events

diff --git a/Engine/Source/Editor/VREditor/UI/VREditorUISystem.cpp b/Engine/Source/Editor/VREditor/UI/VREditorUISystem.cpp
--- a/Engine/Source/Editor/VREditor/UI/VREditorUISystem.cpp
+++ b/Engine/Source/Editor/VREditor/UI/VREditorUISystem.cpp
@@ -607,7 +607,7 @@ void UVREditorUISystem::OnVRHoverUpdate( UViewportInteractor* Interactor, FVecto
 								PressedButtons,
 								FModifierKeysState());
 
-							FSlateApplication::Get().RoutePointerMoveEvent(WidgetPathUnderFinger, PointerEvent, false);
+							FSlateApplication::Get().RoutePointerMoveEvent(WidgetPathUnderFinger, PointerEvent, true);
 							
 							bWasHandled = true;
 							HoverImpactPoint = HitResult.ImpactPoint;
@@ -730,7 +730,7 @@ void UVREditorUISystem::OnVRHoverUpdate( UViewportInteractor* Interactor, FVecto
 				PressedButtons,
 				FModifierKeysState() );
 
-			FSlateApplication::Get().RoutePointerMoveEvent( FWidgetPath(), PointerEvent, false );
+			FSlateApplication::Get().RoutePointerMoveEvent( FWidgetPath(), PointerEvent, true );
 		}
 	}
 }
