Disabling Sequencer fading and color scaling in VR Editor. Refreshing mesh trails when track data changes.

diff --git a/Engine/Source/Editor/Sequencer/Private/Sequencer.cpp b/Engine/Source/Editor/Sequencer/Private/Sequencer.cpp
--- a/Engine/Source/Editor/Sequencer/Private/Sequencer.cpp
+++ b/Engine/Source/Editor/Sequencer/Private/Sequencer.cpp
@@ -986,6 +986,14 @@ void FSequencer::NotifyMovieSceneDataChanged( EMovieSceneDataChangeType DataChan
 		bNeedInstanceRefresh = true;
 	}
 
+	if (DataChangeType == EMovieSceneDataChangeType::TrackValueChanged || 
+		DataChangeType == EMovieSceneDataChangeType::TrackValueChangedRefreshImmediately || 
+		DataChangeType == EMovieSceneDataChangeType::Unknown)
+	{
+		FSequencerEdMode* SequencerEdMode = (FSequencerEdMode*)(GLevelEditorModeTools().GetActiveMode(FSequencerEdMode::EM_SequencerMode));
+		SequencerEdMode->CleanUpMeshTrails();
+	}
+
 	bNeedsEvaluate = true;
 	State.ClearObjectCaches();
 
@@ -1939,6 +1947,11 @@ void FSequencer::SetViewportSettings(const TMap<FViewportClient*, EMovieSceneVie
 				}
 			}
 		}
+		else
+		{
+			LevelVC->bEnableFading = false;
+			LevelVC->bEnableColorScaling = false;
+		}
 	}
 }
 
diff --git a/Engine/Source/Editor/VREditor/Public/VREditorMode.h b/Engine/Source/Editor/VREditor/Public/VREditorMode.h
--- a/Engine/Source/Editor/VREditor/Public/VREditorMode.h
+++ b/Engine/Source/Editor/VREditor/Public/VREditorMode.h
@@ -258,6 +258,7 @@ public:
 		float TransformGizmoScale;
 		EHMDTrackingOrigin::Type TrackingOrigin;
 		float WorldToMetersScale;
+		bool bCinematicPreviewViewport;
 		bool bKeyAllEnabled;
 		EAutoKeyMode AutoKeyMode;
 
@@ -276,6 +277,7 @@ public:
 			  TransformGizmoScale( 1.0f ),
 			  TrackingOrigin(EHMDTrackingOrigin::Eye),
 			  WorldToMetersScale(100.0f),
+			  bCinematicPreviewViewport(false),
 			  bKeyAllEnabled(false),
 			  AutoKeyMode()
 		{
diff --git a/Engine/Source/Editor/VREditor/VREditorMode.cpp b/Engine/Source/Editor/VREditor/VREditorMode.cpp
--- a/Engine/Source/Editor/VREditor/VREditorMode.cpp
+++ b/Engine/Source/Editor/VREditor/VREditorMode.cpp
@@ -301,11 +301,20 @@ void UVREditorMode::Enter(const bool bReenteringVREditing)
 			// Enable selection outline right away
 			VREditorViewportClient.EngineShowFlags.SetSelection( true );
 			VREditorViewportClient.EngineShowFlags.SetSelectionOutline( true );
+
+			// Change viewport settings to more VR-friendly sequencer settings
+			SavedEditorState.bCinematicPreviewViewport = VRViewportClient.AllowsCinematicPreview();
+			VRViewportClient.SetAllowCinematicPreview(false);
+			// Need to force fading and color scaling off in case we enter VR editing mode with a sequence open
+			VRViewportClient.bEnableFading = false;
+			VRViewportClient.bEnableColorScaling = false;
+			VRViewportClient.Invalidate(true);
 		}
 
 		VREditorLevelViewport->EnableStereoRendering( bActuallyUsingVR );
 		VREditorLevelViewport->SetRenderDirectlyToWindow( bActuallyUsingVR );
 
+
 		if( bActuallyUsingVR )
 		{
 			if( !( GEngine->HMDDevice->IsStereoEnabled() ) || bReenteringVREditing )
@@ -429,7 +438,6 @@ void UVREditorMode::Exit(const bool bHMDShouldExitStereo)
 			{
 				VREditorLevelViewport->EnableStereoRendering( false );
 				VREditorLevelViewport->SetRenderDirectlyToWindow( false );
-
 				{
 					FLevelEditorViewportClient& VRViewportClient = VREditorLevelViewport->GetLevelViewportClient();
 					FEditorViewportClient& VREditorViewportClient = VRViewportClient;
@@ -440,6 +448,10 @@ void UVREditorMode::Exit(const bool bHMDShouldExitStereo)
 					VRViewportClient.bAlwaysShowModeWidgetAfterSelectionChanges = SavedEditorState.bAlwaysShowModeWidgetAfterSelectionChanges;
 					VRViewportClient.EngineShowFlags = SavedEditorState.ShowFlags;
 					VRViewportClient.SetGameView( SavedEditorState.bGameView );
+					VRViewportClient.SetAllowCinematicPreview(SavedEditorState.bCinematicPreviewViewport);
+					VRViewportClient.bEnableFading = true;
+					VRViewportClient.bEnableColorScaling = true;
+					VRViewportClient.Invalidate(true);
 
 					if(bActuallyUsingVR)
 					{
