Reverting change to mouse pointer move events because of Content Browser impact

diff --git a/Engine/Source/Editor/VREditor/UI/VREditorUISystem.cpp b/Engine/Source/Editor/VREditor/UI/VREditorUISystem.cpp
--- a/Engine/Source/Editor/VREditor/UI/VREditorUISystem.cpp
+++ b/Engine/Source/Editor/VREditor/UI/VREditorUISystem.cpp
@@ -607,7 +607,7 @@ void UVREditorUISystem::OnVRHoverUpdate( UViewportInteractor* Interactor, FVecto
 								PressedButtons,
 								FModifierKeysState());
 
-							FSlateApplication::Get().RoutePointerMoveEvent(WidgetPathUnderFinger, PointerEvent, true);
+							FSlateApplication::Get().RoutePointerMoveEvent(WidgetPathUnderFinger, PointerEvent, false);
 							
 							bWasHandled = true;
 							HoverImpactPoint = HitResult.ImpactPoint;
@@ -730,7 +730,7 @@ void UVREditorUISystem::OnVRHoverUpdate( UViewportInteractor* Interactor, FVecto
 				PressedButtons,
 				FModifierKeysState() );
 
-			FSlateApplication::Get().RoutePointerMoveEvent( FWidgetPath(), PointerEvent, true );
+			FSlateApplication::Get().RoutePointerMoveEvent( FWidgetPath(), PointerEvent, false );
 		}
 	}
 }
