Removing unused VREditorWorldInteraction::IsInteractableComponent and changing to a cast-based check for editor UI.

diff --git a/Engine/Source/Editor/VREditor/UI/VREditorUISystem.cpp b/Engine/Source/Editor/VREditor/UI/VREditorUISystem.cpp
--- a/Engine/Source/Editor/VREditor/UI/VREditorUISystem.cpp
+++ b/Engine/Source/Editor/VREditor/UI/VREditorUISystem.cpp
@@ -384,8 +384,8 @@ void UVREditorUISystem::OnVRAction( FEditorViewportClient& ViewportClient, UView
 						if ( HitResult.Actor.IsValid() )
 						{
 							// Only allow clicks to our own widget components
-							UWidgetComponent* WidgetComponent = Cast<UWidgetComponent>( HitResult.GetComponent() );
-							if ( WidgetComponent != nullptr && IsWidgetAnEditorUIWidget( WidgetComponent ) )
+							UVREditorWidgetComponent* WidgetComponent = Cast<UVREditorWidgetComponent>( HitResult.GetComponent() );
+							if ( WidgetComponent != nullptr )
 							{
 								// Always mark the event as handled so that the editor doesn't try to select the widget component
 								bWasHandled = true;
@@ -575,8 +575,8 @@ void UVREditorUISystem::OnVRHoverUpdate( UViewportInteractor* Interactor, FVecto
 					}
 
 					// Only allow clicks to our own widget components
-					UWidgetComponent* WidgetComponent = Cast<UWidgetComponent>( HitResult.GetComponent() );
-					if( WidgetComponent != nullptr && IsWidgetAnEditorUIWidget( WidgetComponent ) )
+					UVREditorWidgetComponent* WidgetComponent = Cast<UVREditorWidgetComponent>( HitResult.GetComponent() );
+					if( WidgetComponent != nullptr )
 					{
 						FVector2D LastLocalHitLocation = WidgetComponent->GetLastLocalHitLocation();
 
@@ -1191,26 +1191,6 @@ void UVREditorUISystem::OnAssetEditorOpened(UObject* Asset)
 	}
 }
 
-bool UVREditorUISystem::IsWidgetAnEditorUIWidget( const UActorComponent* WidgetComponent ) const
-{
-	if( WidgetComponent != nullptr && WidgetComponent->IsA( UWidgetComponent::StaticClass() ) )
-	{
-		for( AVREditorFloatingUI* FloatingUIPtr : FloatingUIs )
-		{
-			if( FloatingUIPtr != nullptr )
-			{
-				if( WidgetComponent == FloatingUIPtr->GetWidgetComponent() )
-				{
-					return true;
-				}
-			}
-		}
-	}
-
-	return false;
-}
-
-
 bool UVREditorUISystem::IsShowingEditorUIPanel( const EEditorUIPanel EditorUIPanel ) const
 {
 	AVREditorFloatingUI* Panel = EditorUIPanels[ (int32)EditorUIPanel ];
diff --git a/Engine/Source/Editor/VREditor/UI/VREditorUISystem.h b/Engine/Source/Editor/VREditor/UI/VREditorUISystem.h
--- a/Engine/Source/Editor/VREditor/UI/VREditorUISystem.h
+++ b/Engine/Source/Editor/VREditor/UI/VREditorUISystem.h
@@ -144,9 +144,6 @@ public:
 	    up our spawned actors */
 	void CleanUpActorsBeforeMapChangeOrSimulate();
 
-	/** Searches to see if the specified widget component is from our editor UI */
-	bool IsWidgetAnEditorUIWidget( const class UActorComponent* WidgetComponent ) const;
-
 	/** Returns true if the specified editor UI panel is currently visible */
 	bool IsShowingEditorUIPanel( const EEditorUIPanel EditorUIPanel ) const;
 
diff --git a/Engine/Source/Editor/VREditor/VREditorMotionControllerInteractor.cpp b/Engine/Source/Editor/VREditor/VREditorMotionControllerInteractor.cpp
--- a/Engine/Source/Editor/VREditor/VREditorMotionControllerInteractor.cpp
+++ b/Engine/Source/Editor/VREditor/VREditorMotionControllerInteractor.cpp
@@ -637,6 +637,12 @@ void UVREditorMotionControllerInteractor::HandleInputKey( FEditorViewportClient&
 
 		if( Event == IE_Pressed )
 		{ 
+			// We may be activating sequencer scrub mode
+			if (DraggingMode == EViewportInteractionDraggingMode::Nothing && bIsTrackpadPositionValid[0] && bIsTrackpadPositionValid[1])
+			{
+				
+			}
+
 			// Start dragging at laser impact when already dragging actors freely
 			if( DraggingMode == EViewportInteractionDraggingMode::TransformablesFreely )
 			{
diff --git a/Engine/Source/Editor/VREditor/VREditorWorldInteraction.cpp b/Engine/Source/Editor/VREditor/VREditorWorldInteraction.cpp
--- a/Engine/Source/Editor/VREditor/VREditorWorldInteraction.cpp
+++ b/Engine/Source/Editor/VREditor/VREditorWorldInteraction.cpp
@@ -81,32 +81,6 @@ void UVREditorWorldInteraction::Shutdown()
 	Owner = nullptr;
 }
 
-//@todo vreditor: This function was overriding ViewportWorldInteraction
-bool UVREditorWorldInteraction::IsInteractableComponent( const UActorComponent* Component ) const
-{
-	bool bResult = false;
-	//bResult = UViewportWorldInteraction::IsInteractableComponent( Component );
-
-	// Don't bother if the base class already found out that this is a custom interactable
-	if ( bResult )
-	{
-		// Don't interact with frozen actors, our UI, or primitive components that have been set as not selectable
-		if ( Component != nullptr )
-		{
-			static const bool bIsVREditorDemo = FParse::Param( FCommandLine::Get(), TEXT( "VREditorDemo" ) );	// @todo vreditor: Remove this when no longer needed (console variable, too!)
-			const bool bIsFrozen = bIsVREditorDemo && Component->GetOwner()->GetActorLabel().StartsWith( TEXT( "Frozen_" ) );
-			const UPrimitiveComponent* ComponentAsPrimitive = Cast<UPrimitiveComponent>(Component);
-			// Don't allow user to move around our UI widgets 
-			//@todo VREditor: This wouldn't be necessary if the UI used the interactable interface
-			bResult = !bIsFrozen && 
-				!Owner->GetUISystem().IsWidgetAnEditorUIWidget( Component )
-				&& (ComponentAsPrimitive->bSelectable == true); 
-		}
-	}
-	
-	return bResult;
-}
-
 void UVREditorWorldInteraction::StopDragging( UViewportInteractor* Interactor )
 {
 	EViewportInteractionDraggingMode InteractorDraggingMode = Interactor->GetDraggingMode();
diff --git a/Engine/Source/Editor/VREditor/VREditorWorldInteraction.h b/Engine/Source/Editor/VREditor/VREditorWorldInteraction.h
--- a/Engine/Source/Editor/VREditor/VREditorWorldInteraction.h
+++ b/Engine/Source/Editor/VREditor/VREditorWorldInteraction.h
@@ -30,9 +30,6 @@ public:
 	/** Removes registered event */
 	void Shutdown();
 
-	/** If this component is interactable by an interactor */
-	bool IsInteractableComponent( const UActorComponent* Component ) const;
-
 	/** Snaps the selected objects to the ground */
 	void SnapSelectedActorsToGround();
 
