VREditor: Adding motion controller-controlled playback to Sequencer. Right = play, Left = reverse, Up = FF, Down = fast reverse. Additional cleanup of mesh trails and Sequencer instances when closing Sequencer panel, leaving VR editor.

diff --git a/Engine/Plugins/MovieScene/LevelSequenceEditor/Source/LevelSequenceEditor/Private/LevelSequenceEditorToolkit.cpp b/Engine/Plugins/MovieScene/LevelSequenceEditor/Source/LevelSequenceEditor/Private/LevelSequenceEditorToolkit.cpp
--- a/Engine/Plugins/MovieScene/LevelSequenceEditor/Source/LevelSequenceEditor/Private/LevelSequenceEditorToolkit.cpp
+++ b/Engine/Plugins/MovieScene/LevelSequenceEditor/Source/LevelSequenceEditor/Private/LevelSequenceEditorToolkit.cpp
@@ -289,7 +289,7 @@ void FLevelSequenceEditorToolkit::Initialize(const EToolkitMode::Type Mode, cons
 			Sequencer->SetAutoKeyMode(EAutoKeyMode::KeyAll);
 			Sequencer->SetKeyAllEnabled(true);
 			// Tell the VR Editor mode that Sequencer has refreshed
-			VRMode->RefreshVREditorSequencer();
+			VRMode->RefreshVREditorSequencer(Sequencer.Get());
 		}
 	}
 }
@@ -967,6 +967,13 @@ void FLevelSequenceEditorToolkit::HandleTrackMenuExtensionAddTrack(FMenuBuilder&
 
 bool FLevelSequenceEditorToolkit::OnRequestClose()
 {
+	const UWorld* World = CastChecked<UWorld>(GetLevelSequenceEditorPlaybackContext());
+	UVREditorMode* VRMode = Cast<UVREditorMode>(GEditor->GetEditorWorldExtensionsManager()->GetEditorWorldExtensions(World)->FindExtension(UVREditorMode::StaticClass()));
+	if (VRMode != nullptr)
+	{
+		// Null out the VR Mode's sequencer pointer
+		VRMode->RefreshVREditorSequencer(nullptr);
+	}
 	OpenToolkits.Remove(this);
 	OnClosedEvent.Broadcast();
 	return true;
diff --git a/Engine/Source/Editor/LevelEditor/Private/SLevelEditor.cpp b/Engine/Source/Editor/LevelEditor/Private/SLevelEditor.cpp
--- a/Engine/Source/Editor/LevelEditor/Private/SLevelEditor.cpp
+++ b/Engine/Source/Editor/LevelEditor/Private/SLevelEditor.cpp
@@ -1541,7 +1541,3 @@ TSharedRef<SWidget> SLevelEditor::CreateToolBox()
 	return NewToolBox;
 }
 
-TSharedPtr<SWidget> SLevelEditor::GetSequencerWidgetForLevelEditor()
-{
-	return SequencerWidgetPtr.Pin();
-}
\ No newline at end of file
diff --git a/Engine/Source/Editor/LevelEditor/Private/SLevelEditor.h b/Engine/Source/Editor/LevelEditor/Private/SLevelEditor.h
--- a/Engine/Source/Editor/LevelEditor/Private/SLevelEditor.h
+++ b/Engine/Source/Editor/LevelEditor/Private/SLevelEditor.h
@@ -76,7 +76,6 @@ public:
 	virtual TSharedPtr< class FAssetThumbnailPool > GetThumbnailPool() const override;
 	virtual void AppendCommands( const TSharedRef<FUICommandList>& InCommandsToAppend ) override;
 	virtual void AddStandaloneLevelViewport( const TSharedRef<SLevelViewport>& LevelViewport ) override;
-	virtual TSharedPtr<SWidget> GetSequencerWidgetForLevelEditor() override;
 
 	/**
 	 * Given a tab ID, summons a new tab in the position saved in the current layout, or in a default position.
diff --git a/Engine/Source/Editor/LevelEditor/Public/ILevelEditor.h b/Engine/Source/Editor/LevelEditor/Public/ILevelEditor.h
--- a/Engine/Source/Editor/LevelEditor/Public/ILevelEditor.h
+++ b/Engine/Source/Editor/LevelEditor/Public/ILevelEditor.h
@@ -63,9 +63,6 @@ public:
 	/** Spawns a level editor ToolBox widget (aka. "Modes") */
 	virtual TSharedRef<SWidget> CreateToolBox() = 0;
 
-	/** Return the level editor's Sequencer widget */
-	virtual TSharedPtr<SWidget> GetSequencerWidgetForLevelEditor() = 0;
-
 };
 
 
diff --git a/Engine/Source/Editor/Sequencer/Private/Sequencer.h b/Engine/Source/Editor/Sequencer/Private/Sequencer.h
--- a/Engine/Source/Editor/Sequencer/Private/Sequencer.h
+++ b/Engine/Source/Editor/Sequencer/Private/Sequencer.h
@@ -371,7 +371,6 @@ public:
 	void RecordSelectedActors();
 
 	/** Functions to push on to the transport controls we use */
-	FReply OnPlay(bool bTogglePlay=true, float InPlayRate=1.f);
 	FReply OnRecord();
 	FReply OnStepForward();
 	FReply OnStepBackward();
@@ -612,6 +611,8 @@ public:
 	virtual TSharedRef<INumericTypeInterface<float>> GetNumericTypeInterface() override;
 	virtual TSharedRef<INumericTypeInterface<float>> GetZeroPadNumericTypeInterface() override;
 	virtual TSharedRef<SWidget> MakeTransportControls(bool bExtended) override;
+	virtual FReply OnPlay(bool bTogglePlay = true, float InPlayRate = 1.f) override;
+	virtual void Pause() override;
 	virtual TSharedRef<SWidget> MakeTimeRange(const TSharedRef<SWidget>& InnerContent, bool bShowWorkingRange, bool bShowViewRange, bool bShowPlaybackRange) override;
 	virtual UObject* FindSpawnedObjectOrTemplate(const FGuid& BindingId) override;
 	virtual FGuid MakeNewSpawnable(UObject& SourceObject) override;
@@ -746,7 +747,6 @@ protected:
 	void Rewind();
 	void ShuttleForward();
 	void ShuttleBackward();
-	void Pause();
 	void StepForward();
 	void StepBackward();
 	void StepToNextKey();
diff --git a/Engine/Source/Editor/Sequencer/Private/SequencerEdMode.cpp b/Engine/Source/Editor/Sequencer/Private/SequencerEdMode.cpp
--- a/Engine/Source/Editor/Sequencer/Private/SequencerEdMode.cpp
+++ b/Engine/Source/Editor/Sequencer/Private/SequencerEdMode.cpp
@@ -43,6 +43,8 @@ void FSequencerEdMode::Enter()
 
 void FSequencerEdMode::Exit()
 {
+	CleanUpMeshTrails();
+
 	Sequencers.Reset();
 
 	FEdMode::Exit();
diff --git a/Engine/Source/Editor/Sequencer/Public/ISequencer.h b/Engine/Source/Editor/Sequencer/Public/ISequencer.h
--- a/Engine/Source/Editor/Sequencer/Public/ISequencer.h
+++ b/Engine/Source/Editor/Sequencer/Public/ISequencer.h
@@ -365,6 +365,12 @@ public:
 	/** @return Returns a widget containing the sequencer's playback controls */
 	virtual TSharedRef<SWidget> MakeTransportControls(bool bExtended) = 0;
 
+	/** Play or toggle playback at the specified play rate */
+	virtual FReply OnPlay(bool bTogglePlay = true, float InPlayRate = 1.f) = 0;
+
+	/** Pause playback */
+	virtual void Pause() = 0;
+
 	/** Attempt to find a spawned object in the currently focused movie scene, or the template object for the specified binding ID, if possible */
 	virtual UObject* FindSpawnedObjectOrTemplate(const FGuid& BindingId) = 0;
 
diff --git a/Engine/Source/Editor/VREditor/Public/VREditorMode.h b/Engine/Source/Editor/VREditor/Public/VREditorMode.h
--- a/Engine/Source/Editor/VREditor/Public/VREditorMode.h
+++ b/Engine/Source/Editor/VREditor/Public/VREditorMode.h
@@ -447,7 +447,10 @@ public:
 	float GetDefaultVRNearClipPlane() const;
 
 	/** Runtime and plugin modules can force VR Editor to refresh using this function */
-	void RefreshVREditorSequencer();
+	void RefreshVREditorSequencer(class ISequencer* InCurrentSequencer);
+
+	/** Returns the currently active sequencer */
+	class ISequencer* GetCurrentSequencer();
 
 private:
 
@@ -459,4 +462,7 @@ private:
 
 	/** If this current mode is running */
 	bool bIsActive;
+
+	/** Pointer to the current Sequencer */
+	class ISequencer* CurrentSequencer;
 };
diff --git a/Engine/Source/Editor/VREditor/UI/VREditorUISystem.cpp b/Engine/Source/Editor/VREditor/UI/VREditorUISystem.cpp
--- a/Engine/Source/Editor/VREditor/UI/VREditorUISystem.cpp
+++ b/Engine/Source/Editor/VREditor/UI/VREditorUISystem.cpp
@@ -88,6 +88,8 @@
 #include "GenericCommands.h"
 #include "SVRRadialPanel.h"
 
+#include "ISequencer.h"
+
 #define LOCTEXT_NAMESPACE "VREditor"
 
 namespace VREd
@@ -279,6 +281,16 @@ void UVREditorUISystem::OnVRAction( FEditorViewportClient& ViewportClient, UView
 			if ( Action.ActionType == VRActionTypes::ConfirmRadialSelection )
 			{
 				const EViewportInteractionDraggingMode DraggingMode = Interactor->GetDraggingMode();
+
+				// We may be activating sequencer scrub mode
+				if (DraggingMode == EViewportInteractionDraggingMode::Nothing && 
+					IsShowingEditorUIPanel(EEditorUIPanel::SequencerUI) && 
+					Action.Event == IE_Pressed && 
+					GetOwner().GetHMDDeviceType() == EHMDDeviceType::DT_SteamVR )
+				{
+					SequencerPlayback(Interactor, VREditorInteractor);
+				}
+
 				// Toggle the radial menu by pressing the touchpad or thumb stick, as long as we're not busy dragging something
 				if ( Action.Event == IE_Pressed && 
 					 DraggingMode != EViewportInteractionDraggingMode::TransformablesFreely &&
@@ -551,11 +563,71 @@ void UVREditorUISystem::OnVRAction( FEditorViewportClient& ViewportClient, UView
 }
 
 
-void UVREditorUISystem::OnVRHoverUpdate( UViewportInteractor* Interactor, FVector& HoverImpactPoint, bool& bWasHandled )
+void UVREditorUISystem::SequencerPlayback(UViewportInteractor* Interactor, UVREditorMotionControllerInteractor* VREditorInteractor)
+{
+	ISequencer* CurrentSequencer = GetOwner().GetCurrentSequencer();
+	if (Interactor->GetOtherInteractor() != nullptr && CurrentSequencer != nullptr)
+	{
+		if (!VREditorInteractor->IsTouchingTrackpad())
+		{
+			// Pause: For releasing touchpad on Vive
+			FVREditorActionCallbacks::SimulateKeyDown(EKeys::K, false);
+			FVREditorActionCallbacks::SimulateKeyUp(EKeys::K);
+
+		}
+		UVREditorMotionControllerInteractor* OtherMotionController = CastChecked<UVREditorMotionControllerInteractor>(Interactor->GetOtherInteractor());
+		if (!OtherMotionController->IsTouchingTrackpad())
+		{
+			// Basic Sequencer Playback
+			
+			if (VREditorInteractor->GetTrackpadPosition().Y > 0.5f && FMath::Abs(VREditorInteractor->GetTrackpadPosition().X) < 0.15)
+			{
+				// FF
+				
+				CurrentSequencer->OnPlay(false, 2.0f);
+				
+			}
+			else if (VREditorInteractor->GetTrackpadPosition().Y <= -0.5f && FMath::Abs(VREditorInteractor->GetTrackpadPosition().X) < 0.15)
+			{
+				// Fast Reverse
+				CurrentSequencer->OnPlay(false, -2.0f);
+			}
+			else if (VREditorInteractor->GetTrackpadPosition().X > 0.5f)
+			{
+				// Play
+				CurrentSequencer->OnPlay(false, 1.0f);
+			}
+			else if (VREditorInteractor->GetTrackpadPosition().X < -0.5f)
+			{
+				// Reverse
+				CurrentSequencer->OnPlay(false, -1.0f);
+			}
+			
+			else
+			{
+				// Pause: Center of touchpad, center of thumbstick
+				CurrentSequencer->Pause();
+			}
+		}
+		else
+		{
+			// Two-Handed Sequencer Playback
+		}
+	}
+}
+
+void UVREditorUISystem::OnVRHoverUpdate(UViewportInteractor* Interactor, FVector& HoverImpactPoint, bool& bWasHandled)
 {
 	UVREditorMotionControllerInteractor* VREditorInteractor = Cast<UVREditorMotionControllerInteractor>( Interactor );
 	if( VREditorInteractor != nullptr )
 	{
+		// We may be activating sequencer scrub mode
+		if (Interactor->GetDraggingMode() == EViewportInteractionDraggingMode::Nothing &&
+			IsShowingEditorUIPanel(EEditorUIPanel::SequencerUI) &&
+			GetOwner().GetHMDDeviceType() == EHMDDeviceType::DT_OculusRift)
+		{
+			SequencerPlayback(Interactor, VREditorInteractor);
+		}
 		if( !bWasHandled && Interactor->GetDraggingMode() != EViewportInteractionDraggingMode::Interactable )
 		{
 			FVector LaserPointerStart, LaserPointerEnd;
@@ -1126,10 +1198,8 @@ void UVREditorUISystem::CreateUIs()
 		}
 
 		{
-			const TSharedRef< ILevelEditor >& LevelEditor = FModuleManager::GetModuleChecked<FLevelEditorModule>("LevelEditor").GetFirstLevelEditor().ToSharedRef();
-
 			const FName TabIdentifier = NAME_None;	// No tab for us!
-			TSharedRef<SWidget> SequencerWidget = SNullWidget::NullWidget;
+			const TSharedRef<SWidget> SequencerWidget = SNullWidget::NullWidget;
 			
 			TSharedRef<SWidget> WidgetToDraw =
 				SNew(SDPIScaler)
@@ -1141,7 +1211,7 @@ void UVREditorUISystem::CreateUIs()
 
 			const bool bWithSceneComponent = false;
 			AVREditorFloatingUI* SequencerUI = UViewportWorldInteraction::SpawnTransientSceneActor< AVREditorDockableWindow >( GetOwner().GetWorld(), TEXT("SequencerUI"), bWithSceneComponent);
-			SequencerUI->SetSlateWidget(*this, WidgetToDraw, FIntPoint(VREd::ContentBrowserUIResolutionX->GetFloat(), VREd::ContentBrowserUIResolutionY->GetFloat()), VREd::ContentBrowserUISize->GetFloat(), AVREditorFloatingUI::EDockedTo::Nothing);
+			SequencerUI->SetSlateWidget(*this, WidgetToDraw, FIntPoint(VREd::SequencerUIResolutionX->GetFloat(), VREd::SequencerUIResolutionY->GetFloat()), VREd::ContentBrowserUISize->GetFloat(), AVREditorFloatingUI::EDockedTo::Nothing);
 			SequencerUI->ShowUI(false);
 			FloatingUIs.Add(SequencerUI);
 
@@ -1152,6 +1222,13 @@ void UVREditorUISystem::CreateUIs()
 
 void UVREditorUISystem::CleanUpActorsBeforeMapChangeOrSimulate()
 {
+	// If we have a sequence tab open, reset its widget and close the associated Sequencer
+	if (GetOwner().GetCurrentSequencer() != nullptr)
+	{
+		EditorUIPanels[(int32)EEditorUIPanel::SequencerUI]->SetSlateWidget(*this, SNullWidget::NullWidget, FIntPoint(VREd::SequencerUIResolutionX->GetFloat(), VREd::SequencerUIResolutionY->GetFloat()), VREd::EditorUISize->GetFloat(), AVREditorFloatingUI::EDockedTo::Nothing);
+		FVREditorActionCallbacks::CloseSequencer(GetOwner().GetCurrentSequencer()->GetRootMovieSceneSequence());
+	}
+
 	for( AVREditorFloatingUI* FloatingUIPtr : FloatingUIs )
 	{
 		if( FloatingUIPtr != nullptr )
@@ -1265,7 +1342,14 @@ void UVREditorUISystem::ShowEditorUIPanel( AVREditorFloatingUI* Panel, UVREditor
 			Panel->SetLocalRotation( FRotator( 90.0f, 180.0f, 0.0f ) ); // @todo vreditor: needs initial rotation
 		}
 
-		Panel->ShowUI( bShouldShow );
+		// If we are closing the sequencer panel, then also null out the sequencer widget and close the Sequencer instance
+		if (!bShouldShow && Panel == EditorUIPanels[(int32)EEditorUIPanel::SequencerUI] && GetOwner().GetCurrentSequencer() != nullptr)
+		{
+			Panel->SetSlateWidget(*this, SNullWidget::NullWidget, FIntPoint(VREd::SequencerUIResolutionX->GetFloat(), VREd::SequencerUIResolutionY->GetFloat()), VREd::EditorUISize->GetFloat(), AVREditorFloatingUI::EDockedTo::Nothing);
+			FVREditorActionCallbacks::CloseSequencer(GetOwner().GetCurrentSequencer()->GetRootMovieSceneSequence());
+		}
+
+		Panel->ShowUI(bShouldShow);
 
 		if ( Panel->GetDockedTo() == AVREditorFloatingUI::EDockedTo::LeftHand || Panel->GetDockedTo() == AVREditorFloatingUI::EDockedTo::RightHand )
 		{
@@ -2634,29 +2718,29 @@ void UVREditorUISystem::UpdateSequencerUI()
 	if (EditorUIPanels[(int32)EEditorUIPanel::SequencerUI] != nullptr)
 	{
 		AVREditorFloatingUI* SequencerPanel = EditorUIPanels[(int32)EEditorUIPanel::SequencerUI];
-		const TSharedRef< ILevelEditor >& LevelEditor = FModuleManager::GetModuleChecked<FLevelEditorModule>("LevelEditor").GetFirstLevelEditor().ToSharedRef();
-
 		const FName TabIdentifier = NAME_None;	// No tab for us!
-		TSharedRef<SWidget> SequencerWidget = LevelEditor->GetSequencerWidgetForLevelEditor().ToSharedRef();
+		if(GetOwner().GetCurrentSequencer() != nullptr)
+		{
+			const TSharedRef<SWidget> SequencerWidget = GetOwner().GetCurrentSequencer()->GetSequencerWidget();
 
-		TSharedRef<SWidget> WidgetToDraw =
-			SNew(SDPIScaler)
-			.DPIScale(1.0f)
-			[
-				SequencerWidget
-			]
-		;
+			TSharedRef<SWidget> WidgetToDraw =
+				SNew(SDPIScaler)
+				.DPIScale(1.0f)
+				[
+					SequencerWidget
+				]
+			;
 
-		const bool bWithSceneComponent = false;
-		SequencerPanel->SetSlateWidget(*this, WidgetToDraw, FIntPoint(VREd::SequencerUIResolutionX->GetFloat(), VREd::SequencerUIResolutionY->GetFloat()), VREd::EditorUISize->GetFloat(), AVREditorFloatingUI::EDockedTo::Nothing);
-		UVREditorInteractor* VREditorInteractor = VRMode->GetHandInteractor(EControllerHand::Left);
-		if (VREditorInteractor->IsHoveringOverUI())
-		{
-			VREditorInteractor = VRMode->GetHandInteractor(EControllerHand::Right);
+			const bool bWithSceneComponent = false;
+			SequencerPanel->SetSlateWidget(*this, WidgetToDraw, FIntPoint(VREd::SequencerUIResolutionX->GetFloat(), VREd::SequencerUIResolutionY->GetFloat()), VREd::EditorUISize->GetFloat(), AVREditorFloatingUI::EDockedTo::Nothing);
+			UVREditorInteractor* VREditorInteractor = VRMode->GetHandInteractor(EControllerHand::Left);
+			if (VREditorInteractor->IsHoveringOverUI())
+			{
+				VREditorInteractor = VRMode->GetHandInteractor(EControllerHand::Right);
 
+			}
+			ShowEditorUIPanel(UVREditorUISystem::EEditorUIPanel::SequencerUI, VREditorInteractor, true);
 		}
-		ShowEditorUIPanel(UVREditorUISystem::EEditorUIPanel::SequencerUI, VREditorInteractor, true);
-
 	}
 }
 
diff --git a/Engine/Source/Editor/VREditor/UI/VREditorUISystem.h b/Engine/Source/Editor/VREditor/UI/VREditorUISystem.h
--- a/Engine/Source/Editor/VREditor/UI/VREditorUISystem.h
+++ b/Engine/Source/Editor/VREditor/UI/VREditorUISystem.h
@@ -212,6 +212,8 @@ protected:
 	void OnVRAction( class FEditorViewportClient& ViewportClient, UViewportInteractor* Interactor,
 		const struct FViewportActionKeyInput& Action, bool& bOutIsInputCaptured, bool& bWasHandled );
 
+	void SequencerPlayback(UViewportInteractor* Interactor, class UVREditorMotionControllerInteractor* VREditorInteractor);
+
 	/** Called every frame to update hover state */
 	void OnVRHoverUpdate( UViewportInteractor* Interactor, FVector& HoverImpactPoint, bool& bWasHandled );
 
diff --git a/Engine/Source/Editor/VREditor/VREditorActions.cpp b/Engine/Source/Editor/VREditor/VREditorActions.cpp
--- a/Engine/Source/Editor/VREditor/VREditorActions.cpp
+++ b/Engine/Source/Editor/VREditor/VREditorActions.cpp
@@ -355,4 +355,13 @@ void FVREditorActionCallbacks::CreateNewSequence(class UVREditorMode* InVRMode)
 	FAssetEditorManager::Get().OpenEditorForAsset(NewAsset);
 }
 
+void FVREditorActionCallbacks::CloseSequencer(class UMovieSceneSequence* OpenSequence)
+{
+	IAssetEditorInstance* SequencerEditor = FAssetEditorManager::Get().FindEditorForAsset(OpenSequence, false);
+	if (SequencerEditor != nullptr)
+	{
+		SequencerEditor->CloseWindow();
+	}
+}
+
 #undef LOCTEXT_NAMESPACE
diff --git a/Engine/Source/Editor/VREditor/VREditorActions.h b/Engine/Source/Editor/VREditor/VREditorActions.h
--- a/Engine/Source/Editor/VREditor/VREditorActions.h
+++ b/Engine/Source/Editor/VREditor/VREditorActions.h
@@ -141,6 +141,9 @@ public:
 	/** Create a new level sequence with an auto-generated name */
 	static void CreateNewSequence(class UVREditorMode* InVRMode);
 
+	/** Close the sequencer instance */
+	static void CloseSequencer(class UMovieSceneSequence* OpenSequence);
+
 public:
 	
 	static FText GizmoCoordinateSystemText;
diff --git a/Engine/Source/Editor/VREditor/VREditorMode.cpp b/Engine/Source/Editor/VREditor/VREditorMode.cpp
--- a/Engine/Source/Editor/VREditor/VREditorMode.cpp
+++ b/Engine/Source/Editor/VREditor/VREditorMode.cpp
@@ -34,8 +34,6 @@
 #include "ViewportWorldInteractionManager.h"
 #include "EditorWorldExtension.h"
 
-#include "ISequencer.h"
-
 #define LOCTEXT_NAMESPACE "VREditorMode"
 
 namespace VREd
@@ -847,13 +845,19 @@ float UVREditorMode::GetDefaultVRNearClipPlane() const
 	return VREd::VRNearClipPlane->GetFloat();
 }
 
-void UVREditorMode::RefreshVREditorSequencer()
+void UVREditorMode::RefreshVREditorSequencer(class ISequencer* InCurrentSequencer)
 {
+	CurrentSequencer = InCurrentSequencer;
 	// Tell the VR Editor UI system to refresh the Sequencer UI
 	GetUISystem().UpdateSequencerUI();
 }
 
-bool UVREditorMode::IsHandAimingTowardsCapsule( UViewportInteractor* Interactor, const FTransform& CapsuleTransform, FVector CapsuleStart, const FVector CapsuleEnd, const float CapsuleRadius, const float MinDistanceToCapsule, const FVector CapsuleFrontDirection, const float MinDotForAimingAtCapsule ) const
+class ISequencer* UVREditorMode::GetCurrentSequencer()
+{
+	return CurrentSequencer;
+}
+
+bool UVREditorMode::IsHandAimingTowardsCapsule(UViewportInteractor* Interactor, const FTransform& CapsuleTransform, FVector CapsuleStart, const FVector CapsuleEnd, const float CapsuleRadius, const float MinDistanceToCapsule, const FVector CapsuleFrontDirection, const float MinDotForAimingAtCapsule) const
 {
 	bool bIsAimingTowards = false;
 	const float WorldScaleFactor = GetWorldScaleFactor();
diff --git a/Engine/Source/Editor/VREditor/VREditorMotionControllerInteractor.cpp b/Engine/Source/Editor/VREditor/VREditorMotionControllerInteractor.cpp
--- a/Engine/Source/Editor/VREditor/VREditorMotionControllerInteractor.cpp
+++ b/Engine/Source/Editor/VREditor/VREditorMotionControllerInteractor.cpp
@@ -637,12 +637,6 @@ void UVREditorMotionControllerInteractor::HandleInputKey( FEditorViewportClient&
 
 		if( Event == IE_Pressed )
 		{ 
-			// We may be activating sequencer scrub mode
-			if (DraggingMode == EViewportInteractionDraggingMode::Nothing && bIsTrackpadPositionValid[0] && bIsTrackpadPositionValid[1])
-			{
-				
-			}
-
 			// Start dragging at laser impact when already dragging actors freely
 			if( DraggingMode == EViewportInteractionDraggingMode::TransformablesFreely )
 			{
