Removing VREditorDemo flag

diff --git a/Engine/Source/Editor/ContentBrowser/Private/SAssetView.cpp b/Engine/Source/Editor/ContentBrowser/Private/SAssetView.cpp
--- a/Engine/Source/Editor/ContentBrowser/Private/SAssetView.cpp
+++ b/Engine/Source/Editor/ContentBrowser/Private/SAssetView.cpp
@@ -159,12 +159,6 @@ void SAssetView::Construct( const FArguments& InArgs )
 	TileViewThumbnailSize = 128;
 	TileViewThumbnailPadding = 5;
 
-	const bool bIsVREditorDemo = FParse::Param( FCommandLine::Get(), TEXT( "VREditorDemo" ) );	// @todo vreditor: Remove this when no longer needed
-	if( bIsVREditorDemo )
-	{
-		TileViewThumbnailPadding = 0;
-	}
-
 	TileViewNameHeight = 36;
 	ThumbnailScaleSliderValue = InArgs._ThumbnailScale; 
 
@@ -1051,12 +1045,6 @@ void SAssetView::CalculateFillScale( const FGeometry& AllottedGeometry )
 	{
 		float ItemWidth = GetTileViewItemBaseWidth();
 
- 		const bool bIsVREditorDemo = FParse::Param( FCommandLine::Get(), TEXT( "VREditorDemo" ) );	// @todo vreditor: Remove this when no longer needed
-		if( bIsVREditorDemo )
-		{
-			ItemWidth /= AllottedGeometry.Scale;
-		}
-
 		// Scrollbars are 16, but we add 1 to deal with half pixels.
 		const float ScrollbarWidth = 16 + 1;
 		float TotalWidth = AllottedGeometry.Size.X - ( ScrollbarWidth / AllottedGeometry.Scale );
diff --git a/Engine/Source/Editor/EditorStyle/Private/SlateEditorStyle.cpp b/Engine/Source/Editor/EditorStyle/Private/SlateEditorStyle.cpp
--- a/Engine/Source/Editor/EditorStyle/Private/SlateEditorStyle.cpp
+++ b/Engine/Source/Editor/EditorStyle/Private/SlateEditorStyle.cpp
@@ -5911,17 +5911,9 @@ void FSlateEditorStyle::FStyle::SetupContentBrowserStyle()
 		Set( "ContentBrowser.SourceTitleFont", TTF_CORE_FONT( "Fonts/Roboto-Regular", 12 ) );
 
 		// @todo vreditor urgent: Increase the size of Content Browser fonts while in VR
-		const bool bIsVREditorDemo = FParse::Param( FCommandLine::Get(), TEXT( "VREditorDemo" ) );	// @todo vreditor: Remove this when no longer needed
-		if( bIsVREditorDemo )
-		{
-			Set( "ContentBrowser.SourceListItemFont", TTF_CORE_FONT( "Fonts/Roboto-Regular", 20 ) );
-			Set( "ContentBrowser.SourceTreeItemFont", TTF_CORE_FONT( "Fonts/Roboto-Regular", 20 ) );
-		}
-		else
-		{
-			Set( "ContentBrowser.SourceListItemFont", TTF_CORE_FONT( "Fonts/Roboto-Regular", 10 ) );
-			Set( "ContentBrowser.SourceTreeItemFont", TTF_CORE_FONT( "Fonts/Roboto-Regular", 10 ) );
-		}
+		Set("ContentBrowser.SourceListItemFont", TTF_CORE_FONT("Fonts/Roboto-Regular", 10));
+		Set("ContentBrowser.SourceTreeItemFont", TTF_CORE_FONT("Fonts/Roboto-Regular", 10));
+
 		Set( "ContentBrowser.SourceTreeRootItemFont", TTF_CORE_FONT( "Fonts/Roboto-Regular", 12 ) );
 		Set( "ContentBrowser.AssetTreeFolderClosed", new IMAGE_BRUSH( "Icons/FolderClosed", FVector2D(18, 16) ) );
 		Set( "ContentBrowser.BreadcrumbPathPickerFolder", new IMAGE_BRUSH( "Icons/FolderClosed", FVector2D(18, 16), FLinearColor::Gray ) );
diff --git a/Engine/Source/Editor/VREditor/UI/VREditorUISystem.cpp b/Engine/Source/Editor/VREditor/UI/VREditorUISystem.cpp
--- a/Engine/Source/Editor/VREditor/UI/VREditorUISystem.cpp
+++ b/Engine/Source/Editor/VREditor/UI/VREditorUISystem.cpp
@@ -954,7 +954,6 @@ void UVREditorUISystem::CreateUIs()
 {
 	const FIntPoint DefaultResolution( VREd::DefaultEditorUIResolutionX->GetInt(), VREd::DefaultEditorUIResolutionY->GetInt() );
 
-	const bool bIsVREditorDemo = FParse::Param( FCommandLine::Get(), TEXT( "VREditorDemo" ) );	// @todo vreditor: Remove this when no longer needed
 	{
 		const FIntPoint Resolution(VREd::QuickMenuUIResolutionX->GetInt(), VREd::QuickMenuUIResolutionY->GetInt());
 		const bool bWithSceneComponent = false;
@@ -1004,23 +1003,11 @@ void UVREditorUISystem::CreateUIs()
 			Config.bExpandSourcesView = true;
 			Config.ThumbnailLabel = EThumbnailLabel::NoLabel;
 			Config.ThumbnailScale = 0.4f;
+			Config.bCanShowFilters = true;
+			Config.bUsePathPicker = true;
+			Config.bShowAssetPathTree = true;
+			Config.bAlwaysShowCollections = false;
 
-			if( bIsVREditorDemo )
-			{
-				Config.bUsePathPicker = false;
-				Config.bCanShowFilters = false;
-				Config.bShowAssetPathTree = false;
-				Config.bAlwaysShowCollections = true;
-				Config.SelectedCollectionName.Name = FName( "Meshes" );	// @todo vreditor: hard-coded collection name and type
-				Config.SelectedCollectionName.Type = ECollectionShareType::CST_Shared;
-			}
-			else
-			{
-				Config.bCanShowFilters = true;
-				Config.bUsePathPicker = true;
-				Config.bShowAssetPathTree = true;
-				Config.bAlwaysShowCollections = false;
-			}
 
 			Config.bShowBottomToolbar = false;
 			Config.bCanShowLockButton = false;
@@ -1118,7 +1105,6 @@ void UVREditorUISystem::CreateUIs()
 		// Create the tutorial dockable window
 		{
 			const bool bShowAtStart = 
-				!bIsVREditorDemo && 
 				GetOwner().IsActuallyUsingVR() &&	// Don't show tutorial at startup when in "forced VR" mode, because you can't interact with UI using the mouse currently
 				(GetDefault<UEditorExperimentalSettings>()->bAlwaysShowVRTutorial); // Only show tutorial at startup if the user has this enabled
 
diff --git a/Engine/Source/Editor/VREditor/VREditorMode.cpp b/Engine/Source/Editor/VREditor/VREditorMode.cpp
--- a/Engine/Source/Editor/VREditor/VREditorMode.cpp
+++ b/Engine/Source/Editor/VREditor/VREditorMode.cpp
@@ -319,13 +319,6 @@ void UVREditorMode::Enter(const bool bReenteringVREditing)
 			{
 				GEngine->HMDDevice->EnableStereo( true );
 			}
-			// @todo vreditor: Force single eye, undistorted mirror for demos
-			const bool bIsVREditorDemo = FParse::Param( FCommandLine::Get(), TEXT( "VREditorDemo" ) );	// @todo vreditor: Remove this when no longer needed (console variable, too!)
-			if( bIsVREditorDemo && GetHMDDeviceType() == EHMDDeviceType::DT_OculusRift )
-			{
-				// If we're using an Oculus Rift, go ahead and set the mirror mode to a single undistorted eye
-				GEngine->DeferredCommands.Add( FString::Printf( TEXT( "HMD MIRROR MODE %i" ), VREd::ForceOculusMirrorMode->GetInt() ) );
-			}
 		}
 
 		if( bActuallyUsingVR )
