VREditor: Fix warning for a declaration that shadows a field in VIBaseTransformGizmo. Added "In" for other parameters to prevent this in future cases.

diff --git a/Engine/Source/Editor/VREditor/Gizmo/VREditorTransformGizmo.cpp b/Engine/Source/Editor/VREditor/Gizmo/VREditorTransformGizmo.cpp
--- a/Engine/Source/Editor/VREditor/Gizmo/VREditorTransformGizmo.cpp
+++ b/Engine/Source/Editor/VREditor/Gizmo/VREditorTransformGizmo.cpp
@@ -139,19 +139,20 @@ ATransformGizmo::ATransformGizmo()
 }
 
 
-void ATransformGizmo::UpdateGizmo( const EGizmoHandleTypes GizmoType, const ECoordSystem GizmoCoordinateSpace, const FTransform& LocalToWorld, const FBox& LocalBounds, const FVector ViewLocation, const float ScaleMultiplier, bool bAllHandlesVisible,
-	const bool bAllowRotationAndScaleHandles, UActorComponent* DraggingHandle, const TArray< UActorComponent* >& HoveringOverHandles, const float GizmoHoverScale, const float GizmoHoverAnimationDuration )
+void ATransformGizmo::UpdateGizmo(const EGizmoHandleTypes InGizmoType, const ECoordSystem InGizmoCoordinateSpace, const FTransform& InLocalToWorld, const FBox& InLocalBounds, 
+	const FVector& InViewLocation, const float InScaleMultiplier, bool bInAllHandlesVisible, const bool bInAllowRotationAndScaleHandles, class UActorComponent* DraggingHandle, 
+	const TArray<UActorComponent*>& InHoveringOverHandles, const float InGizmoHoverScale, const float InGizmoHoverAnimationDuration)
 {
-	Super::UpdateGizmo( GizmoType, GizmoCoordinateSpace, LocalToWorld, LocalBounds, ViewLocation, ScaleMultiplier, bAllHandlesVisible,
-		bAllowRotationAndScaleHandles, DraggingHandle, HoveringOverHandles, GizmoHoverScale, GizmoHoverAnimationDuration );
+	Super::UpdateGizmo(InGizmoType, InGizmoCoordinateSpace, InLocalToWorld, InLocalBounds, InViewLocation, InScaleMultiplier, bInAllHandlesVisible,
+		bInAllowRotationAndScaleHandles, DraggingHandle, InHoveringOverHandles, InGizmoHoverScale, InGizmoHoverAnimationDuration);
 
 	const float WorldScaleFactor = GetWorld()->GetWorldSettings()->WorldToMeters / 100.0f;
 
 	// Position the gizmo at the location of the first selected actor
 	const bool bSweep = false;
-	this->SetActorTransform( LocalToWorld, bSweep );
+	this->SetActorTransform(InLocalToWorld, bSweep);
 
-	const FBox WorldSpaceBounds = LocalBounds.TransformBy( LocalToWorld );
+	const FBox WorldSpaceBounds = InLocalBounds.TransformBy(InLocalToWorld);
 	float ScaleCompensationForTinyGizmo = 1.0f;
 	{
 		const float AvgBounds = FMath::Lerp( WorldSpaceBounds.GetSize().GetAbsMin(), WorldSpaceBounds.GetSize().GetAbsMax(), 0.5f );
@@ -165,9 +166,9 @@ void ATransformGizmo::UpdateGizmo( const EGizmoHandleTypes GizmoType, const ECoo
 	// Increase scale with distance, to make gizmo handles easier to click on
 	// @todo vreditor: Should probably be a curve, not linear
 	// @todo vreditor: Should take FOV into account (especially in non-stereo/HMD mode)
-	const float WorldSpaceDistanceToGizmoBounds = FMath::Sqrt( WorldSpaceBounds.ComputeSquaredDistanceToPoint( ViewLocation ) );
+	const float WorldSpaceDistanceToGizmoBounds = FMath::Sqrt(WorldSpaceBounds.ComputeSquaredDistanceToPoint(InViewLocation));
 	const float GizmoScaleUpClose = VREd::GizmoScale->GetFloat();
-	const float GizmoScale( ScaleMultiplier * ( GizmoScaleUpClose + ( WorldSpaceDistanceToGizmoBounds / WorldScaleFactor ) * VREd::GizmoDistanceScaleFactor->GetFloat() ) * ScaleCompensationForTinyGizmo * WorldScaleFactor );
+	const float GizmoScale(InScaleMultiplier * (GizmoScaleUpClose + (WorldSpaceDistanceToGizmoBounds / WorldScaleFactor) * VREd::GizmoDistanceScaleFactor->GetFloat()) * ScaleCompensationForTinyGizmo * WorldScaleFactor);
 
 	// Update animation
 	float AnimationAlpha = GetAnimationAlpha();
@@ -179,8 +180,8 @@ void ATransformGizmo::UpdateGizmo( const EGizmoHandleTypes GizmoType, const ECoo
 		if ( HandleGroup != nullptr )
 		{
 			bool bIsHoveringOrDraggingThisHandleGroup = false;
-			HandleGroup->UpdateGizmoHandleGroup( LocalToWorld, LocalBounds, ViewLocation, bAllHandlesVisible, DraggingHandle, 
-				HoveringOverHandles, AnimationAlpha, GizmoScale, GizmoHoverScale, GizmoHoverAnimationDuration, /* Out */ bIsHoveringOrDraggingThisHandleGroup );
+			HandleGroup->UpdateGizmoHandleGroup(InLocalToWorld, InLocalBounds, InViewLocation, bInAllHandlesVisible, DraggingHandle, 
+				InHoveringOverHandles, AnimationAlpha, GizmoScale, InGizmoHoverScale, InGizmoHoverAnimationDuration, /* Out */ bIsHoveringOrDraggingThisHandleGroup);
 			
 			if( HandleGroup->GetHandleType() == EGizmoHandleTypes::Scale && bIsHoveringOrDraggingThisHandleGroup )
 			{
@@ -197,13 +198,13 @@ void ATransformGizmo::UpdateGizmo( const EGizmoHandleTypes GizmoType, const ECoo
 
 			// Update measurement text
 			{
-				const float LocalSpaceLengthOfBoundsAlongAxis = LocalBounds.GetSize()[ AxisIndex ];
+				const float LocalSpaceLengthOfBoundsAlongAxis = InLocalBounds.GetSize()[ AxisIndex ];
 
 				FVector LocalSpaceAxisLengthVector = FVector::ZeroVector;
 				LocalSpaceAxisLengthVector[ AxisIndex ] = LocalSpaceLengthOfBoundsAlongAxis;
 
 				// Transform to world space
-				const FVector WorldSpaceAxisLengthVector = LocalToWorld.TransformVector( LocalSpaceAxisLengthVector );
+				const FVector WorldSpaceAxisLengthVector = InLocalToWorld.TransformVector( LocalSpaceAxisLengthVector );
 				const float WorldSpaceLengthOfBoundsAlongAxis = WorldSpaceAxisLengthVector.Size();
 
 				FText BestSizeString;
@@ -234,12 +235,12 @@ void ATransformGizmo::UpdateGizmo( const EGizmoHandleTypes GizmoType, const ECoo
 					for( int32 EdgeIndex = 0; EdgeIndex < 4; ++EdgeIndex )
 					{
 						FVector GizmoSpaceEdge0, GizmoSpaceEdge1;
-						GetBoundingBoxEdge( LocalBounds, AxisIndex, EdgeIndex, /* Out */ GizmoSpaceEdge0, /* Out */ GizmoSpaceEdge1 );
+						GetBoundingBoxEdge( InLocalBounds, AxisIndex, EdgeIndex, /* Out */ GizmoSpaceEdge0, /* Out */ GizmoSpaceEdge1 );
 
-						const FVector WorldSpaceEdge0 = LocalToWorld.TransformPosition( GizmoSpaceEdge0 );
-						const FVector WorldSpaceEdge1 = LocalToWorld.TransformPosition( GizmoSpaceEdge1 );
+						const FVector WorldSpaceEdge0 = InLocalToWorld.TransformPosition( GizmoSpaceEdge0 );
+						const FVector WorldSpaceEdge1 = InLocalToWorld.TransformPosition( GizmoSpaceEdge1 );
 
-						const float SquaredDistance = FMath::PointDistToSegmentSquared( ViewLocation, WorldSpaceEdge0, WorldSpaceEdge1 );
+						const float SquaredDistance = FMath::PointDistToSegmentSquared( InViewLocation, WorldSpaceEdge0, WorldSpaceEdge1 );
 						if( SquaredDistance < ClosestSquaredDistance )
 						{
 							ClosestSquaredDistance = SquaredDistance;
@@ -255,7 +256,7 @@ void ATransformGizmo::UpdateGizmo( const EGizmoHandleTypes GizmoType, const ECoo
 				const FVector GizmoSpaceEdgeCenter = GizmoSpaceClosestEdge0 + EdgeVector * 0.5f;
 
 				const float UpOffsetAmount = GizmoScale * 5.0f;	// Push the text up a bit so that it doesn't overlap so much with our handles
-				const FVector GizmoSpaceUpVector = LocalToWorld.InverseTransformVectorNoScale( FVector::UpVector );
+				const FVector GizmoSpaceUpVector = InLocalToWorld.InverseTransformVectorNoScale( FVector::UpVector );
 
 				Measurement.MeasurementText->SetRelativeLocation( GizmoSpaceEdgeCenter + GizmoSpaceUpVector * UpOffsetAmount );
 
@@ -263,8 +264,8 @@ void ATransformGizmo::UpdateGizmo( const EGizmoHandleTypes GizmoType, const ECoo
 				FVector TextFacingDirection = UGizmoHandleGroup::GetAxisVector( TextFacingAxisIndex, ETransformGizmoHandleDirection::Positive );
 
 				// Make sure to face the camera
-				const FVector WorldSpaceEdgeCenterToViewDirection = ( LocalToWorld.TransformPosition( GizmoSpaceEdgeCenter ) - ViewLocation ).GetSafeNormal();
-				if( FVector::DotProduct( LocalToWorld.TransformVectorNoScale( TextFacingDirection ), WorldSpaceEdgeCenterToViewDirection ) > 0.0f )
+				const FVector WorldSpaceEdgeCenterToViewDirection = ( InLocalToWorld.TransformPosition( GizmoSpaceEdgeCenter ) - InViewLocation ).GetSafeNormal();
+				if( FVector::DotProduct( InLocalToWorld.TransformVectorNoScale( TextFacingDirection ), WorldSpaceEdgeCenterToViewDirection ) > 0.0f )
 				{
 					TextFacingDirection *= -1.0f;
 				}
@@ -290,4 +291,4 @@ void ATransformGizmo::UpdateGizmo( const EGizmoHandleTypes GizmoType, const ECoo
 			Measurement.MeasurementText->SetVisibility( bShouldShowMeasurements );
 		}
 	}
-}
+}
\ No newline at end of file
diff --git a/Engine/Source/Editor/VREditor/Gizmo/VREditorTransformGizmo.h b/Engine/Source/Editor/VREditor/Gizmo/VREditorTransformGizmo.h
--- a/Engine/Source/Editor/VREditor/Gizmo/VREditorTransformGizmo.h
+++ b/Engine/Source/Editor/VREditor/Gizmo/VREditorTransformGizmo.h
@@ -23,8 +23,9 @@ public:
 	/** Called by the world interaction system after we've been spawned into the world, to allow
 	    us to create components and set everything up nicely for the selected objects that we'll be
 		used to manipulate */
-	virtual void UpdateGizmo( const EGizmoHandleTypes GizmoType, const ECoordSystem GizmoCoordinateSpace, const FTransform& LocalToWorld, const FBox& LocalBounds, const FVector ViewLocation, const float ScaleMultiplier, bool bAllHandlesVisible,
-		const bool bAllowRotationAndScaleHandles, class UActorComponent* DraggingHandle, const TArray< UActorComponent* >& HoveringOverHandles, const float GizmoHoverScale, const float GizmoHoverAnimationDuration ) override;
+	virtual void UpdateGizmo(const EGizmoHandleTypes InGizmoType, const ECoordSystem InGizmoCoordinateSpace, const FTransform& InLocalToWorld, const FBox& InLocalBounds, 
+		const FVector& InViewLocation, const float InScaleMultiplier, bool bInAllHandlesVisible, const bool bInAllowRotationAndScaleHandles, class UActorComponent* DraggingHandle, 
+		const TArray<UActorComponent*>& InHoveringOverHandles, const float InGizmoHoverScale, const float InGizmoHoverAnimationDuration) override;
 
 private:
 
diff --git a/Engine/Source/Editor/ViewportInteraction/Gizmo/VIPivotTransformGizmo.cpp b/Engine/Source/Editor/ViewportInteraction/Gizmo/VIPivotTransformGizmo.cpp
--- a/Engine/Source/Editor/ViewportInteraction/Gizmo/VIPivotTransformGizmo.cpp
+++ b/Engine/Source/Editor/ViewportInteraction/Gizmo/VIPivotTransformGizmo.cpp
@@ -72,8 +72,7 @@ APivotTransformGizmo::APivotTransformGizmo() :
 	OnNewObjectsSelected();
 }
 
-void APivotTransformGizmo::UpdateGizmo( const EGizmoHandleTypes InGizmoType, const ECoordSystem GizmoCoordinateSpace, const FTransform& LocalToWorld, const FBox& LocalBounds, const FVector ViewLocation, const float ScaleMultiplier, bool bAllHandlesVisible, 
-	const bool bAllowRotationAndScaleHandles, UActorComponent* DraggingHandle, const TArray< UActorComponent* >& HoveringOverHandles, const float GizmoHoverScale, const float GizmoHoverAnimationDuration )
+void APivotTransformGizmo::UpdateGizmo(const EGizmoHandleTypes InGizmoType, const ECoordSystem InGizmoCoordinateSpace, const FTransform& InLocalToWorld, const FBox& InLocalBounds, const FVector& InViewLocation, const float InScaleMultiplier, bool bInAllHandlesVisible, const bool bInAllowRotationAndScaleHandles, class UActorComponent* DraggingHandle, const TArray<UActorComponent*>& InHoveringOverHandles, const float InGizmoHoverScale, const float InGizmoHoverAnimationDuration)
 {
 	GizmoType = InGizmoType;
 
@@ -81,34 +80,34 @@ void APivotTransformGizmo::UpdateGizmo( const EGizmoHandleTypes InGizmoType, con
 
 	// Position the gizmo at the location of the first selected actor
 	const bool bSweep = false;
-	this->SetActorTransform( LocalToWorld, bSweep );
+	this->SetActorTransform(InLocalToWorld, bSweep);
 
 
 	// Increase scale with distance, to make gizmo handles easier to click on
-	const float WorldSpaceDistanceToToPivot = FMath::Max( VREd::PivotGizmoMinDistanceForScaling->GetFloat(), FMath::Sqrt( FVector::DistSquared( GetActorLocation(), ViewLocation ) ) );
+	const float WorldSpaceDistanceToToPivot = FMath::Max(VREd::PivotGizmoMinDistanceForScaling->GetFloat(), FMath::Sqrt(FVector::DistSquared( GetActorLocation(), InViewLocation )));
 
 	// @todo gizmo: Scale causes bounds-based features to look wrong
-	const float GizmoScale( ScaleMultiplier * ( ( WorldSpaceDistanceToToPivot / WorldScaleFactor ) * VREd::PivotGizmoDistanceScaleFactor->GetFloat() ) * WorldScaleFactor );
+	const float GizmoScale(InScaleMultiplier * ((WorldSpaceDistanceToToPivot / WorldScaleFactor) * VREd::PivotGizmoDistanceScaleFactor->GetFloat() ) * WorldScaleFactor);
 
 	// Update animation
 	float AnimationAlpha = GetAnimationAlpha();
 
 	// Update all the handles
 	bool bIsHoveringOrDraggingScaleGizmo = false;
-	for ( UGizmoHandleGroup* HandleGroup : AllHandleGroups )
+	for (UGizmoHandleGroup* HandleGroup : AllHandleGroups)
 	{
-		if ( HandleGroup != nullptr)
+		if (HandleGroup != nullptr)
 		{
-			if ( DraggingHandle == nullptr || HandleGroup == StretchGizmoHandleGroup || HandleGroup == RotationGizmoHandleGroup)
+			if (DraggingHandle == nullptr || HandleGroup == StretchGizmoHandleGroup || HandleGroup == RotationGizmoHandleGroup)
 			{
 				bool bIsHoveringOrDraggingThisHandleGroup = false;
-				HandleGroup->UpdateGizmoHandleGroup( LocalToWorld, LocalBounds, ViewLocation, bAllHandlesVisible, DraggingHandle,
-					HoveringOverHandles, AnimationAlpha, GizmoScale, GizmoHoverScale, GizmoHoverAnimationDuration, /* Out */ bIsHoveringOrDraggingThisHandleGroup );
+				HandleGroup->UpdateGizmoHandleGroup(InLocalToWorld, InLocalBounds, InViewLocation, bInAllHandlesVisible, DraggingHandle,
+					InHoveringOverHandles, AnimationAlpha, GizmoScale, InGizmoHoverScale, InGizmoHoverAnimationDuration, /* Out */ bIsHoveringOrDraggingThisHandleGroup );
 			}
 			
 			if (HandleGroup != RotationGizmoHandleGroup)
 			{
-				HandleGroup->UpdateVisibilityAndCollision(GizmoType, GizmoCoordinateSpace, bAllHandlesVisible, bAllowRotationAndScaleHandles, DraggingHandle);
+				HandleGroup->UpdateVisibilityAndCollision(InGizmoType, InGizmoCoordinateSpace, bInAllHandlesVisible, bInAllowRotationAndScaleHandles, DraggingHandle);
 			}
 		}
 	}
diff --git a/Engine/Source/Editor/ViewportInteraction/Gizmo/VIPivotTransformGizmo.h b/Engine/Source/Editor/ViewportInteraction/Gizmo/VIPivotTransformGizmo.h
--- a/Engine/Source/Editor/ViewportInteraction/Gizmo/VIPivotTransformGizmo.h
+++ b/Engine/Source/Editor/ViewportInteraction/Gizmo/VIPivotTransformGizmo.h
@@ -24,8 +24,9 @@ public:
 	/** Called by the world interaction system after we've been spawned into the world, to allow
 	    us to create components and set everything up nicely for the selected objects that we'll be
 		used to manipulate */
-	virtual void UpdateGizmo( const EGizmoHandleTypes InGizmoType, const ECoordSystem GizmoCoordinateSpace, const FTransform& LocalToWorld, const FBox& LocalBounds, const FVector ViewLocation, const float ScaleMultiplier,
-		bool bAllHandlesVisible, const bool bAllowRotationAndScaleHandles, class UActorComponent* DraggingHandle, const TArray< UActorComponent* >& HoveringOverHandles, const float GizmoHoverScale, const float GizmoHoverAnimationDuration ) override;
+	virtual void UpdateGizmo(const EGizmoHandleTypes InGizmoType, const ECoordSystem InGizmoCoordinateSpace, const FTransform& InLocalToWorld, const FBox& InLocalBounds,
+		const FVector& InViewLocation, const float InScaleMultiplier, bool bInAllHandlesVisible, const bool bInAllowRotationAndScaleHandles, class UActorComponent* DraggingHandle,
+		const TArray<UActorComponent*>& InHoveringOverHandles, const float InGizmoHoverScale, const float InGizmoHoverAnimationDuration) override;
 
 private:
 	/** Uniform scale handle group component */
diff --git a/Engine/Source/Editor/ViewportInteraction/Public/VIBaseTransformGizmo.h b/Engine/Source/Editor/ViewportInteraction/Public/VIBaseTransformGizmo.h
--- a/Engine/Source/Editor/ViewportInteraction/Public/VIBaseTransformGizmo.h
+++ b/Engine/Source/Editor/ViewportInteraction/Public/VIBaseTransformGizmo.h
@@ -77,8 +77,8 @@ public:
 	/** Called by the world interaction system after we've been spawned into the world, to allow
 	    us to create components and set everything up nicely for the selected objects that we'll be
 		used to manipulate */
-	virtual void UpdateGizmo( const EGizmoHandleTypes GizmoType, const ECoordSystem GizmoCoordinateSpace, const FTransform& LocalToWorld, const FBox& LocalBounds, const FVector ViewLocation, const float ScaleMultiplier, bool bAllHandlesVisible,
-		const bool bAllowRotationAndScaleHandles, class UActorComponent* DraggingHandle, const TArray< UActorComponent* >& HoveringOverHandles, const float GizmoHoverScale, const float GizmoHoverAnimationDuration ) {}
+	virtual void UpdateGizmo(const EGizmoHandleTypes InGizmoType, const ECoordSystem InGizmoCoordinateSpace, const FTransform& InLocalToWorld, const FBox& InLocalBounds, const FVector& InViewLocation, const float InScaleMultiplier, bool bInAllHandlesVisible,
+		const bool bInAllowRotationAndScaleHandles, class UActorComponent* DraggingHandle, const TArray<UActorComponent*>& InHoveringOverHandles, const float InGizmoHoverScale, const float InGizmoHoverAnimationDuration) {}
 
 
 	/** Gets the current gizmo handle type */
