Merging //UE4/Private-Geometry to Dev-VREditor (//UE4/Dev-VREditor) 3339870

diff --git a/Engine/Source/Editor/UnrealEd/Private/Factories/ActorFactory.cpp b/Engine/Source/Editor/UnrealEd/Private/Factories/ActorFactory.cpp
--- a/Engine/Source/Editor/UnrealEd/Private/Factories/ActorFactory.cpp
+++ b/Engine/Source/Editor/UnrealEd/Private/Factories/ActorFactory.cpp
@@ -21,7 +21,6 @@ ActorFactory.cpp:
 #include "ActorFactories/ActorFactoryBoxVolume.h"
 #include "ActorFactories/ActorFactoryCameraActor.h"
 #include "ActorFactories/ActorFactoryCharacter.h"
-#include "ActorFactories/ActorFactorySubDSurface.h"
 #include "ActorFactories/ActorFactoryClass.h"
 #include "ActorFactories/ActorFactoryCylinderVolume.h"
 #include "ActorFactories/ActorFactoryDeferredDecal.h"
@@ -75,7 +74,6 @@ ActorFactory.cpp:
 #include "Sound/SoundBase.h"
 #include "Sound/AmbientSound.h"
 #include "GameFramework/Volume.h"
-#include "Engine/SubDSurface.h"
 #include "Engine/DecalActor.h"
 #include "Atmosphere/AtmosphericFog.h"
 #include "Engine/ExponentialHeightFog.h"
@@ -113,8 +111,6 @@ ActorFactory.cpp:
 #include "Engine/TriggerSphere.h"
 #include "Engine/TriggerCapsule.h"
 #include "Engine/TextRenderActor.h"
-#include "Engine/SubDSurfaceActor.h"
-#include "Components/SubDSurfaceComponent.h"
 
 #include "Engine/DestructibleMesh.h"
 #include "Components/SkeletalMeshComponent.h"
@@ -588,48 +584,6 @@ UActorFactoryTextRender::UActorFactoryTextRender(const FObjectInitializer& Objec
 	bUseSurfaceOrientation = true;
 }
 
-/*-----------------------------------------------------------------------------
-UActorFactorySubDSurface
------------------------------------------------------------------------------*/
-UActorFactorySubDSurface::UActorFactorySubDSurface(const FObjectInitializer& ObjectInitializer)
-	: Super(ObjectInitializer)
-{
-	// Property initialization
-	DisplayName = LOCTEXT("SubDSurfaceDisplayName", "Subdivision Surface");
-	NewActorClass = ASubDSurfaceActor::GetPrivateStaticClass(L"...");
-	bUseSurfaceOrientation = true;
-}
-
-bool UActorFactorySubDSurface::CanCreateActorFrom( const FAssetData& AssetData, FText& OutErrorMsg )
-{
-	// We allow creating ASubDSurfaceActor without an existing asset
-	if ( UActorFactory::CanCreateActorFrom( AssetData, OutErrorMsg ) )
-	{
-		return true;
-	}
-
-	if ( !AssetData.IsValid() || !AssetData.GetClass()->IsChildOf( USubDSurface::StaticClass() ) )
-	{
-		OutErrorMsg = NSLOCTEXT("CanCreateActor", "NoSubDSurface", "A valid SubDSurface must be specified.");
-		return false;
-	}
-
-	return true;
-}
-
-void UActorFactorySubDSurface::PostSpawnActor( UObject* Asset, AActor* NewActor )
-{
-	Super::PostSpawnActor(Asset, NewActor);
-
-	USubDSurface* SubDSurface = Cast<USubDSurface>(Asset);
-	ASubDSurfaceActor* SubDSurfaceActor = CastChecked<ASubDSurfaceActor>(NewActor);
-
-	if(SubDSurface && SubDSurfaceActor)
-	{
-		SubDSurfaceActor->SubDSurface->SetMesh(SubDSurface);
-	}
-}
-
 /*-----------------------------------------------------------------------------
 UActorFactoryEmitter
 -----------------------------------------------------------------------------*/
diff --git a/Engine/Source/Editor/UnrealEd/Private/Factories/EditorFactories.cpp b/Engine/Source/Editor/UnrealEd/Private/Factories/EditorFactories.cpp
--- a/Engine/Source/Editor/UnrealEd/Private/Factories/EditorFactories.cpp
+++ b/Engine/Source/Editor/UnrealEd/Private/Factories/EditorFactories.cpp
@@ -106,7 +106,6 @@
 #include "Factories/ReimportSoundSurroundFactory.h"
 #include "Factories/StructureFactory.h"
 #include "Factories/SubsurfaceProfileFactory.h"
-#include "Factories/SubDSurfaceFactory.h"
 #include "Factories/Texture2dFactoryNew.h"
 #include "Engine/Texture.h"
 #include "Factories/TextureFactory.h"
@@ -136,7 +135,6 @@
 #include "Sound/SoundWave.h"
 #include "GameFramework/DefaultPhysicsVolume.h"
 #include "Engine/SubsurfaceProfile.h"
-#include "Engine/SubDSurface.h"
 #include "Misc/ConfigCacheIni.h"
 #include "Misc/FeedbackContext.h"
 #include "GameFramework/WorldSettings.h"
@@ -6596,25 +6594,6 @@ UObject* USubsurfaceProfileFactory::FactoryCreateNew(UClass* InClass, UObject* I
 	return NewObject<USubsurfaceProfile>(InParent, InName, Flags);
 }
 
-
-/*-----------------------------------------------------------------------------
-USubDSurfaceFactory implementation.
------------------------------------------------------------------------------*/
-USubDSurfaceFactory::USubDSurfaceFactory(const FObjectInitializer& ObjectInitializer)
-	: Super(ObjectInitializer)
-{
-
-	SupportedClass = USubDSurface::StaticClass();
-	bCreateNew = true;
-	bEditorImport = false;
-	bEditAfterNew = true;
-}
-
-UObject* USubDSurfaceFactory::FactoryCreateNew(UClass* InClass, UObject* InParent, FName InName, EObjectFlags Flags, UObject* Context, FFeedbackContext* Warn)
-{
-	return NewObject<USubDSurface>(InParent, USubDSurface::StaticClass(), InName, Flags);
-}
-
 /*-----------------------------------------------------------------------------
 	UTouchInterfaceFactory implementation.
 -----------------------------------------------------------------------------*/
diff --git a/Engine/Source/Editor/UnrealEd/Private/Fbx/FbxFactory.cpp b/Engine/Source/Editor/UnrealEd/Private/Fbx/FbxFactory.cpp
--- a/Engine/Source/Editor/UnrealEd/Private/Fbx/FbxFactory.cpp
+++ b/Engine/Source/Editor/UnrealEd/Private/Fbx/FbxFactory.cpp
@@ -12,7 +12,6 @@
 #include "Factories/FbxTextureImportData.h"
 #include "Factories/FbxImportUI.h"
 #include "Engine/StaticMesh.h"
-#include "Engine/SubDSurface.h"
 #include "Editor.h"
 
 #include "Logging/TokenizedMessage.h"
@@ -53,7 +52,7 @@ void UFbxFactory::PostInitProperties()
 
 bool UFbxFactory::DoesSupportClass(UClass * Class)
 {
-	return (Class == UStaticMesh::StaticClass() || Class == USkeletalMesh::StaticClass() || Class == UAnimSequence::StaticClass() || Class == USubDSurface::StaticClass());
+	return (Class == UStaticMesh::StaticClass() || Class == USkeletalMesh::StaticClass() || Class == UAnimSequence::StaticClass());
 }
 
 UClass* UFbxFactory::ResolveSupportedClass()
@@ -68,10 +67,6 @@ UClass* UFbxFactory::ResolveSupportedClass()
 	{
 		ImportClass = UAnimSequence::StaticClass();
 	}
-	else if (ImportUI->MeshTypeToImport == FBXIT_SubDSurface)
-	{
-		ImportClass = USubDSurface::StaticClass();
-	}
 	else
 	{
 		ImportClass = UStaticMesh::StaticClass();
@@ -528,31 +523,7 @@ UObject* UFbxFactory::FactoryCreateBinary
 							FFbxErrors::SkeletalMesh_NoMeshFoundOnRoot);
 					}
 				}
-				else if ( ImportUI->MeshTypeToImport == FBXIT_SubDSurface ) // SubDSurface
-				{
-					USubDSurface* NewMesh = NULL;
-					TArray<FbxNode*> FbxMeshArray;
-					FbxImporter->FillFbxMeshArray(RootNodeToImport, FbxMeshArray, FbxImporter);
-					if (FbxMeshArray.Num() > 0)
-					{
-						NewMesh = Cast<USubDSurface>( CreateOrOverwriteAsset(USubDSurface::StaticClass(), InParent, Name, Flags) );
-						
-						bool bOk = FbxImporter->ImportSubDSurface(NewMesh, InParent, FbxMeshArray, Name, Flags, ImportUI->StaticMeshImportData);
-
-						if(!bOk)
-						{
-							FbxImporter->AddTokenizedErrorMessage(FTokenizedMessage::Create(EMessageSeverity::Error, LOCTEXT("FailedToImport_ImportSubDSurface", "Could not import subdivision surface mesh (no quad mesh?).")), 
-								FFbxErrors::SkeletalMesh_ImportSubDSurface);
-
-							ObjectTools::DeleteSingleObject(NewMesh);
-						}
-					}
-
-					ImportedMeshCount = NewMesh ? 1 : 0;
-
-					NewObject = NewMesh;
-				}
-					else if ( ImportUI->MeshTypeToImport == FBXIT_Animation )// animation
+				else if ( ImportUI->MeshTypeToImport == FBXIT_Animation )// animation
 				{
 					if (ImportOptions->SkeletonForAnimation)
 					{
diff --git a/Engine/Source/Editor/UnrealEd/Private/Fbx/FbxMainImport.cpp b/Engine/Source/Editor/UnrealEd/Private/Fbx/FbxMainImport.cpp
--- a/Engine/Source/Editor/UnrealEd/Private/Fbx/FbxMainImport.cpp
+++ b/Engine/Source/Editor/UnrealEd/Private/Fbx/FbxMainImport.cpp
@@ -72,7 +72,6 @@ FBXImportOptions* GetImportOptions( UnFbx::FFbxImporter* FbxImporter, UFbxImport
 		}
 
 		ImportUI->bImportAsSkeletal = ImportUI->MeshTypeToImport == FBXIT_SkeletalMesh;
-		ImportUI->bImportAsSubDSurface = ImportUI->MeshTypeToImport == FBXIT_SubDSurface;
 		ImportUI->bImportMesh = ImportUI->MeshTypeToImport != FBXIT_Animation;
 		ImportUI->bIsObjImport = bIsObjFormat;
 
diff --git a/Engine/Source/Editor/UnrealEd/Private/Fbx/FbxStaticMeshImport.cpp b/Engine/Source/Editor/UnrealEd/Private/Fbx/FbxStaticMeshImport.cpp
--- a/Engine/Source/Editor/UnrealEd/Private/Fbx/FbxStaticMeshImport.cpp
+++ b/Engine/Source/Editor/UnrealEd/Private/Fbx/FbxStaticMeshImport.cpp
@@ -26,7 +26,6 @@
 #include "ObjectTools.h"
 #include "PackageTools.h"
 #include "Logging/TokenizedMessage.h"
-#include "Engine/SubDSurface.h"
 #include "FbxImporter.h"
 #include "GeomFitUtils.h"
 #include "Interfaces/ITargetPlatform.h"
@@ -1660,166 +1659,6 @@ void UnFbx::FFbxImporter::ReorderMaterialToFbxOrder(UStaticMesh* StaticMesh, TAr
 	}
 }
 
-bool UnFbx::FFbxImporter::ImportSubDSurface(USubDSurface* Out, UObject* InParent, TArray<FbxNode*>& MeshNodeArray, const FName InName, EObjectFlags Flags, UFbxStaticMeshImportData* TemplateImportData)
-{
-	if(!MeshNodeArray.Num())
-	{
-		return false;
-	}
-
-	// Make sure rendering is done - so we are not changing data being used
-	FlushRenderingCommands();
-
-	double StartTime = FPlatformTime::Seconds();
-
-	Parent = InParent;
-
-	// warning for missing smoothing group info
-	CheckSmoothingInfo(MeshNodeArray[0]->GetMesh());
-
-	uint32 ModelCount = (uint32)MeshNodeArray.Num();
-
-	// RemoveBadPolygons(), count Vertex and Indices
-	uint32 TotalIndexCount = 0, TotalVertexCount = 0, TotalPolyCount = 0;
-	{
-		for(uint32 ModelIndex = 0; ModelIndex < ModelCount; ++ModelIndex)
-		{
-			FbxNode* Node = MeshNodeArray[ModelIndex];
-			FbxMesh* FbxMesh = Node->GetMesh();
-
-			int32 Removed = FbxMesh->RemoveBadPolygons();
-
-			TotalVertexCount += FbxMesh->GetControlPointsCount();
-
-			int32 PolyCount = FbxMesh->GetPolygonCount();
-
-			TotalPolyCount += PolyCount;
-
-			for(int32 PolyIndex = 0; PolyIndex < PolyCount; ++PolyIndex)
-			{
-				int32 PolySize = FbxMesh->GetPolygonSize(PolyIndex);
-				
-				TotalIndexCount += PolySize;
-			}
-		}
-	}
-
-	Out->IndicesPerFace.AddUninitialized(TotalIndexCount);
-	Out->VertexCountPerFace.AddUninitialized(TotalPolyCount);
-
-	uint32* IndicesPerFace = Out->IndicesPerFace.GetData();
-	uint32* VertexCountPerFace = Out->VertexCountPerFace.GetData();
-
-	// add position vertex attribute
-	FVector* Positions = 0;
-	{
-		auto Stream = Out->CreateVertexAttributeStream(FName(TEXT("Position")));
-		Positions = Stream->CreateFVectorUninitialized(TotalVertexCount);
-	}
-
-	// todo: support multiple UV
-
-	// add UV vertex attribute
-	FVector2D* UV0 = 0;
-	{
-		auto Stream = Out->CreateVertexAttributeStream(FName(TEXT("UV0")));
-		UV0 = Stream->CreateFVector2DUninitialized(TotalIndexCount);
-	}
-
-	uint32 CurrentIndexStart = 0;
-	for(uint32 ModelIndex = 0; ModelIndex < ModelCount; ++ModelIndex)
-	{
-		FbxNode* Node = MeshNodeArray[ModelIndex];
-		FbxMesh* FbxMesh = Node->GetMesh();
-		uint32 PolyCount = (uint32)FbxMesh->GetPolygonCount();
-
-		FFBXUVs FBXUVs(FbxMesh);
-
-		FBXUVs.Phase2(FbxMesh);
-
-		for(uint32 PolyIndex = 0; PolyIndex < PolyCount; ++PolyIndex)
-		{
-			uint32 CornerCount = FbxMesh->GetPolygonSize(PolyIndex);
-
-			// store VertexCountPerFace
-			*VertexCountPerFace++ = CornerCount;
-
-			for(uint32 Corner = 0; Corner < CornerCount; ++Corner)
-			{
-				uint32 VertexIndex = (uint32)FbxMesh->GetPolygonVertex(PolyIndex, Corner);
-
-				check(VertexIndex < TotalVertexCount);
-
-				// store IndicesPerFace
-				*IndicesPerFace++ = CurrentIndexStart + VertexIndex;
-			}
-		}
-
-		int32 ControlPointsCount = FbxMesh->GetControlPointsCount();
-
-		CurrentIndexStart += ControlPointsCount;
-
-		FbxAMatrix Matrix = ComputeTotalMatrix(Node);
-
-		{
-			
-			FbxVector4* ControlPoints = FbxMesh->GetControlPoints();
-
-			for(int32 ControlPointsIndex = 0; ControlPointsIndex < ControlPointsCount; ControlPointsIndex++ )
-			{
-				FVector Pos = Converter.ConvertPos(Matrix.MultT(ControlPoints[ControlPointsIndex]));
-
-				*Positions++ = Pos;
-			}
-		}
-
-
-
-		//
-		// uvs
-		//
-		// In FBX file, the same UV may be saved multiple times, i.e., there may be same UV in LayerElementUV
-		// So we don't import the duplicate UVs
-		{
-			// We don't support multiple UV yet
-			ensure(FBXUVs.UniqueUVCount == 1);
-
-			int32 UVLayerIndex;
-			for (UVLayerIndex = 0; UVLayerIndex < FBXUVs.UniqueUVCount; UVLayerIndex++)
-			{
-				if (FBXUVs.LayerElementUV[UVLayerIndex] != NULL) 
-				{
-					uint32 UniqueCornerIndex = 0;
-
-					for(uint32 PolyIndex = 0; PolyIndex < PolyCount; ++PolyIndex)
-					{
-						int32 PolySize = FbxMesh->GetPolygonSize(PolyIndex);
-
-						for (int32 CornerIndex = 0; CornerIndex < PolySize; CornerIndex++)
-						{
-							// If there are odd number negative scale, invert the vertex order for triangles
-	//todo						int32 WedgeIndex = WedgeOffset + TriangleIndex * 3 + (OddNegativeScale ? 2 - CornerIndex : CornerIndex);
-		
-							int lControlPointIndex = FbxMesh->GetPolygonVertex(PolyIndex, CornerIndex);
-
-							int UVIndex = FBXUVs.ComputeUVIndex(UVLayerIndex, lControlPointIndex, UniqueCornerIndex);
-	
-							FbxVector2 UVVector = FBXUVs.LayerElementUV[UVLayerIndex]->GetDirectArray().GetAt(UVIndex);
-
-							*UV0++ = FVector2D(static_cast<float>(UVVector[0]), 1.0f - static_cast<float>(UVVector[1]));   //flip the Y of UVs for DirectX
-
-							++UniqueCornerIndex;
-						}
-					}
-				}
-			}
-		}
-	}
-	
-	UE_LOG(LogFbx, Log, TEXT("SubDSurface import: %5.2f seconds"), FPlatformTime::Seconds() - StartTime);
-
-	return true;
-}
 struct FbxSocketNode
 {
 	FName SocketName;
diff --git a/Engine/Source/Editor/UnrealEd/Private/FbxImporter.h b/Engine/Source/Editor/UnrealEd/Private/FbxImporter.h
--- a/Engine/Source/Editor/UnrealEd/Private/FbxImporter.h
+++ b/Engine/Source/Editor/UnrealEd/Private/FbxImporter.h
@@ -32,7 +32,6 @@ class UPhysicsAsset;
 class USkeletalMesh;
 class USkeleton;
 class UStaticMesh;
-class USubDSurface;
 class UTexture;
 struct FExpressionInput;
 struct FRawMesh;
@@ -613,14 +612,6 @@ public:
 	 */
 	void ReorderMaterialToFbxOrder(UStaticMesh* StaticMesh, TArray<FbxNode*>& MeshNodeArray);
 
-	/**
-	* Creates a SubDSurface mesh from all the meshes in FBX scene with the given name and flags.
-	*
-	* @param MeshNodeArray	Fbx Nodes to import
-	* @param InName	the Unreal Mesh name after import
-	*/
-	UNREALED_API bool ImportSubDSurface(USubDSurface* Out, UObject* InParent, TArray<FbxNode*>& MeshNodeArray, const FName InName, EObjectFlags Flags, UFbxStaticMeshImportData* TemplateImportData);
-
 	void ImportStaticMeshGlobalSockets( UStaticMesh* StaticMesh );
 	void ImportStaticMeshLocalSockets( UStaticMesh* StaticMesh, TArray<FbxNode*>& MeshNodeArray);
 
diff --git a/Engine/Source/Editor/UnrealEd/Private/Kismet2/DebuggerCommands.cpp b/Engine/Source/Editor/UnrealEd/Private/Kismet2/DebuggerCommands.cpp
--- a/Engine/Source/Editor/UnrealEd/Private/Kismet2/DebuggerCommands.cpp
+++ b/Engine/Source/Editor/UnrealEd/Private/Kismet2/DebuggerCommands.cpp
@@ -43,7 +43,6 @@
 
 #include "IHeadMountedDisplay.h"
 #include "Editor.h"
-#include "IVREditorModule.h"
 
 //@TODO: Remove this dependency
 #include "EngineGlobals.h"
@@ -144,10 +143,8 @@ public:
 
 	static FSlateIcon GetResumePlaySessionImage();
 	static FText GetResumePlaySessionToolTip();
-	static void ResumePlaySession_Clicked();
 	static void StopPlaySession_Clicked();
 	static void LateJoinSession_Clicked();
-	static void PausePlaySession_Clicked();
 	static void SingleFrameAdvance_Clicked();
 
 	static void ShowCurrentStatement_Clicked();
@@ -166,9 +163,6 @@ public:
 	static FText GetPossessEjectTooltip();
 	static FSlateIcon GetPossessEjectImage();
 
-	static bool HasPlayWorld();
-	static bool HasPlayWorldAndPaused();
-	static bool HasPlayWorldAndRunning();
 	static bool CanLateJoin();
 	static bool CanShowLateJoinButton();
 
@@ -457,24 +451,24 @@ void FPlayWorldCommands::BindGlobalPlayWorldCommands()
 
 	// Play, Pause, Toggle between play and pause
 	ActionList.MapAction( Commands.ResumePlaySession,
-		FExecuteAction::CreateStatic( &FInternalPlayWorldCommandCallbacks::ResumePlaySession_Clicked ),
-		FCanExecuteAction::CreateStatic( &FInternalPlayWorldCommandCallbacks::HasPlayWorldAndPaused ),
+		FExecuteAction::CreateStatic( &FPlayWorldCommandCallbacks::ResumePlaySession_Clicked ),
+		FCanExecuteAction::CreateStatic( &FPlayWorldCommandCallbacks::HasPlayWorldAndPaused ),
 		FIsActionChecked(),
-		FIsActionButtonVisible::CreateStatic( &FInternalPlayWorldCommandCallbacks::HasPlayWorldAndPaused )
+		FIsActionButtonVisible::CreateStatic( &FPlayWorldCommandCallbacks::HasPlayWorldAndPaused )
 		);
 
 	ActionList.MapAction( Commands.PausePlaySession,
-		FExecuteAction::CreateStatic( &FInternalPlayWorldCommandCallbacks::PausePlaySession_Clicked ),
-		FCanExecuteAction::CreateStatic( &FInternalPlayWorldCommandCallbacks::HasPlayWorldAndRunning ),
+		FExecuteAction::CreateStatic( &FPlayWorldCommandCallbacks::PausePlaySession_Clicked ),
+		FCanExecuteAction::CreateStatic( &FPlayWorldCommandCallbacks::HasPlayWorldAndRunning ),
 		FIsActionChecked(),
-		FIsActionButtonVisible::CreateStatic( &FInternalPlayWorldCommandCallbacks::HasPlayWorldAndRunning )
+		FIsActionButtonVisible::CreateStatic( &FPlayWorldCommandCallbacks::HasPlayWorldAndRunning )
 		);
 
 	ActionList.MapAction( Commands.SingleFrameAdvance,
 		FExecuteAction::CreateStatic( &FInternalPlayWorldCommandCallbacks::SingleFrameAdvance_Clicked ),
-		FCanExecuteAction::CreateStatic( &FInternalPlayWorldCommandCallbacks::HasPlayWorldAndPaused ),
+		FCanExecuteAction::CreateStatic( &FPlayWorldCommandCallbacks::HasPlayWorldAndPaused ),
 		FIsActionChecked(),
-		FIsActionChecked::CreateStatic( &FInternalPlayWorldCommandCallbacks::HasPlayWorldAndPaused )
+		FIsActionChecked::CreateStatic( &FPlayWorldCommandCallbacks::HasPlayWorldAndPaused )
 		);
 
 	ActionList.MapAction( Commands.TogglePlayPauseOfPlaySession,
@@ -1015,6 +1009,39 @@ void FPlayWorldCommandCallbacks::StartPlayFromHere()
 }
 
 
+void FPlayWorldCommandCallbacks::ResumePlaySession_Clicked()
+{
+	if (HasPlayWorld())
+	{
+		LeaveDebuggingMode();
+		GUnrealEd->PlaySessionResumed();
+		uint32 UserIndex = 0;
+		FSlateApplication::Get().SetUserFocusToGameViewport(UserIndex);
+	}
+}
+
+
+void FPlayWorldCommandCallbacks::PausePlaySession_Clicked()
+{
+	if (HasPlayWorld())
+	{
+		GUnrealEd->PlayWorld->bDebugPauseExecution = true;
+		GUnrealEd->PlaySessionPaused();
+		if (IsInPIE()) {
+			FSlateApplication::Get().ClearKeyboardFocus(EFocusCause::SetDirectly);
+			FSlateApplication::Get().ResetToDefaultInputSettings();
+
+			TWeakPtr<SGlobalPlayWorldActions> ActiveGlobalPlayWorldWidget = FPlayWorldCommands::GetActiveGlobalPlayWorldActionsWidget();
+			if (ActiveGlobalPlayWorldWidget.IsValid())
+			{
+				uint32 UserIndex = 0;
+				FSlateApplication::Get().SetUserFocus(UserIndex, ActiveGlobalPlayWorldWidget.Pin());
+			}
+		}
+	}
+}
+
+
 bool FPlayWorldCommandCallbacks::IsInSIE()
 {
 	return GEditor->bIsSimulatingInEditor;
@@ -1039,6 +1066,24 @@ bool FPlayWorldCommandCallbacks::IsInPIE_AndRunning()
 }
 
 
+bool FPlayWorldCommandCallbacks::HasPlayWorld()
+{
+	return GEditor->PlayWorld != NULL;
+}
+
+
+bool FPlayWorldCommandCallbacks::HasPlayWorldAndPaused()
+{
+	return HasPlayWorld() && GUnrealEd->PlayWorld->bDebugPauseExecution;
+}
+
+
+bool FPlayWorldCommandCallbacks::HasPlayWorldAndRunning()
+{
+	return HasPlayWorld() && !GUnrealEd->PlayWorld->bDebugPauseExecution;
+}
+
+
 //////////////////////////////////////////////////////////////////////////
 // FInternalPlayWorldCommandCallbacks
 
@@ -1093,6 +1138,16 @@ FSlateIcon FInternalPlayWorldCommandCallbacks::GetPossessEjectImage()
 }
 
 
+bool FInternalPlayWorldCommandCallbacks::CanLateJoin()
+{
+	return HasPlayWorld();
+}
+
+bool FInternalPlayWorldCommandCallbacks::CanShowLateJoinButton()
+{
+	return GetDefault<UEditorExperimentalSettings>()->bAllowLateJoinInPIE && HasPlayWorld();
+}
+
 void SetLastExecutedPlayMode(EPlayModeType PlayMode)
 {
 	ULevelEditorPlaySettings* PlaySettings = GetMutableDefault<ULevelEditorPlaySettings>();
@@ -1848,39 +1903,6 @@ FText FInternalPlayWorldCommandCallbacks::GetResumePlaySessionToolTip()
 }
 
 
-void FInternalPlayWorldCommandCallbacks::ResumePlaySession_Clicked()
-{
-	if (HasPlayWorld())
-	{
-		LeaveDebuggingMode();
-		GUnrealEd->PlaySessionResumed();
-		uint32 UserIndex = 0;
-		FSlateApplication::Get().SetUserFocusToGameViewport(UserIndex);
-	}
-}
-
-
-void FInternalPlayWorldCommandCallbacks::PausePlaySession_Clicked()
-{
-	if (HasPlayWorld())
-	{
-		GUnrealEd->PlayWorld->bDebugPauseExecution = true;
-		GUnrealEd->PlaySessionPaused();
-		if (IsInPIE()) {
-			FSlateApplication::Get().ClearKeyboardFocus(EFocusCause::SetDirectly);
-			FSlateApplication::Get().ResetToDefaultInputSettings();
-
-			TWeakPtr<SGlobalPlayWorldActions> ActiveGlobalPlayWorldWidget = FPlayWorldCommands::GetActiveGlobalPlayWorldActionsWidget();
-			if (ActiveGlobalPlayWorldWidget.IsValid())
-			{
-				uint32 UserIndex = 0;
-				FSlateApplication::Get().SetUserFocus(UserIndex, ActiveGlobalPlayWorldWidget.Pin());
-			}
-		}
-	}
-}
-
-
 void FInternalPlayWorldCommandCallbacks::SingleFrameAdvance_Clicked()
 {
 	// We want to function just like Single stepping where we will stop at a breakpoint if one is encountered but we also want to stop after 1 tick if a breakpoint is not encountered.
@@ -1986,13 +2008,7 @@ void FInternalPlayWorldCommandCallbacks::TogglePlayPause_Clicked()
 
 bool FInternalPlayWorldCommandCallbacks::CanShowNonPlayWorldOnlyActions()
 {
-	// @todo vreditor: Don't display Simulate if we're currently in VR mode to prevent crash in blueprint (see UE-27984)
-	bool bIsInVREditor = false;
-	if (IVREditorModule::IsAvailable())
-	{
-		bIsInVREditor = IVREditorModule::Get().IsVREditorEnabled();
-	}
-	return (!HasPlayWorld() && !bIsInVREditor);
+	return !HasPlayWorld();
 }
 
 bool FInternalPlayWorldCommandCallbacks::CanShowVulkanNonPlayWorldOnlyActions()
@@ -2005,22 +2021,6 @@ bool FInternalPlayWorldCommandCallbacks::CanShowVROnlyActions()
 	return !HasPlayWorld();
 }
 
-bool FInternalPlayWorldCommandCallbacks::HasPlayWorld()
-{
-	return GEditor->PlayWorld != NULL;
-}
-
-bool FInternalPlayWorldCommandCallbacks::CanLateJoin()
-{
-	return HasPlayWorld();
-}
-
-bool FInternalPlayWorldCommandCallbacks::CanShowLateJoinButton()
-{
-	return GetDefault<UEditorExperimentalSettings>()->bAllowLateJoinInPIE && HasPlayWorld();
-}
-
-
 int32 FInternalPlayWorldCommandCallbacks::GetNumberOfClients()
 {
 	const ULevelEditorPlaySettings* PlayInSettings = GetDefault<ULevelEditorPlaySettings>();
@@ -2059,18 +2059,6 @@ bool FInternalPlayWorldCommandCallbacks::OnIsDedicatedServerPIEEnabled()
 }
 
 
-bool FInternalPlayWorldCommandCallbacks::HasPlayWorldAndPaused()
-{
-	return HasPlayWorld() && GUnrealEd->PlayWorld->bDebugPauseExecution;
-}
-
-
-bool FInternalPlayWorldCommandCallbacks::HasPlayWorldAndRunning()
-{
-	return HasPlayWorld() && !GUnrealEd->PlayWorld->bDebugPauseExecution;
-}
-
-
 bool FInternalPlayWorldCommandCallbacks::IsStoppedAtBreakpoint()
 {
 	return GIntraFrameDebuggingGameThread;
diff --git a/Engine/Source/Editor/UnrealEd/Private/LevelEditorViewport.cpp b/Engine/Source/Editor/UnrealEd/Private/LevelEditorViewport.cpp
--- a/Engine/Source/Editor/UnrealEd/Private/LevelEditorViewport.cpp
+++ b/Engine/Source/Editor/UnrealEd/Private/LevelEditorViewport.cpp
@@ -68,6 +68,7 @@
 #include "Framework/Notifications/NotificationManager.h"
 #include "Widgets/Notifications/SNotificationList.h"
 #include "Settings/EditorProjectSettings.h"
+#include "IHeadMountedDisplay.h"
 
 DEFINE_LOG_CATEGORY(LogEditorViewport);
 
@@ -4133,7 +4134,19 @@ void FLevelEditorViewportClient::UpdateAudioListener(const FSceneView& View)
 	{
 		if (FAudioDevice* AudioDevice = ViewportWorld->GetAudioDevice())
 		{
-			const FVector& ViewLocation = GetViewLocation();
+			FVector ViewLocation = GetViewLocation();
+
+			const bool bStereoRendering = GEngine->HMDDevice.IsValid() && GEngine->IsStereoscopic3D( Viewport );
+			if( bStereoRendering && GEngine->HMDDevice->IsHeadTrackingAllowed() )
+			{
+				FQuat RoomSpaceHeadOrientation;
+				FVector RoomSpaceHeadLocation;
+				GEngine->HMDDevice->GetCurrentOrientationAndPosition( /* Out */ RoomSpaceHeadOrientation, /* Out */ RoomSpaceHeadLocation );
+
+				// NOTE: The RoomSpaceHeadLocation has already been adjusted for WorldToMetersScale
+				const FVector WorldSpaceHeadLocation = GetViewLocation() + GetViewRotation().RotateVector( RoomSpaceHeadLocation );
+				ViewLocation = WorldSpaceHeadLocation;
+			}
 
 			FMatrix CameraToWorld = View.ViewMatrices.GetInvViewMatrix();
 			FVector ProjUp = CameraToWorld.TransformVector(FVector(0, 1000, 0));
@@ -4525,8 +4538,7 @@ UWorld* FLevelEditorViewportClient::GetWorld() const
 {
 	if (bIsSimulateInEditorViewport)
 	{
-		// TODO: Find a proper way to get this
-		return GWorld;
+		return GEditor->PlayWorld;
 	}
 	else if (World)
 	{
diff --git a/Engine/Source/Editor/UnrealEd/Private/PlayLevel.cpp b/Engine/Source/Editor/UnrealEd/Private/PlayLevel.cpp
--- a/Engine/Source/Editor/UnrealEd/Private/PlayLevel.cpp
+++ b/Engine/Source/Editor/UnrealEd/Private/PlayLevel.cpp
@@ -156,6 +156,8 @@ void UEditorEngine::EndPlayMap()
 
 	TGuardValue<bool> GuardIsEndingPlay(bIsEndingPlay, true);
 
+	FEditorDelegates::PrePIEEnded.Broadcast( bIsSimulatingInEditor );
+
 	FlushAsyncLoading();
 
 	// Monitoring when PIE corrupts references between the World and the PIE generated World for UE-20486
@@ -185,7 +187,7 @@ void UEditorEngine::EndPlayMap()
 		}
 	}
 
-	if (GEngine->HMDDevice.IsValid())
+	if (GEngine->HMDDevice.IsValid() && !bIsSimulatingInEditor)
 	{
 		GEngine->HMDDevice->OnEndPlay(*GEngine->GetWorldContextFromWorld(PlayWorld));
 	}
@@ -588,8 +590,8 @@ void UEditorEngine::TeardownPlaySession(FWorldContext& PieWorldContext)
 
 				if( !bIsSimulatingInEditor)
 				{
-					// Set the editor viewport location to match that of Play in Viewport if we aren't simulating in the editor, we have a valid player to get the location from 
-					if (bLastViewAndLocationValid == true)
+					// Set the editor viewport location to match that of Play in Viewport if we aren't simulating in the editor, we have a valid player to get the location from (unless we're going back to VR Editor, in which case we won't teleport the user.)
+					if (bLastViewAndLocationValid == true && !GEngine->IsStereoscopic3D( Viewport->GetActiveViewport() ) )
 					{
 						bLastViewAndLocationValid = false;
 						Viewport->GetLevelViewportClient().SetViewLocation( LastViewLocation );
@@ -2420,7 +2422,7 @@ void UEditorEngine::PlayInEditor( UWorld* InWorld, bool bInSimulateInEditor )
 		OutputLogErrorsToMessageLogProxyPtr = MakeShareable(new FOutputLogErrorsToMessageLogProxy());
 	}
 
-	if (GEngine->HMDDevice.IsValid())
+	if (GEngine->HMDDevice.IsValid() && !bInSimulateInEditor)
 	{
 		GEngine->HMDDevice->OnBeginPlay(*GEngine->GetWorldContextFromWorld(InWorld));
 	}
@@ -2572,6 +2574,13 @@ void UEditorEngine::PlayInEditor( UWorld* InWorld, bool bInSimulateInEditor )
 		}
 	}
 
+	// Make sure to focus the game viewport.
+	if (!bInSimulateInEditor)
+	{
+		FSlateApplication::Get().SetAllUserFocusToGameViewport();
+	}
+
+	FEditorDelegates::PostPIEStarted.Broadcast( bInSimulateInEditor );
 }
 
 void UEditorEngine::SpawnIntraProcessPIEWorlds(bool bAnyBlueprintErrors, bool bStartInSpectatorMode)
@@ -3144,7 +3153,7 @@ UGameInstance* UEditorEngine::CreatePIEGameInstance(int32 InPIEInstance, bool bI
 	UGameViewportClient* ViewportClient = NULL;
 	ULocalPlayer *NewLocalPlayer = NULL;
 	
-	if (GEngine->HMDDevice.IsValid())
+	if (GEngine->HMDDevice.IsValid() && !bInSimulateInEditor )
 	{
 		GEngine->HMDDevice->OnBeginPlay(*PieWorldContext);
 	}
diff --git a/Engine/Source/Editor/UnrealEd/Private/SnappingUtils.cpp b/Engine/Source/Editor/UnrealEd/Private/SnappingUtils.cpp
--- a/Engine/Source/Editor/UnrealEd/Private/SnappingUtils.cpp
+++ b/Engine/Source/Editor/UnrealEd/Private/SnappingUtils.cpp
@@ -20,11 +20,6 @@
 #include "ISnappingPolicy.h"
 #include "ViewportSnappingModule.h"
 
-namespace EditorViewportSnapping
-{
-	static FAutoConsoleVariable EnableGuides(TEXT("EditorViewportSnapping.EnableGuides"), 0, TEXT("Whether or not guidelines should be enabled. Off by default, set to 1 to enable."));
-}
-
 //////////////////////////////////////////////////////////////////////////
 // FEditorViewportSnapping
 
@@ -376,184 +371,6 @@ void FEditorViewportSnapping::ClearSnappingHelpers( bool bClearImmediately )
 
 void FEditorViewportSnapping::DrawSnappingHelpers(const FSceneView* View, FPrimitiveDrawInterface* PDI)
 {
-	struct Local
-	{
-		struct FGuideData
-		{
-			bool bGuideShouldDraw;
-			FVector SnapPoint;
-			FVector GuideStart;
-			FVector GuideEnd;
-			FLinearColor GuideColor;
-			float GuideLength;
-		};
-
-		static  TArray<FVector> FindPotentialSnapPoints_ActorSpace(const AActor* SnapActor)
-		{
-			TArray<FVector> PotentialSnapPoints;
-
-			FBox ActorSpaceBoundingBox = SnapActor->CalculateComponentsBoundingBoxInLocalSpace();
-			FVector BoxCenter;
-			FVector BoxExtents;
-
-			ActorSpaceBoundingBox.GetCenterAndExtents(BoxCenter, BoxExtents);
-
-			FVector PotentialSnapPoint = FVector::ZeroVector;
-			
-			// Potential snap points are:
-			// The center of each face
-			for (int32 X = -1; X < 2; ++X)
-			{
-				PotentialSnapPoint[0] = X*BoxExtents[0];
-				PotentialSnapPoints.Add(PotentialSnapPoint);
-				// The center of each edge
-				for (int32 Y = -1; Y < 2; ++Y)
-				{
-					PotentialSnapPoint[1] = Y*BoxExtents[1];
-					PotentialSnapPoints.Add(PotentialSnapPoint);
-					// Each corner
-					for (int32 Z = -1; Z < 2; ++Z)
-					{
-						PotentialSnapPoint[2] = Z*BoxExtents[2];
-						PotentialSnapPoints.Add(PotentialSnapPoint);
-					}
-				}
-			}
-			return PotentialSnapPoints;
-		}
-	};
-
-
-	if (EditorViewportSnapping::EnableGuides->GetInt() == 1)
-	{
-		// Guidelines draw differently at different ranges
-		const float LargeGuideDistance = 10.0f;
-		const float MediumGuideDistance = 1.0f;
-		const float SmallGuideDistance = 0.25f;
-
-		// Use the last selected actor as the basis for drawing guidelines. 
-		AActor* TransformingActor = GEditor->GetSelectedActors()->GetBottom<AActor>();
-		if (TransformingActor)
-		{
-			// Don't let the guide lines be shorter than the local bounding box extent in any direction. 
-			// This helps when objects are close together
-			const FBox ActorSpaceBoundingBox = TransformingActor->CalculateComponentsBoundingBoxInLocalSpace();
-			FVector ActorBoxExtents;
-			FVector ActorBoxCenter;
-			ActorSpaceBoundingBox.GetCenterAndExtents(ActorBoxCenter, ActorBoxExtents);
-			const FVector MinGuideLength = ActorBoxExtents - ActorBoxCenter;
-
-			// Get all the potential snap points on the transforming actor
-			const TArray<FVector> ActorTransformingSnapPoints = Local::FindPotentialSnapPoints_ActorSpace(TransformingActor);
-
-			// Find all possible candidates for alignment
-			// TODO: add the world grid
-			// TODO: remove any actors it might not make sense to align to
-			TArray<const AActor*> CandidateActors;
-			for (FActorIterator It(TransformingActor->GetWorld()); It; ++It)
-			{
-				const AActor* PossibleCandidateActor = *It;
-
-				// Don't align to yourself, the entire world, or any actors hidden in the editor
-				if (PossibleCandidateActor != TransformingActor
-					&&	PossibleCandidateActor != TransformingActor->GetWorld()->GetDefaultBrush()
-					&& PossibleCandidateActor->IsHiddenEd() == false)
-				{
-					{
-						CandidateActors.Add(PossibleCandidateActor);
-					}
-				}
-
-			}
-
-			for (const AActor* CandidateActor : CandidateActors)
-			{
-				// Get the actor space snap points for the stationary actor
-				const TArray<FVector> ActorCandidateSnapPoints = Local::FindPotentialSnapPoints_ActorSpace(CandidateActor);
-
-				// Set up the initial guide information for X, Y, and Z guide line
-				TArray<Local::FGuideData> GuideHelpers;
-				Local::FGuideData InitialGuideHelper;
-				InitialGuideHelper.bGuideShouldDraw = false;
-				InitialGuideHelper.SnapPoint = InitialGuideHelper.GuideStart = InitialGuideHelper.GuideEnd = FVector::ZeroVector;
-				InitialGuideHelper.GuideColor = FLinearColor::Black;
-				InitialGuideHelper.GuideLength = 10000000.0f;
-				GuideHelpers.Init(InitialGuideHelper, 3);
-
-				for (const FVector ActorCandidateSnapPoint : ActorCandidateSnapPoints)
-				{
-					// Transform the stationary snap point into world space
-					const FVector WorldCandidateSnapPoint = CandidateActor->ActorToWorld().TransformPosition(ActorCandidateSnapPoint);
-
-					// Check it against each moving snap point
-					for (const FVector ActorTransformingSnapPoint : ActorTransformingSnapPoints)
-					{
-						int32 NumberOfMatchingAxes = 0;
-						// Transform the moving snap point into world space
-						const FVector WorldTransformingSnapPoint = TransformingActor->ActorToWorld().TransformPosition(ActorTransformingSnapPoint);
-						// Transform the moving snap point into stationary local space
-						FVector CandidateActorTransformingSnapPoint = CandidateActor->ActorToWorld().Inverse().TransformPosition(WorldTransformingSnapPoint);
-
-						FLinearColor GuideColor;
-						for (int32 PointAxis = 0; PointAxis < 3; ++PointAxis)
-						{
-							if (FMath::IsNearlyEqual(CandidateActorTransformingSnapPoint[PointAxis], ActorCandidateSnapPoint[PointAxis], LargeGuideDistance))
-							{
-								GuideColor = FLinearColor::Blue;
-								NumberOfMatchingAxes++;
-								if (FMath::IsNearlyEqual(CandidateActorTransformingSnapPoint[PointAxis], ActorCandidateSnapPoint[PointAxis], SmallGuideDistance))
-								{
-									GuideColor = FLinearColor::Green;
-								}
-							}
-						}
-						// Points must match in at least two axes to create a guide line
-						if (NumberOfMatchingAxes >= 2)
-						{
-							for (int32 PointAxis = 0; PointAxis < 3; ++PointAxis)
-							{
-								// If the two points have a component that is not within the LargeGuideDistance of each other, that is the axis we are drawing along
-								if (!FMath::IsNearlyEqual(CandidateActorTransformingSnapPoint[PointAxis], ActorCandidateSnapPoint[PointAxis], LargeGuideDistance))
-								{
-									// The guide line is drawn from where the actor will be aligned to the stationary actor it is aligning to
-									FVector ActorGuideStart = ActorCandidateSnapPoint;
-									ActorGuideStart[PointAxis] = CandidateActorTransformingSnapPoint[PointAxis];
-									const FVector WorldGuideStart = CandidateActor->ActorToWorld().TransformPosition(ActorGuideStart);
-									GuideHelpers[PointAxis].bGuideShouldDraw = true;
-
-									// Try to find the shortest guide line, but also don't switch between two nearly-equal guides (prevents jitter between snap points)
-									// Also, don't let the new guide length drop below our minimum guide length
-									if ((WorldCandidateSnapPoint - WorldGuideStart).Size() < GuideHelpers[PointAxis].GuideLength &&
-										!FMath::IsNearlyEqual(GuideHelpers[PointAxis].GuideLength, (WorldCandidateSnapPoint - WorldGuideStart).Size(), 1.0f)
-										&& (WorldCandidateSnapPoint - WorldGuideStart).Size() >= MinGuideLength[PointAxis])
-									{
-										GuideHelpers[PointAxis].GuideLength = (WorldCandidateSnapPoint - WorldGuideStart).Size();
-										GuideHelpers[PointAxis].GuideStart = WorldGuideStart;
-										GuideHelpers[PointAxis].GuideEnd = WorldCandidateSnapPoint;
-										GuideHelpers[PointAxis].SnapPoint = WorldTransformingSnapPoint;
-										GuideHelpers[PointAxis].GuideColor = GuideColor;
-									}
-								}
-							}
-						}
-
-					}
-				}
-
-				// Draw any of the guide lines in X, Y, or Z that should draw
-				for (int32 PossibleGuides = 0; PossibleGuides < 3; ++PossibleGuides)
-				{
-					if (GuideHelpers[PossibleGuides].bGuideShouldDraw)
-					{
-						PDI->DrawLine(GuideHelpers[PossibleGuides].GuideStart, GuideHelpers[PossibleGuides].GuideEnd, GuideHelpers[PossibleGuides].GuideColor, SDPG_World, 2.0f, 1.0f);
-						PDI->DrawPoint(GuideHelpers[PossibleGuides].SnapPoint, GuideHelpers[PossibleGuides].GuideColor, 10.0f, SDPG_World);
-						PDI->DrawPoint(GuideHelpers[PossibleGuides].GuideEnd, GuideHelpers[PossibleGuides].GuideColor, 10.0f, SDPG_World);
-					}
-				}
-			}
-		}
-	}
-	
 	VertexSnappingImpl.DrawSnappingHelpers( View, PDI );
 }
 
diff --git a/Engine/Source/Editor/UnrealEd/Public/Editor.h b/Engine/Source/Editor/UnrealEd/Public/Editor.h
--- a/Engine/Source/Editor/UnrealEd/Public/Editor.h
+++ b/Engine/Source/Editor/UnrealEd/Public/Editor.h
@@ -132,8 +132,12 @@ struct UNREALED_API FEditorDelegates
 	static FOnEditorModeTransitioned EditorModeExit;
 	/** Sent when a PIE session is beginning (before we decide if PIE can run - allows clients to avoid blocking PIE) */
 	static FOnPIEEvent PreBeginPIE;
-	/** Sent when a PIE session is beginning */
+	/** Sent when a PIE session is beginning (but hasn't actually started yet) */
 	static FOnPIEEvent BeginPIE;
+	/** Sent when a PIE session has fully started and after BeginPlay() has been called */
+	static FOnPIEEvent PostPIEStarted;
+	/** Sent when a PIE session is ending, before anything else happens */
+	static FOnPIEEvent PrePIEEnded;
 	/** Sent when a PIE session is ending */
 	static FOnPIEEvent EndPIE;
 	/** Sent when a PIE session is paused */
diff --git a/Engine/Source/Editor/UnrealEd/Public/EditorWorldExtension.h b/Engine/Source/Editor/UnrealEd/Public/EditorWorldExtension.h
--- a/Engine/Source/Editor/UnrealEd/Public/EditorWorldExtension.h
+++ b/Engine/Source/Editor/UnrealEd/Public/EditorWorldExtension.h
@@ -13,7 +13,7 @@ class UNREALED_API UEditorWorldExtension : public UObject
 {
 	GENERATED_BODY()
 
-	friend class FEditorWorldExtensionCollection;
+	friend class UEditorWorldExtensionCollection;
 public:
 	
 	/** Default constructor */
@@ -37,33 +37,72 @@ public:
 	/** Gets the world owning this extension */
 	virtual UWorld* GetWorld() const override;
 
+	/**  Spawns a transient actor that we can use in the current world of this extension (templated for convenience) */
+	template<class T>
+	inline T* SpawnTransientSceneActor(const FString& ActorName, const bool bWithSceneComponent = false, const EObjectFlags InObjectFlags = EObjectFlags::RF_Transient | EObjectFlags::RF_DuplicateTransient )
+	{
+		return CastChecked<T>(SpawnTransientSceneActor(T::StaticClass(), ActorName, bWithSceneComponent, InObjectFlags));
+	}
+
+	/** Spawns a transient actor that we can use in the current world of this extension */
+	AActor* SpawnTransientSceneActor(TSubclassOf<AActor> ActorClass, const FString& ActorName, const bool bWithSceneComponent = false, const EObjectFlags InObjectFlags = EObjectFlags::RF_Transient | EObjectFlags::RF_DuplicateTransient );
+
+	/** Destroys a transient actor we created earlier */
+	void DestroyTransientActor(AActor* Actor);
+
+	/** Get the owning collection of extensions */
+	UEditorWorldExtensionCollection* GetOwningCollection();
+
+	/** Executes command */
+	bool ExecCommand(const FString& InCommand);
+
+protected:
+	
+	/** Reparent actors to a new world */
+	virtual void TransitionWorld(UWorld* NewWorld);
+
+	/** Give base class a chance to act on entering simulate mode */
+	virtual void EnteredSimulateInEditor() {};
+	
+	/** Give base class a chance to act on leaving simulate mode */
+	virtual void LeftSimulateInEditor() {};
+
+	/** The collection of extensions that is owning this extension */
+	UEditorWorldExtensionCollection* OwningExtensionsCollection;
+
 private:
 
+	void ReparentActor(AActor* Actor, UWorld* NewWorld);
+
 	/** Let the FEditorWorldExtensionCollection set the world of this extension before init */
-	void SetWorld( UWorld* InitWorld );
+	void InitInternal(UEditorWorldExtensionCollection* InOwningExtensionsCollection);
 
-	/** The world that is owning this extension */
-	UWorld* OwningWorld;
+	UPROPERTY()
+	TArray<AActor*> ExtensionActors;
 };
 
 /**
  * Holds a collection of UEditorExtension
  */
-class UNREALED_API FEditorWorldExtensionCollection : public FGCObject
+UCLASS()
+class UNREALED_API UEditorWorldExtensionCollection : public UObject
 {
+	GENERATED_BODY()
+
+	friend class UEditorWorldExtensionManager;
 public:
 
 	/** Default constructor */
-	FEditorWorldExtensionCollection( FWorldContext& InWorldContext );
+	UEditorWorldExtensionCollection();
 
 	/** Default destructor */
-	virtual ~FEditorWorldExtensionCollection();
+	virtual ~UEditorWorldExtensionCollection();
 
 	/** Gets the world from the world context */
-	UWorld* GetWorld() const;
+	virtual UWorld* GetWorld() const override;
 
 	/** Gets the world context */
-	FWorldContext& GetWorldContext() const;
+	FWorldContext* GetWorldContext() const;
 
 	/** 
 	 * Adds an extension to the collection
@@ -99,15 +138,24 @@ public:
 
 	/** Notifies all extensions of axis movement */
 	bool InputAxis( FEditorViewportClient* InViewportClient, FViewport* Viewport, int32 ControllerId, FKey Key, float Delta, float DeltaTime);
+	
+private:
 
-	// FGCObject
-	virtual void AddReferencedObjects( FReferenceCollector& Collector ) override;
-	// End FGCObject
+	/** Sets the world for this collection and gives every extension an opportunity to transition */
+	void SetWorldContext(FWorldContext* InWorldContext);
 
-private:
+	/** Called by the editor after PIE or Simulate is started */
+	void PostPIEStarted( bool bIsSimulatingInEditor );
+
+	/** Called when PIE or Simulate ends */
+	void OnEndPIE( bool bWasSimulatingInEditor );
 
 	/** World context */
-	FWorldContext& WorldContext;
+	FWorldContext* WorldContext;
+
+	/** After entering Simulate, this stores the counterpart editor world to the Simulate world, so that we
+        know this collection needs to transition back to editor world after Simulate finishes */
+	FWorldContext* EditorWorldOnSimulate;
 
 	/** List of extensions along with their reference count.  Extensions will only be truly removed and Shutdown() after their
 	    reference count drops to zero. */
@@ -118,19 +166,21 @@ private:
 /**
  * Holds a map of extension collections paired with worlds
  */
-class UNREALED_API FEditorWorldExtensionManager	
+UCLASS()
+class UNREALED_API UEditorWorldExtensionManager	: public UObject
 {
-public:
+	GENERATED_BODY()
 
+public:
 	/** Default constructor */
-	FEditorWorldExtensionManager();
+	UEditorWorldExtensionManager();
 
 	/** Default destructor */
-	virtual ~FEditorWorldExtensionManager();
+	virtual ~UEditorWorldExtensionManager();
 
 	/** Gets the editor world wrapper that is found with the world passed.
 	 * Adds one for this world if there was non found. */
-	TSharedPtr<FEditorWorldExtensionCollection> GetEditorWorldExtensions(const UWorld* InWorld);
+	UEditorWorldExtensionCollection* GetEditorWorldExtensions(UWorld* InWorld);
 
 	/** Ticks all the collections */
 	void Tick( float DeltaSeconds );
@@ -138,11 +188,14 @@ public:
 private:
 
 	/** Adds a new editor world wrapper when a new world context was created */
-	TSharedPtr<FEditorWorldExtensionCollection> OnWorldContextAdd(FWorldContext& InWorldContext);
+	UEditorWorldExtensionCollection* OnWorldContextAdd(FWorldContext* InWorldContext);
 
 	/** Adds a new editor world wrapper when a new world context was created */
 	void OnWorldContextRemove(FWorldContext& InWorldContext);
 
+	UEditorWorldExtensionCollection* FindExtensionCollection(const UWorld* InWorld);
+
 	/** Map of all the editor world maps */
-	TMap<uint32, TSharedPtr<FEditorWorldExtensionCollection>> EditorWorldExtensionCollection;
+	UPROPERTY()
+	TArray<UEditorWorldExtensionCollection*> EditorWorldExtensionCollection;
 };
\ No newline at end of file
diff --git a/Engine/Source/Editor/UnrealEd/Public/Kismet2/DebuggerCommands.h b/Engine/Source/Editor/UnrealEd/Public/Kismet2/DebuggerCommands.h
--- a/Engine/Source/Editor/UnrealEd/Public/Kismet2/DebuggerCommands.h
+++ b/Engine/Source/Editor/UnrealEd/Public/Kismet2/DebuggerCommands.h
@@ -144,11 +144,18 @@ public:
 	 */
 	static void StartPlayFromHere();
 
+	static void ResumePlaySession_Clicked();
+	static void PausePlaySession_Clicked();
+
 	static bool IsInSIE();
 	static bool IsInPIE();
 
 	static bool IsInSIE_AndRunning();
 	static bool IsInPIE_AndRunning();
+
+	static bool HasPlayWorld();
+	static bool HasPlayWorldAndPaused();
+	static bool HasPlayWorldAndRunning();
 };
 
 #endif // __DebuggerCommands_h__
diff --git a/Engine/Source/Editor/UnrealEd/Public/UnrealEd.h b/Engine/Source/Editor/UnrealEd/Public/UnrealEd.h
--- a/Engine/Source/Editor/UnrealEd/Public/UnrealEd.h
+++ b/Engine/Source/Editor/UnrealEd/Public/UnrealEd.h
@@ -27,7 +27,6 @@ MONOLITHIC_HEADER_BOILERPLATE()
 #include "Classes/ActorFactories/ActorFactoryBoxVolume.h"
 #include "Classes/ActorFactories/ActorFactoryCameraActor.h"
 #include "Classes/ActorFactories/ActorFactoryCharacter.h"
-#include "Classes/ActorFactories/ActorFactorySubDSurface.h"
 #include "Classes/ActorFactories/ActorFactoryClass.h"
 #include "Classes/ActorFactories/ActorFactoryCylinderVolume.h"
 #include "Classes/ActorFactories/ActorFactoryDeferredDecal.h"
@@ -168,7 +167,6 @@ MONOLITHIC_HEADER_BOILERPLATE()
 #include "Classes/Factories/SlateWidgetStyleAssetFactory.h"
 #include "Classes/Factories/StructureFactory.h"
 #include "Classes/Factories/SubsurfaceProfileFactory.h"
-#include "Classes/Factories/SubDSurfaceFactory.h"
 #include "Classes/Factories/Texture2dFactoryNew.h"
 #include "Classes/Factories/TextureFactory.h"
 #include "Classes/Factories/ReimportTextureFactory.h"
diff --git a/Engine/Source/Editor/UnrealEd/UnrealEd.Build.cs b/Engine/Source/Editor/UnrealEd/UnrealEd.Build.cs
--- a/Engine/Source/Editor/UnrealEd/UnrealEd.Build.cs
+++ b/Engine/Source/Editor/UnrealEd/UnrealEd.Build.cs
@@ -149,6 +149,7 @@ public class UnrealEd : ModuleRules
                 "PixelInspectorModule",
 				"MovieScene",
 				"MovieSceneTracks",
+                "ViewportInteraction"
             }
 		);
 
@@ -214,7 +215,7 @@ public class UnrealEd : ModuleRules
                 "IOSPlatformEditor",
 				"HTML5PlatformEditor",
 				"SizeMap",
-                "PortalProxies",
+				"PortalProxies",
                 "PortalServices",
                 "GeometryCacheEd",
                 "BlueprintNativeCodeGen",
@@ -231,7 +232,8 @@ public class UnrealEd : ModuleRules
 			{
                 "GraphEditor",
 				"Kismet",
-				"AudioEditor"
+				"AudioEditor",
+                "ViewportInteraction"
             }
 		); 
 
diff --git a/Engine/Source/Editor/VREditor/Public/VREditorAssetContainer.h b/Engine/Source/Editor/VREditor/Public/VREditorAssetContainer.h
new file mode 100644
--- /dev/null
+++ b/Engine/Source/Editor/VREditor/Public/VREditorAssetContainer.h
@@ -0,0 +1,151 @@
+// Copyright 1998-2017 Epic Games, Inc. All Rights Reserved.
+
+#pragma once
+
+#include "Engine/DataAsset.h"
+#include "VREditorAssetContainer.generated.h"
+
+
+// Forward declarations
+class USoundBase;
+class USoundCue;
+class UStaticMesh;
+class UMaterial;
+class UMaterialInterface;
+class UMaterialInstance;
+class UFont;
+
+/**
+ * Asset container for VREditor.
+ */
+UCLASS()
+class VREDITOR_API UVREditorAssetContainer : public UDataAsset
+{
+	GENERATED_BODY()
+
+public:
+	
+	//
+	// Sounds
+	//
+
+	UPROPERTY(EditAnywhere, Category = Sound)
+	USoundBase* DockableWindowCloseSound;
+
+	UPROPERTY(EditAnywhere, Category = Sound)
+	USoundBase* DockableWindowOpenSound;
+
+	UPROPERTY(EditAnywhere, Category = Sound)
+	USoundBase* DockableWindowDropSound;
+
+	UPROPERTY(EditAnywhere, Category = Sound)
+	USoundBase* DockableWindowDragSound;
+
+	UPROPERTY(EditAnywhere, Category = Sound)
+	USoundBase* DropFromContentBrowserSound;
+
+	UPROPERTY(EditAnywhere, Category = Sound)
+	USoundBase* RadialMenuOpenSound;
+
+	UPROPERTY(EditAnywhere, Category = Sound)
+	USoundBase* RadialMenuCloseSound;
+
+	UPROPERTY(EditAnywhere, Category = Sound)
+	USoundBase* TeleportSound;
+
+	UPROPERTY(EditAnywhere, Category = Sound)
+	USoundCue* ButtonPressSound;
+
+	//
+	// Meshes
+	//
+
+	UPROPERTY(EditAnywhere, Category = Mesh)
+	UStaticMesh* GenericHMDMesh;
+
+	UPROPERTY(EditAnywhere, Category = Mesh)
+	UStaticMesh* PlaneMesh;
+
+	UPROPERTY(EditAnywhere, Category = Mesh)
+	UStaticMesh* CylinderMesh;
+
+	UPROPERTY(EditAnywhere, Category = Mesh)
+	UStaticMesh* LaserPointerStartMesh;
+
+	UPROPERTY(EditAnywhere, Category = Mesh)
+	UStaticMesh* LaserPointerMesh;
+
+	UPROPERTY(EditAnywhere, Category = Mesh)
+	UStaticMesh* LaserPointerEndMesh;
+
+	UPROPERTY(EditAnywhere, Category = Mesh)
+	UStaticMesh* LaserPointerHoverMesh;
+
+	UPROPERTY(EditAnywhere, Category = Mesh)
+	UStaticMesh* VivePreControllerMesh;
+
+	UPROPERTY(EditAnywhere, Category = Mesh)
+	UStaticMesh* OculusControllerMesh;
+
+	UPROPERTY(EditAnywhere, Category = Mesh)
+	UStaticMesh* GenericControllerMesh;
+
+	UPROPERTY(EditAnywhere, Category = Mesh)
+	UStaticMesh* TeleportRootMesh;
+
+	UPROPERTY(EditAnywhere, Category = Mesh)
+	UStaticMesh* WindowMesh;
+
+	UPROPERTY(EditAnywhere, Category = Mesh)
+	UStaticMesh* WindowSelectionBarMesh;
+
+	UPROPERTY(EditAnywhere, Category = Mesh)
+	UStaticMesh* WindowCloseButtonMesh;
+
+	UPROPERTY(EditAnywhere, Category = Mesh)
+	UStaticMesh* RadialMenuMainMesh;
+
+	UPROPERTY(EditAnywhere, Category = Mesh)
+	UStaticMesh* RadialMenuPointerMesh;
+	
+	//
+	// Materials
+	//
+
+	UPROPERTY(EditAnywhere, Category = Material)
+	UMaterialInterface* GridMaterial;
+
+	UPROPERTY(EditAnywhere, Category = Material)
+	UMaterialInterface* LaserPointerMaterial;
+
+	UPROPERTY(EditAnywhere, Category = Material)
+	UMaterialInterface* LaserPointerTranslucentMaterial;
+
+	UPROPERTY(EditAnywhere, Category = Material)
+	UMaterial* WorldMovementPostProcessMaterial;
+
+	UPROPERTY(EditAnywhere, Category = Material)
+	UMaterialInterface* TextMaterial;
+	
+	UPROPERTY(EditAnywhere, Category = Material)
+	UMaterialInterface* VivePreControllerMaterial;
+
+	UPROPERTY(EditAnywhere, Category = Material)
+	UMaterialInterface* OculusControllerMaterial;
+
+	UPROPERTY(EditAnywhere, Category = Material)
+	UMaterialInterface* TeleportMaterial;
+
+	UPROPERTY(EditAnywhere, Category = Material)
+	UMaterialInterface* WindowMaterial;
+
+	UPROPERTY(EditAnywhere, Category = Material)
+	UMaterialInterface* WindowTranslucentMaterial;
+
+	//
+	// Fonts
+	//
+
+	UPROPERTY(EditAnywhere, Category = Font)
+	UFont* TextFont;
+};
diff --git a/Engine/Source/Editor/VREditor/Public/VREditorInteractor.h b/Engine/Source/Editor/VREditor/Public/VREditorInteractor.h
--- a/Engine/Source/Editor/VREditor/Public/VREditorInteractor.h
+++ b/Engine/Source/Editor/VREditor/Public/VREditorInteractor.h
@@ -12,6 +12,15 @@
 class AActor;
 class UWidgetComponent;
 
+enum class EControllerType : uint8
+{
+	Laser,
+	AssistingLaser,
+	UI,
+	Unknown
+};
+
+
 /**
  * VREditor default interactor
  */
@@ -44,8 +53,9 @@ public:
 	virtual void Tick( const float DeltaTime ) override;
 	virtual FHitResult GetHitResultFromLaserPointer( TArray<AActor*>* OptionalListOfIgnoredActors = nullptr, const bool bIgnoreGizmos = false,
 		TArray<UClass*>* ObjectsInFrontOfGizmo = nullptr, const bool bEvenIfBlocked = false, const float LaserLengthOverride = 0.0f ) override;
-	virtual void ResetHoverState( const float DeltaTime ) override;
-	virtual void OnStartDragging( const FVector& HitLocation, const bool bIsPlacingNewObjects ) override;
+	virtual void ResetHoverState() override;
+	virtual bool IsModifierPressed() const override;
+	virtual void PreviewInputKey( class FEditorViewportClient& ViewportClient, FViewportActionKeyInput& Action, const FKey Key, const EInputEvent Event, bool& bOutWasHandled ) override;
 
 	/** Returns the slide delta for pushing and pulling objects. Needs to be implemented by derived classes (e.g. touchpad for vive controller or scrollweel for mouse ) */
 	virtual float GetSlideDelta();
@@ -54,14 +64,17 @@ public:
 	// Getters and setters
 	//
 
+	/** Returns what controller type this is for asymmetric control schemes */
+	EControllerType GetControllerType() const
+	{
+		return ControllerType;
+	};
+
+	/** Set what controller type this is for asymmetric control schemes */
+	void SetControllerType( const EControllerType& InControllerType );
+
 	/** Gets if this interactor is hovering over UI */
 	bool IsHoveringOverUI() const;
-	
-	/** Sets if there is a docked window on this interactor */
-	void SetHasUIInFront( const bool bInHasUIInFront );
-	
-	/** Check if there is any docked window on this interactor */
-	bool HasUIInFront() const;
 
 	/** Sets if the quick menu is on this interactor */
 	void SetHasUIOnForearm( const bool bInHasUIOnForearm );
@@ -70,13 +83,10 @@ public:
 	bool HasUIOnForearm() const;
 
 	/** Gets the current hovered widget component if any */
-	UWidgetComponent* GetHoveringOverWidgetComponent() const;
+	UWidgetComponent* GetLastHoveredWidgetComponent() const;
 
 	/** Sets the current hovered widget component */
-	void SetHoveringOverWidgetComponent( UWidgetComponent* NewHoveringOverWidgetComponent );
-
-	/** If the modifier key is currently pressed */
-	bool IsModifierPressed() const;
+	void SetLastHoveredWidgetComponent( UWidgetComponent* NewHoveringOverWidgetComponent );
 
 	/** Sets if the interactor is clicking on any UI */
 	void SetIsClickingOnUI( const bool bInIsClickingOnUI );
@@ -93,11 +103,11 @@ public:
 	/** Gets if the interactor is right clicking on UI */
 	bool IsRightClickingOnUI() const;
 
-	/** Gets if the interactor is right clicking on UI */
-	void SetLastClickReleaseTime( const double InLastClickReleaseTime );
+	/** Sets the time the interactor last pressed on UI */
+	void SetLastUIPressTime( const double InLastUIPressTime );
 
-	/** Gets last time the interactor released */
-	double GetLastClickReleaseTime() const;
+	/** Gets last time the interactor pressed on UI */
+	double GetLastUIPressTime() const;
 
 	/** Sets the UI scroll velocity */
 	void SetUIScrollVelocity( const float InUIScrollVelocity );
@@ -121,6 +131,9 @@ protected:
 	// General input @todo: VREditor: Should this be general (non-UI) in interactordata ?
 	//
 
+	/** For asymmetrical systems - what type of controller this is */
+	EControllerType ControllerType;
+
 	/** Is the Modifier button held down? */
 	bool bIsModifierPressed;
 
@@ -152,8 +165,8 @@ private:
 	/** Inertial scrolling -- how fast to scroll the mousewheel over UI */
 	float UIScrollVelocity;	
 
-	/** Last real time that we released the 'SelectAndMove' button on UI.  This is used to detect double-clicks. */
-	double LastClickReleaseTime;
+	/** Last real time that we pressed the 'SelectAndMove' button on UI.  This is used to detect double-clicks. */
+	double LastUIPressTime;
 
 protected:
 
diff --git a/Engine/Source/Editor/VREditor/Public/VREditorMode.h b/Engine/Source/Editor/VREditor/Public/VREditorMode.h
--- a/Engine/Source/Editor/VREditor/Public/VREditorMode.h
+++ b/Engine/Source/Editor/VREditor/Public/VREditorMode.h
@@ -14,6 +14,7 @@
 #include "Widgets/SWindow.h"
 #include "HeadMountedDisplayTypes.h"
 #include "Kismet/HeadMountedDisplayFunctionLibrary.h"
+#include "UI/VRRadialMenuHandler.h"
 #include "VREditorMode.generated.h"
 
 class AActor;
@@ -24,17 +25,9 @@ enum class EAutoKeyMode : uint8;
 // Forward declare the GizmoHandleTypes that is defined in VIBaseTransformGizmo.h
 enum class EGizmoHandleTypes : uint8;
 
-/** Reason for exiting the VR editor */
-enum EVREditorExitType
-{
-	Normal = 0,
-	PIE_VR = 1,		//Play in editor with VR
-	SIE_VR = 2		//Simulate in editor with VR
-};
-
 /**
-* Types of actions that can be performed with VR controller devices
-*/
+ * Types of actions that can be performed with VR controller devices
+ */
 namespace VRActionTypes
 {
 	static const FName Touch( "Touch" );
@@ -69,10 +62,10 @@ public:
 	virtual void Shutdown() override;	
 	
 	/** When the user actually enters the VR Editor mode */
-	void Enter(const bool bReenteringVREditing);
+	void Enter();
 
 	/** When the user leaves the VR Editor mode */
-	void Exit(const bool bHMDShouldExitStereo);
+	void Exit( const bool bShouldDisableStereo );
 
 	/** Tick before the ViewportWorldInteraction is ticked */
 	void PreTick( const float DeltaTime );
@@ -99,29 +92,8 @@ public:
 	}
 
 	/** Call this to start exiting VR mode */
-	void StartExitingVRMode( const EVREditorExitType InExitType = EVREditorExitType::Normal );
-
-	/** Gets the reason for exiting the mode */
-	EVREditorExitType GetExitType() const
-	{
-		return ExitType;
-	}
+	void StartExitingVRMode();
 	
-	//bool InputKey( FEditorViewportClient* ViewportClient, FViewport* Viewport, FKey Key, EInputEvent Event );
-	//bool InputAxis( FEditorViewportClient* ViewportClient, FViewport* Viewport, int32 ControllerId, FKey Key, float Delta, float DeltaTime );
-	//bool IsCompatibleWith( FEditorModeID OtherModeID ) const;
-	//void Render( const FSceneView* SceneView, FViewport* Viewport, FPrimitiveDrawInterface* PDI );	
-
-	/**
-	 * Gets our avatar's mesh actor
-	 *
-	 * @return	The mesh actor
-	 */
-	AActor* GetAvatarMeshActor();
-
-	/** Gets the World from the ViewportWorldInteraction used by this mode */
-	virtual UWorld* GetWorld() const override;
-
 	/** Gets the world space transform of the calibrated VR room origin.  When using a seated VR device, this will feel like the
 	camera's world transform (before any HMD positional or rotation adjustments are applied.) */
 	FTransform GetRoomTransform() const;
@@ -211,9 +183,6 @@ public:
 	/** Gets the world scale factor, which can be multiplied by a scale vector to convert to room space */
 	float GetWorldScaleFactor() const;
 
-	/** Sets up our avatar mesh, if not already spawned */
-	void SpawnAvatarMeshActor();
-
 	/** Called internally when the user changes maps, enters/exits PIE or SIE, or switched between PIE/SIE */
 	void CleanUpActorsBeforeMapChangeOrSimulate();
 
@@ -229,15 +198,12 @@ public:
 	/** @return Returns the type of HMD we're dealing with */
 	EHMDDeviceType::Type GetHMDDeviceType() const;
 
-	/** @return Checks to see if the specified hand is aiming roughly toward the specified capsule */
-	bool IsHandAimingTowardsCapsule( class UViewportInteractor* Interactor, const FTransform& CapsuleTransform, const FVector CapsuleStart, const FVector CapsuleEnd, const float CapsuleRadius, const float MinDistanceToCapsule, const FVector CapsuleFrontDirection, const float MinDotForAimingAtCapsule ) const;
-	
+	/** @return Checks to see if the specified interactor is aiming roughly toward the specified capsule */
+	bool IsHandAimingTowardsCapsule(class UViewportInteractor* Interactor, const FTransform& CapsuleTransform, const FVector CapsuleStart, const FVector CapsuleEnd, const float CapsuleRadius, const float MinDistanceToCapsule, const FVector CapsuleFrontDirection, const float MinDotForAimingAtCapsule) const;
+
 	/** Gets the hand interactor  */
 	class UVREditorInteractor* GetHandInteractor( const EControllerHand ControllerHand ) const;
 
-	/** Stops the haptic feedback for the left and right hand interactors for the oculus */
-	void StopOldHapticEffects();
-
 	/** Snaps the current selected actor to the ground */
 	void SnapSelectedActorsToGround();
 	
@@ -291,24 +257,72 @@ public:
 	/** Used to override dockable area restoration behavior */
 	FOnVREditingModeExit OnVREditingModeExit_Handler;
 
-	void SaveSequencerSettings(bool bInKeyAllEnabled, EAutoKeyMode InAutoKeyMode);
+	void SaveSequencerSettings(bool bInKeyAllEnabled, EAutoKeyMode InAutoKeyMode, const class USequencerSettings& InSequencerSettings);
+
+	/** Start or stop simulate-in-editor mode */
+	void ToggleSIEAndVREditor();
+
+	/** Start or stop play-in-editor mode */
+	void TogglePIEAndVREditor();
+
+	/** Create a static motion controller mesh for the current HMD platform */
+	UStaticMeshComponent* CreateMotionControllerMesh( AActor* OwningActor, USceneComponent* AttachmentToComponent );
+
+	/** Helper functions to create a static mesh */
+	UStaticMeshComponent* CreateMesh( AActor* OwningActor, const FString& MeshName, USceneComponent* AttachmentToComponent /*= nullptr */ );
+	UStaticMeshComponent* CreateMesh(AActor* OwningActor, UStaticMesh* Mesh, USceneComponent* AttachmentToComponent /*= nullptr */);
+
+	/** Sets a delegate for the context-specific actions menu */
+	void SetActionsMenuGenerator(const FOnRadialMenuGenerated NewMenuGenerator, const FText NewLabel);
 
+	/** Resets the delegate and button for the context-specific actions menu */
+	void ResetActionsMenuGenerator();
+
+	/** Gets access to VREditorWorldInteraction */
+	class UVREditorWorldInteraction* GetVREditorWorldInteraction()
+	{
+		return VRWorldInteractionExtension;
+	}
+
+	/** Returns true if we started the play in editor session from this VR Editor */
+	bool GetStartedPlayFromVREditor() const;
+	
+	/** Gets the container for all the assets of VREditor. */
+	const class UVREditorAssetContainer& GetAssetContainer() const;
+	
+	/** Plays sound at location. */
+	void PlaySound(USoundBase* SoundBase, const FVector& InWorldLocation, const float InVolume = 1.0f);
+
+	/** Delegate to be called when a material is placed **/
+	DECLARE_EVENT_ThreeParams( UVREditorWorldInteraction, FOnPlaceDraggedMaterial, UPrimitiveComponent*, UMaterialInterface*, bool& );
+	FOnPlaceDraggedMaterial& OnPlaceDraggedMaterial() { return OnPlaceDraggedMaterialEvent; };
+
+protected:
+
+	virtual void TransitionWorld(UWorld* NewWorld) override;
+	virtual void LeftSimulateInEditor() override;
 
 private:
 
 	/** Called when the editor is closed */
 	void OnEditorClosed();
 
-	//Handles closing the VR mode by escape key
-	virtual bool InputKey( FEditorViewportClient* InViewportClient, FViewport* Viewport, FKey Key, EInputEvent Event ) override;
-
 	/** Called when someone closes a standalone VR Editor window */
 	void OnVREditorWindowClosed( const TSharedRef<SWindow>& ClosedWindow );
 
 	/** FEditorDelegates callbacks */
-	void OnBeginPIE( const bool bIsSimulatingInEditor );
-	void OnEndPIE( const bool bIsSimulatingInEditor );
-	void OnSwitchBetweenPIEAndSIE( const bool bIsSimulatingInEditor );
+	void PostPIEStarted( bool bIsSimulatingInEditor );
+	void PrePIEEnded( bool bWasSimulatingInEditor );
+	void OnEndPIE( bool bWasSimulatingInEditor );
+
+	/** Start using the viewport passed */
+	void StartViewport( TSharedPtr<SLevelViewport> Viewport );
+	
+	/** Close the current viewport */
+	void CloseViewport( const bool bShouldDisableStereo );
+
+	/** Restore the world to meters to the saved one when entering VR Editor */
+	void RestoreWorldToMeters();
 
 protected:
 
@@ -331,9 +345,6 @@ protected:
 	/** True if we currently want to exit VR mode.  This is used to defer exiting until it is safe to do that */
 	bool bWantsToExitMode;
 
-	/** The reason for exiting the mode, so the module can close the mode and take further actions on what should happen next */
-	EVREditorExitType ExitType;
-
 	/** True if VR mode is fully initialized and ready to render */
 	bool bIsFullyInitialized;
 
@@ -373,6 +384,9 @@ protected:
 
 	FOnVRTickHandle TickHandle;
 
+	/** Event broadcast when a material is placed */
+	FOnPlaceDraggedMaterial OnPlaceDraggedMaterialEvent;
+
 	//
 	// Subsystems
 	//
@@ -383,7 +397,7 @@ protected:
 
 	/** Teleporter system */
 	UPROPERTY()
-	class UVREditorTeleporter* TeleporterSystem;
+	class AVREditorTeleporter* TeleportActor;
 
 	/** Automatic scale system */
 	UPROPERTY()
@@ -424,15 +438,8 @@ public:
 	{
 		DefaultColor,
 		SelectionColor,
-		TeleportColor,
-		WorldDraggingColor_OneHand,
-		WorldDraggingColor_TwoHands,
-		RedGizmoColor,
-		GreenGizmoColor,
-		BlueGizmoColor,
-		WhiteGizmoColor,
-		HoverGizmoColor,
-		DraggingGizmoColor,
+		WorldDraggingColor,
+		UIColor,
 		UISelectionBarColor,
 		UISelectionBarHoverColor,
 		UICloseButtonColor,
@@ -452,6 +459,9 @@ public:
 	/** Returns the currently active sequencer */
 	class ISequencer* GetCurrentSequencer();
 
+	/** The asset container path */
+	static const FString AssetContainerPath;
+
 private:
 
 	// All the colors for this mode
@@ -465,4 +475,14 @@ private:
 
 	/** Pointer to the current Sequencer */
 	class ISequencer* CurrentSequencer;
+
+	/** The world to meters scale when leaving PIE simulate to restore when back in the editor world. */
+	float SavedWorldToMetersScaleForPIE;
+
+	/** If we started play in editor from the VR Editor*/
+	bool bStartedPlayFromVREditor;	
+	
+	/** Container of assets */
+	UPROPERTY()
+	class UVREditorAssetContainer* AssetContainer;
 };
diff --git a/Engine/Source/Editor/VREditor/Teleporter/VREditorAutoScaler.cpp b/Engine/Source/Editor/VREditor/Teleporter/VREditorAutoScaler.cpp
--- a/Engine/Source/Editor/VREditor/Teleporter/VREditorAutoScaler.cpp
+++ b/Engine/Source/Editor/VREditor/Teleporter/VREditorAutoScaler.cpp
@@ -9,6 +9,11 @@
 #include "GameFramework/WorldSettings.h"
 #include "Sound/SoundCue.h"
 
+namespace VREd
+{
+	static FAutoConsoleVariable AllowResetScale(TEXT("VREd.AllowResetScale"), 1, TEXT("Allowed to reset world to meters to default world to meters"));
+}
+
 UVREditorAutoScaler::UVREditorAutoScaler( const FObjectInitializer& ObjectInitializer ) :
 	Super( ObjectInitializer ),
 	VRMode( nullptr )
@@ -35,55 +40,44 @@ void UVREditorAutoScaler::Shutdown()
 
 void UVREditorAutoScaler::Scale( const float NewWorldToMetersScale )
 {
-	UViewportWorldInteraction* WorldInteraction = &VRMode->GetWorldInteraction();
-
-	// Set the new world to meters scale.
-	const float OldWorldToMetersScale = WorldInteraction->GetWorld()->GetWorldSettings()->WorldToMeters;
-	WorldInteraction->SetWorldToMetersScale( NewWorldToMetersScale );
-	WorldInteraction->SkipInteractiveWorldMovementThisFrame();
-
-	FTransform RoomTransform = WorldInteraction->GetRoomTransform();
-
-	// Use the HMD location as pivot to scale of, this will cause the view to be unchanged. 
-	const FVector RoomSpacePivotLocation = WorldInteraction->GetRoomSpaceHeadTransform().GetLocation();
-	const FVector WorldSpacePivotLocation = WorldInteraction->GetHeadTransform().GetLocation();// *( OldWorldToMetersScale / 100.0f );
-		
-	//The rest of this calculations is a duplicate from UViewportWorldInteraction.
-	const FVector NewRoomSpacePivotLocation = ( RoomSpacePivotLocation / OldWorldToMetersScale ) * NewWorldToMetersScale;
-	const FVector NewWorldSpacePivotLocation = RoomTransform.TransformPosition( NewRoomSpacePivotLocation );
-	const FVector WorldSpacePivotDelta = NewWorldSpacePivotLocation - WorldSpacePivotLocation;
-	const FVector NewWorldSpaceRoomLocation = RoomTransform.GetLocation() - WorldSpacePivotDelta;
-	RoomTransform.SetLocation( NewWorldSpaceRoomLocation );
-
-	WorldInteraction->SetRoomTransform( RoomTransform );
+	if (VREd::AllowResetScale->GetInt() != 0)
+	{
+		UViewportWorldInteraction* WorldInteraction = &VRMode->GetWorldInteraction();
 
-	UGameplayStatics::PlaySound2D( VRMode->GetWorld(), ScaleSound );
+		// Set the new world to meters scale.
+		WorldInteraction->SetWorldToMetersScale( NewWorldToMetersScale, true );
+		WorldInteraction->SkipInteractiveWorldMovementThisFrame();
+		UGameplayStatics::PlaySound2D( VRMode->GetWorld(), ScaleSound );
+	}
 }
 
 
 void UVREditorAutoScaler::OnInteractorAction( class FEditorViewportClient& ViewportClient, UViewportInteractor* Interactor, const struct FViewportActionKeyInput& Action, bool& bOutIsInputCaptured, bool& bWasHandled )
 {
-	if( Action.ActionType == VRActionTypes::ConfirmRadialSelection )
+	if (VREd::AllowResetScale->GetInt() != 0)
 	{
-		if( Action.Event == IE_Pressed )
+		if( Action.ActionType == VRActionTypes::ConfirmRadialSelection )
 		{
-			const EViewportInteractionDraggingMode DraggingMode = Interactor->GetDraggingMode();
-			UViewportInteractor* OtherInteractor = Interactor->GetOtherInteractor();
-
-			// Start dragging at laser impact when already dragging actors freely
-			if( DraggingMode == EViewportInteractionDraggingMode::World ||
-				( DraggingMode == EViewportInteractionDraggingMode::AssistingDrag && OtherInteractor != nullptr && OtherInteractor->GetDraggingMode() == EViewportInteractionDraggingMode::World ) )
+			if( Action.Event == IE_Pressed )
 			{
-				const float DefaultWorldToMetersScale = VRMode->GetSavedEditorState().WorldToMetersScale;
-				Scale( DefaultWorldToMetersScale );
-
-				bWasHandled = true;
-				bOutIsInputCaptured = true;
+				const EViewportInteractionDraggingMode DraggingMode = Interactor->GetDraggingMode();
+				UViewportInteractor* OtherInteractor = Interactor->GetOtherInteractor();
+
+				// Start dragging at laser impact when already dragging actors freely
+				if( DraggingMode == EViewportInteractionDraggingMode::World ||
+					( DraggingMode == EViewportInteractionDraggingMode::AssistingDrag && OtherInteractor != nullptr && OtherInteractor->GetDraggingMode() == EViewportInteractionDraggingMode::World ) )
+				{
+					const float DefaultWorldToMetersScale = VRMode->GetSavedEditorState().WorldToMetersScale;
+					Scale( DefaultWorldToMetersScale );
+
+					bWasHandled = true;
+					bOutIsInputCaptured = true;
+				}
+			}
+			else if( Action.Event == IE_Released )
+			{
+				bOutIsInputCaptured = false;
 			}
-		}
-		else if( Action.Event == IE_Released )
-		{
-			bOutIsInputCaptured = false;
 		}
 	}
 }
diff --git a/Engine/Source/Editor/VREditor/Teleporter/VREditorTeleporter.cpp b/Engine/Source/Editor/VREditor/Teleporter/VREditorTeleporter.cpp
--- a/Engine/Source/Editor/VREditor/Teleporter/VREditorTeleporter.cpp
+++ b/Engine/Source/Editor/VREditor/Teleporter/VREditorTeleporter.cpp
@@ -1,157 +1,473 @@
 // Copyright 1998-2017 Epic Games, Inc. All Rights Reserved.
 
 #include "VREditorTeleporter.h"
-#include "HAL/IConsoleManager.h"
-#include "Engine/EngineTypes.h"
-#include "Kismet/GameplayStatics.h"
-#include "Sound/SoundCue.h"
 #include "VREditorMode.h"
 #include "ViewportInteractionTypes.h"
 #include "ViewportWorldInteraction.h"
 #include "VREditorInteractor.h"
+#include "VREditorMotionControllerInteractor.h"
+#include "VREditorAssetContainer.h"
 
+#include "Engine/EngineTypes.h"
+#include "HeadMountedDisplayTypes.h"
+#include "Sound/SoundCue.h"
+#include "Materials/MaterialInstance.h"
+#include "Materials/MaterialInstanceDynamic.h"
+#include "HAL/IConsoleManager.h"
 
 namespace VREd
 {
-	static FAutoConsoleVariable TeleportLerpTime( TEXT( "VREd.TeleportLerpTime" ), 0.1f, TEXT( "The lerp time to teleport" ) );
-	static FAutoConsoleVariable TeleportOffset( TEXT( "VREd.TeleportOffset" ), 100.0f, TEXT( "The offset from the hitresult towards the controller" ) );
-	static FAutoConsoleVariable TeleportLaserPointerLength( TEXT( "VREd.TeleportLaserPointerLength" ), 500000.0f, TEXT( "Distance of the LaserPointer for teleporting" ) );
-	static FAutoConsoleVariable TeleportDistance( TEXT( "VREd.TeleportDistance" ), 500.0f, TEXT( "Default distance for teleporting when not hitting anything" ) );
+	static FAutoConsoleVariable TeleportLerpTime(TEXT("VREd.TeleportLerpTime"), 0.1f, TEXT("The lerp time to teleport"));
+	static FAutoConsoleVariable TeleportOffset(TEXT("VREd.TeleportOffset"), 100.0f, TEXT("The offset from the hitresult towards the controller"));
+	static FAutoConsoleVariable TeleportLaserPointerLength(TEXT("VREd.TeleportLaserPointerLength"), 500000.0f, TEXT("Distance of the LaserPointer for teleporting"));
+	static FAutoConsoleVariable TeleportDistance(TEXT("VREd.TeleportDistance"), 500.0f, TEXT("Default distance for teleporting when not hitting anything"));
+	static FAutoConsoleVariable TeleportScaleSensitivity(TEXT("VREd.TeleportScaleSensitivity"), 0.05f, TEXT("Teleport world to meters scale touchpad sensitivity"));
+	static FAutoConsoleVariable TeleportOffsetMultiplier(TEXT("VREd.TeleportOffsetMultiplier"), 0.3f, TEXT("Teleport offset multiplier"));
+	static FAutoConsoleVariable TeleportEnableChangeScale(TEXT("VREd.TeleportEnableChangeScale"), 1, TEXT("Ability to change the world to meters scale while teleporting"));
+	static FAutoConsoleVariable TeleportFadeInAnimateSpeed(TEXT("VREd.TeleportAnimateSpeed"), 3.0f, TEXT("How fast the teleporter should fade in"));
+	static FAutoConsoleVariable TeleportDragSpeed(TEXT("VREd.TeleportDragSpeed"), 0.3f, TEXT("How fast the teleporter should drag behind the laser aiming location"));
+	static FAutoConsoleVariable TeleportAllowScaleBackToDefault(TEXT("VREd.TeleportAllowScaleBackToDefault"), 1, TEXT("Scale back to default world to meters scale"));
+	static FAutoConsoleVariable TeleportAllowPushPull(TEXT("VREd.TeleportAllowPushPull"), 1, TEXT("Allow being able to push and pull the teleporter along the laser."));
+	static FAutoConsoleVariable TeleportSlideBuffer(TEXT("VREd.TeleportSlideBuffer"), 0.3f, TEXT("The minimum slide on trackpad to push/pull or change scale."));
 }
 
-UVREditorTeleporter::UVREditorTeleporter( const FObjectInitializer& Initializer ) :
-	Super( Initializer ),
-	VRMode( nullptr ),
-	bIsTeleporting( false ),
-	TeleportLerpAlpha( 0 ),
-	TeleportStartLocation( FVector::ZeroVector ),
-	TeleportGoalLocation( FVector::ZeroVector ),
-	bJustTeleported( false )
+AVREditorTeleporter::AVREditorTeleporter():
+	Super(),
+	VRMode(nullptr),
+	TeleportingState(EState::None),
+	TeleportLerpAlpha(0),
+	TeleportStartLocation(FVector::ZeroVector),
+	TeleportGoalLocation(FVector::ZeroVector),
+	TeleportDirectionMeshComponent(nullptr),
+	HMDMeshComponent(nullptr),
+	LeftMotionControllerMeshComponent(nullptr),
+	RightMotionControllerMeshComponent(nullptr),
+	TeleportMID(nullptr),
+	InteractorTryingTeleport(nullptr),
+	OffsetDistance(FVector::ZeroVector),
+	TeleportGoalScale(0),
+	DragRayLength(0),
+	DragRayLengthVelocity(0),
+	bPushedFromEndOfLaser(false),
+	bInitialTeleportAim(true),
+	FadeAlpha(1.0f),
+	bShouldBeVisible(),
+	TeleportTickDelay(0)
 {
-	// Load sounds
-	TeleportSound = LoadObject<USoundCue>( nullptr, TEXT( "/Engine/VREditor/Sounds/VR_teleport_Cue" ) );
-	check( TeleportSound != nullptr );
 }
 
-void UVREditorTeleporter::Init( UVREditorMode* InMode )
+AVREditorTeleporter::~AVREditorTeleporter()
+{
+	TeleportDirectionMeshComponent = nullptr;
+	HMDMeshComponent = nullptr;
+	LeftMotionControllerMeshComponent = nullptr;
+	RightMotionControllerMeshComponent = nullptr;
+	TeleportMID = nullptr;
+	InteractorTryingTeleport = nullptr;
+	VRMode = nullptr;
+}
+
+void AVREditorTeleporter::Init(UVREditorMode* InMode)
 {
 	VRMode = InMode;
-	VRMode->OnTickHandle().AddUObject( this, &UVREditorTeleporter::Tick );
-	VRMode->GetWorldInteraction().OnViewportInteractionInputAction().AddUObject( this, &UVREditorTeleporter::OnInteractorAction );
+
+	// Register for event for ticking and input
+	VRMode->OnTickHandle().AddUObject(this, &AVREditorTeleporter::Tick);
+	VRMode->GetWorldInteraction().OnPreviewInputAction().AddUObject(this, &AVREditorTeleporter::OnPreviewInputAction);
+
+	SetActorEnableCollision(false);
+
+	const UVREditorAssetContainer& AssetContainer = VRMode->GetAssetContainer();
+
+	{
+		UMaterialInterface* RoomSpaceMaterial = AssetContainer.TeleportMaterial;
+		check(RoomSpaceMaterial != nullptr);
+		TeleportMID = UMaterialInstanceDynamic::Create(RoomSpaceMaterial, GetTransientPackage());
+		check(TeleportMID != nullptr);
+	}
+
+	{
+		TeleportDirectionMeshComponent = VRMode->CreateMesh(this, AssetContainer.TeleportRootMesh, GetRootComponent());
+		TeleportDirectionMeshComponent->SetWorldScale3D(FVector(1, 1, 1));
+		TeleportDirectionMeshComponent->SetMaterial(0, TeleportMID);
+		TeleportDirectionMeshComponent->SetCollisionEnabled(ECollisionEnabled::NoCollision);
+		TeleportDirectionMeshComponent->SetCollisionResponseToAllChannels(ECollisionResponse::ECR_Ignore);
+		TeleportDirectionMeshComponent->SetCollisionResponseToChannel(COLLISION_GIZMO, ECollisionResponse::ECR_Ignore);
+		TeleportDirectionMeshComponent->SetCastShadow(false);
+
+		HMDMeshComponent = VRMode->CreateMesh(this, AssetContainer.GenericHMDMesh, GetRootComponent());
+		HMDMeshComponent->SetMaterial(0, TeleportMID);
+		HMDMeshComponent->SetCastShadow(false);
+
+		LeftMotionControllerMeshComponent = VRMode->CreateMotionControllerMesh(this, GetRootComponent());
+		check(LeftMotionControllerMeshComponent != nullptr);
+		LeftMotionControllerMeshComponent->SetCastShadow(false);
+
+		RightMotionControllerMeshComponent = VRMode->CreateMotionControllerMesh(this, GetRootComponent());
+		check(RightMotionControllerMeshComponent != nullptr);
+		RightMotionControllerMeshComponent->SetCastShadow(false);
+
+		LeftMotionControllerMeshComponent->SetMaterial(0, TeleportMID);
+		RightMotionControllerMeshComponent->SetMaterial(0, TeleportMID);
+	}
+
+	// Set the initial teleport scale to the current world to meters
+	TeleportGoalScale = VRMode->GetSavedEditorState().WorldToMetersScale;
+
+	Show(false);
 }
 
-void UVREditorTeleporter::Shutdown()
+void AVREditorTeleporter::Shutdown()
 {
-	VRMode->OnTickHandle().RemoveAll( this );
-	VRMode->GetWorldInteraction().OnViewportInteractionInputAction().RemoveAll( this );
-	TeleportSound = nullptr;
+	VRMode->OnTickHandle().RemoveAll(this);
+	VRMode->GetWorldInteraction().OnViewportInteractionInputAction().RemoveAll(this);
 	VRMode = nullptr;
 }
 
-void UVREditorTeleporter::Tick( const float DeltaTime )
+void AVREditorTeleporter::Tick(const float DeltaTime)
 {
-	if ( bIsTeleporting )
+	UpdateFadingState(DeltaTime);
+	if (TeleportingState == EState::Aiming && InteractorTryingTeleport != nullptr)
+	{
+		{
+			UVREditorMotionControllerInteractor* VREditorInteractor = Cast<UVREditorMotionControllerInteractor>(InteractorTryingTeleport);
+			if (VREditorInteractor != nullptr && (VREditorInteractor->GetDraggingMode() == EViewportInteractionDraggingMode::World ||
+				(VREditorInteractor->GetOtherInteractor() != nullptr && VREditorInteractor->GetOtherInteractor()->GetInteractorData().DraggingMode == EViewportInteractionDraggingMode::World &&
+					VREditorInteractor->GetDraggingMode() == EViewportInteractionDraggingMode::AssistingDrag)) &&
+				!FMath::IsNearlyZero(VREditorInteractor->GetSelectAndMoveTriggerValue(), KINDA_SMALL_NUMBER))
+			{
+				VREditorInteractor->SetForceShowLaser(true);
+			}
+		}
+
+		UpdateTeleportAim(DeltaTime); 
+	}
+	else if (TeleportingState == EState::Teleporting)
 	{
-		Teleport( DeltaTime );
+		static const uint32 StaticTeleportDelay = 2;
+		if (TeleportTickDelay > StaticTeleportDelay)
+		{
+			Teleport(DeltaTime);
+		}
+		TeleportTickDelay++;
 	}
 }
 
-void UVREditorTeleporter::StartTeleport( UViewportInteractor* Interactor )
+void AVREditorTeleporter::StartTeleport(UViewportInteractor* Interactor)
 {
-	UVREditorInteractor* VREditorInteractor = Cast<UVREditorInteractor>( Interactor );
-	if ( VREditorInteractor && !VREditorInteractor->IsHoveringOverUI() )
+
+	if (VREd::TeleportEnableChangeScale->GetInt() != 0 || VREd::TeleportAllowScaleBackToDefault->GetInt() != 0)
 	{
-		FVector LaserPointerStart, LaserPointerEnd;
-		if ( Interactor->GetLaserPointer( LaserPointerStart, LaserPointerEnd ) )
+		// Set the world to meters scale @todo vreditor: This causes some strange movement 
+		VRMode->GetWorldInteraction().SetWorldToMetersScale(TeleportGoalScale * 100, true);
+	}
+
+	TeleportingState = EState::Teleporting;
+	TeleportLerpAlpha = 0.0f;
+
+	const UVREditorAssetContainer& AssetContainer = VRMode->GetAssetContainer();
+	VRMode->PlaySound(AssetContainer.TeleportSound, TeleportGoalLocation);
+	Interactor->PlayHapticEffect(1.0f);
+}
+
+void AVREditorTeleporter::OnPreviewInputAction(class FEditorViewportClient& ViewportClient, UViewportInteractor* Interactor, const struct FViewportActionKeyInput& Action, bool& bOutIsInputCaptured, bool& bWasHandled)
+{
+	UVREditorMotionControllerInteractor* VRMotionControllerInteractor = Cast<UVREditorMotionControllerInteractor>(Interactor);
+	if (VRMotionControllerInteractor != nullptr && VRMotionControllerInteractor->GetControllerType() == EControllerType::Laser)
+	{
+		const bool bDraggingWorldForTeleport =
+			Interactor->GetDraggingMode() == EViewportInteractionDraggingMode::World &&
+			(Interactor->GetOtherInteractor() == nullptr || Interactor->GetOtherInteractor()->GetDraggingMode() != EViewportInteractionDraggingMode::AssistingDrag);
+		if (bDraggingWorldForTeleport && 
+			Action.Event == IE_Pressed && 
+			Action.ActionType == ViewportWorldActionTypes::SelectAndMove)
 		{
-			FVector EndLocation;
-			FHitResult HitResult = Interactor->GetHitResultFromLaserPointer( nullptr, true, nullptr, false, VREd::TeleportLaserPointerLength->GetFloat() );
-			if( HitResult.bBlockingHit )
+			if (InteractorTryingTeleport == nullptr && TeleportingState == EState::None)
 			{
-				EndLocation = HitResult.Location;
-			}
-			else
-			{
-				FVector Direction = LaserPointerEnd - LaserPointerStart;
-				Direction.Normalize();
-				EndLocation = LaserPointerStart + ( Direction * ( VREd::TeleportDistance->GetFloat() * VRMode->GetWorldScaleFactor() ) );
+				TeleportGoalScale = VRMode->GetWorldScaleFactor();
+
+				// Checking teleport
+				SetVisibility(true);
+				InteractorTryingTeleport = Interactor;
+
+				TeleportingState = EState::Aiming;
+				SetColor(VRMode->GetColor(UVREditorMode::EColors::WorldDraggingColor));
+				VRMode->GetWorldInteraction().AllowWorldMovement(false);
 			}
 
-			const FVector HeadLocation = VRMode->GetHeadTransform().GetLocation();
-			const FVector RoomLocation = VRMode->GetRoomTransform().GetLocation();
+			bWasHandled = true;
+		}
 
-			FVector NewHandLocation = EndLocation;
-			if( HitResult.bBlockingHit )
+		if (Action.Event == IE_Released && 
+			Action.ActionType == ViewportWorldActionTypes::SelectAndMove &&
+			InteractorTryingTeleport != nullptr && InteractorTryingTeleport == Interactor)
+		{
+			UVREditorMotionControllerInteractor* VREditorInteractor = Cast<UVREditorMotionControllerInteractor>(InteractorTryingTeleport);
+			if (VREditorInteractor && VREditorInteractor == InteractorTryingTeleport && !(VREditorInteractor->IsHoveringOverGizmo() || VREditorInteractor->IsHoveringOverUI()))
 			{
-				const FVector Offset = ( LaserPointerStart - EndLocation ).GetSafeNormal() *
-					VREd::TeleportOffset->GetFloat() * VRMode->GetWorldInteraction().GetWorldScaleFactor(); // Offset from the hitresult, TeleportOffset is in centimeter
-				NewHandLocation = EndLocation + Offset;
-			}
+				// Confirm teleport
+				StartTeleport(Interactor);
+
+				// Clean everything up
+				bPushedFromEndOfLaser = false;
+				bInitialTeleportAim = true;
+				bShouldBeVisible = false;
+				SetVisibility(false);
 
-			const FVector HandHeadOffset = HeadLocation - Interactor->GetTransform().GetLocation();		// The offset between the hand and the head
-			const FVector NewHeadLocation = NewHandLocation + HandHeadOffset;							// The location in world space where the head is going to be
-																										// The new roomspace location in world space
-			const FVector HeadRoomOffset = HeadLocation - RoomLocation;									// Where the head is going to be in world space
-			const FVector NewRoomSpaceLocation = NewHeadLocation - HeadRoomOffset;						// Result, where the roomspace is going to be in world space
-																										// Set values to start teleporting
-			bIsTeleporting = true;
-			TeleportLerpAlpha = 0.0f;
-			TeleportStartLocation = RoomLocation;
-			TeleportGoalLocation = NewRoomSpaceLocation;
-
-			UGameplayStatics::PlaySound2D( VRMode->GetWorld(), TeleportSound );
+				VREditorInteractor->SetForceShowLaser(false);
+				InteractorTryingTeleport = nullptr;
+
+				bWasHandled = true;
+			}
 		}
 	}
 }
 
-void UVREditorTeleporter::OnInteractorAction( FEditorViewportClient& ViewportClient, UViewportInteractor* Interactor, const struct FViewportActionKeyInput& Action, bool& bOutIsInputCaptured, bool& bWasHandled )
+void AVREditorTeleporter::Teleport(const float DeltaTime)
+{
+	if (TeleportLerpAlpha == 0.0f)
+	{
+		TeleportStartLocation = VRMode->GetRoomTransform().GetLocation();
+	}
+	
+	FTransform RoomTransform = VRMode->GetRoomTransform();
+	TeleportLerpAlpha += DeltaTime;
+	if (TeleportLerpAlpha > VREd::TeleportLerpTime->GetFloat())
+	{
+		// Teleporting is finished
+		TeleportLerpAlpha = VREd::TeleportLerpTime->GetFloat();
+		TeleportingState = EState::None;
+		VRMode->GetWorldInteraction().AllowWorldMovement(true);
+		TeleportTickDelay = 0;
+	}
+
+	// Calculate the new position of the roomspace
+	FVector NewLocation = FMath::Lerp(TeleportStartLocation, TeleportGoalLocation, TeleportLerpAlpha / VREd::TeleportLerpTime->GetFloat());
+	RoomTransform.SetLocation(NewLocation);
+	VRMode->SetRoomTransform(RoomTransform);
+}
+
+void AVREditorTeleporter::UpdateTeleportAim(const float DeltaTime)
 {
-	if ( ( Action.ActionType == ViewportWorldActionTypes::SelectAndMove ||
-		Action.ActionType == ViewportWorldActionTypes::SelectAndMove_LightlyPressed ) )
+	UVREditorMotionControllerInteractor* VREditorInteractor = Cast<UVREditorMotionControllerInteractor>(InteractorTryingTeleport);
+	if (VREditorInteractor)
 	{
-		if ( Action.Event == IE_Pressed )
+		Show(!(VREditorInteractor->IsHoveringOverGizmo() || VREditorInteractor->IsHoveringOverUI()));
+
+		FVector LaserPointerStart, LaserPointerEnd, EndLocation;
+		if (VREditorInteractor->GetLaserPointer(LaserPointerStart, LaserPointerEnd) && TeleportingState == EState::Aiming)
 		{
-			if ( Action.ActionType == ViewportWorldActionTypes::SelectAndMove_LightlyPressed )
+			const bool bAllowPushPull = VREd::TeleportAllowPushPull->GetInt() == 0 ? false : true;
+
+			if (VREd::TeleportEnableChangeScale->GetInt() != 0)
 			{
-				// Check to teleport before checking on actor selection, otherwise objects will be selected when teleporting
-				if	( Interactor->GetDraggingMode() == EViewportInteractionDraggingMode::World &&
-					( Interactor->GetOtherInteractor() == nullptr || Interactor->GetOtherInteractor()->GetDraggingMode() != EViewportInteractionDraggingMode::AssistingDrag ) )
+				const float SlideDelta = GetSlideDelta(VREditorInteractor, 0);
+				if (SlideDelta != 0.0f)
 				{
-					StartTeleport( Interactor );
-					bJustTeleported = true;
-					bWasHandled = true;
+					// Calculate the new goal scale with the trackpad delta X axis 
+					TeleportGoalScale += SlideDelta * (TeleportGoalScale * VREd::TeleportScaleSensitivity->GetFloat());
+
+					const float MinScale = VRMode->GetWorldInteraction().GetMinScale() * 0.01f;
+					const float MaxScale = VRMode->GetWorldInteraction().GetMaxScale() * 0.01f;
+
+					// The scale must be in the limits of the worldinteraction minimum and maximum scale
+					TeleportGoalScale = FMath::Clamp(TeleportGoalScale, MinScale, MaxScale);
 				}
 			}
-			else if ( Action.ActionType == ViewportWorldActionTypes::SelectAndMove && bJustTeleported )
+			else if (VREd::TeleportAllowScaleBackToDefault->GetInt() != 0)
 			{
-				//Block the input
-				bJustTeleported = false;
-				bWasHandled = true;
+				TeleportGoalScale = VRMode->GetSavedEditorState().WorldToMetersScale / 100.f;
+			}
+			else
+			{
+				TeleportGoalScale = VRMode->GetWorldScaleFactor();
 			}
 
-			// Nothing should happen when the hand has a UI in front, otherwise it will deselect everything
-			if ( Interactor->GetIsLaserBlocked() )
+			// If the laser is hitting something the teleport will go their with an appropriate offset. 
+			FHitResult HitResult = VREditorInteractor->GetHitResultFromLaserPointer(nullptr, true, nullptr, false, VREd::TeleportLaserPointerLength->GetFloat());
+			if (HitResult.bBlockingHit && !bPushedFromEndOfLaser)
 			{
-				bWasHandled = true;
+				// Calculate an offset with the impact normal, so the teleporter will show up on-top, underneath or next to where laser is aiming at.
+				OffsetDistance = HitResult.ImpactNormal * (50 * TeleportGoalScale);
+				float OffsetFromImpactNormalZ = (HitResult.ImpactNormal.Z - 1);
+				OffsetFromImpactNormalZ -= (OffsetFromImpactNormalZ * VREd::TeleportOffsetMultiplier->GetFloat());
+				OffsetDistance.Z = OffsetFromImpactNormalZ * (((VRMode->GetHeadTransform().GetLocation().Z - VRMode->GetRoomTransform().GetLocation().Z) / VRMode->GetWorldScaleFactor())  * TeleportGoalScale);
+
+				// Set the final location based on the hit location and the offset
+				EndLocation = HitResult.Location + OffsetDistance;
+
+				// Update the raylength to the current length, so if we have to pull or push the teleporter next frame it will not jump
+				DragRayLength = FVector::Dist(LaserPointerStart, HitResult.ImpactPoint);
 			}
+			// If the laser is not hitting anything or has been pushed away already, the user can push or pull the teleporter along the laser
+			else
+			{
+				EndLocation = UpdatePushPullTeleporter(VREditorInteractor, LaserPointerStart, LaserPointerEnd, bAllowPushPull);
+			}
+
+			// The trackpad has been used while aiming for teleporting, so the teleporter won't go to the end of the laser after this
+			if (!bPushedFromEndOfLaser && bAllowPushPull &&  GetSlideDelta(VREditorInteractor, 1) != 0.0f)
+			{
+				bPushedFromEndOfLaser = true;
+			} 
+				
+			// Smooth the final location
+			if (!bInitialTeleportAim)
+			{
+				FVector TeleporterAndAimLocationOffset = EndLocation - GetActorLocation();
+				if (TeleporterAndAimLocationOffset.Size() > (VRMode->GetWorldScaleFactor() * 1.0f)) //@todo vreditor: Tweak this value for less jitter
+				{
+					EndLocation = GetActorLocation() + (TeleporterAndAimLocationOffset * VREd::TeleportDragSpeed->GetFloat());
+				}
+			}
+
+			UpdateVisuals(EndLocation);
+			bInitialTeleportAim = false;
 		}
 	}
 }
 
-void UVREditorTeleporter::Teleport( const float DeltaTime )
+FVector AVREditorTeleporter::UpdatePushPullTeleporter(UVREditorMotionControllerInteractor* VREditorInteractor, const FVector& LaserPointerStart, const FVector& LaserPointerEnd, const bool bEnablePushPull /* = true */)
 {
-	FTransform RoomTransform = VRMode->GetRoomTransform();
-	TeleportLerpAlpha += DeltaTime;
+	if (bEnablePushPull && GetSlideDelta(VREditorInteractor, 1) != 0.0f)
+	{
+		VREditorInteractor->CalculateDragRay(DragRayLength, DragRayLengthVelocity);
+	}
 
-	if ( TeleportLerpAlpha > VREd::TeleportLerpTime->GetFloat() )
+	FVector Direction = LaserPointerEnd - LaserPointerStart;
+	Direction.Normalize();
+
+	return (LaserPointerStart + (Direction * DragRayLength)) + OffsetDistance;
+}
+
+void AVREditorTeleporter::SetVisibility(const bool bVisible)
+{
+	TeleportDirectionMeshComponent->SetVisibility(bVisible);
+	HMDMeshComponent->SetVisibility(bVisible);
+	LeftMotionControllerMeshComponent->SetVisibility(bVisible);
+	RightMotionControllerMeshComponent->SetVisibility(bVisible);
+}
+
+void AVREditorTeleporter::SetColor(const FLinearColor& Color)
+{
+	static FName StaticRoomspaceMeshColorName("Color");
+	TeleportMID->SetVectorParameterValue(StaticRoomspaceMeshColorName, Color);
+}
+
+void AVREditorTeleporter::UpdateVisuals(const FVector& NewLocation)
+{
+	const float WorldScale = VRMode->GetWorldScaleFactor();
+	const float CalculatedScale = CalculateAnimatedScaleFactor();
+	const FVector AnimatedScale = FVector(TeleportGoalScale * CalculatedScale);
+
+	SetActorLocation(NewLocation);
+	FTransform HMDTransform = VRMode->GetHeadTransform();
+
+	// Update the teleport direction indicator
 	{
-		// Teleporting is finished
-		TeleportLerpAlpha = VREd::TeleportLerpTime->GetFloat();
-		bIsTeleporting = false;
+		const FRotator TeleportDirection = FRotator(0.f, HMDTransform.GetRotation().Rotator().Yaw, 0.f);
+		TeleportDirectionMeshComponent->SetWorldRotation(TeleportDirection);
+		TeleportDirectionMeshComponent->SetWorldScale3D(AnimatedScale);
+		TeleportDirectionMeshComponent->SetRelativeLocation(FVector(0.0f, 0.0f, AnimatedScale.Z * 0.5f));
 	}
 
-	// Calculate the new position of the roomspace
-	FVector NewLocation = FMath::Lerp( TeleportStartLocation, TeleportGoalLocation, TeleportLerpAlpha / VREd::TeleportLerpTime->GetFloat() );
-	RoomTransform.SetLocation( NewLocation );
-	VRMode->SetRoomTransform( RoomTransform );
+	HMDTransform.SetLocation(FVector(NewLocation.X, NewLocation.Y, NewLocation.Z + ((VRMode->GetRoomSpaceHeadTransform().GetLocation().Z / WorldScale) * TeleportGoalScale)));
+	HMDTransform.SetScale3D(AnimatedScale);
+	HMDMeshComponent->SetWorldTransform(HMDTransform);
+
+	//Calculate the teleported room transform
+	FTransform HeadToWorld = VRMode->GetHeadTransform();
+	HeadToWorld.SetRotation(FQuat::Identity);
+	const FTransform RoomToWorld = VRMode->GetRoomTransform();
+	FTransform RoomToHeadInWorld = RoomToWorld.GetRelativeTransform(HeadToWorld);
+	RoomToHeadInWorld.SetLocation(FVector((RoomToHeadInWorld.GetLocation() / WorldScale) * TeleportGoalScale));
+
+	HMDTransform.SetRotation(FQuat::Identity);
+	FTransform TeleportRoomInWorld = HMDTransform + RoomToHeadInWorld;
+	TeleportGoalLocation = TeleportRoomInWorld.GetLocation();
+
+	// Calculate the teleported motion controllers
+	for (int i = 0; i < 2; i++)
+	{
+		UVREditorInteractor* Interactor = i == 0 ? VRMode->GetHandInteractor(EControllerHand::Left) : VRMode->GetHandInteractor(EControllerHand::Right);
+		UStaticMeshComponent* MotionControllerMeshComponent = i == 0 ? LeftMotionControllerMeshComponent : RightMotionControllerMeshComponent;
+
+		FTransform MotionControllerToWorld = Interactor->GetTransform();
+		FTransform MotionControllerToHeadInWorld = MotionControllerToWorld.GetRelativeTransform(HeadToWorld);
+		MotionControllerToHeadInWorld.SetLocation(FVector((MotionControllerToHeadInWorld.GetLocation() / WorldScale) * TeleportGoalScale));
+		FTransform TeleportedMotionControllerToWorld = HMDTransform + MotionControllerToHeadInWorld;
+		TeleportedMotionControllerToWorld.SetRotation(MotionControllerToWorld.GetRotation());
+		TeleportedMotionControllerToWorld.SetScale3D(AnimatedScale);
+
+		MotionControllerMeshComponent->SetWorldTransform(TeleportedMotionControllerToWorld);
+	}	
+}
+
+void AVREditorTeleporter::Show(const bool bShow)
+{
+	// Only start the fading in or out if it is the first time requesting it. Only reset animation if it wasn't started yet.
+	if (bShow)
+	{
+		if (!bShouldBeVisible.IsSet() || bShouldBeVisible.GetValue() != bShow)
+		{
+			bShouldBeVisible = true;
+			FadeAlpha = 0.0f;
+		}
+	}
+	else
+	{
+		bShouldBeVisible = false;
+		SetVisibility(false);
+	}
+}
+
+void AVREditorTeleporter::UpdateFadingState(const float DeltaTime)
+{
+	if (bShouldBeVisible.IsSet() && bShouldBeVisible.GetValue())
+	{
+		FadeAlpha += VREd::TeleportFadeInAnimateSpeed->GetFloat() * DeltaTime;
+		FadeAlpha = FMath::Clamp(FadeAlpha, 0.0f, 1.0f);
+
+		if (FadeAlpha > 0.0f + KINDA_SMALL_NUMBER)
+		{
+			// At least a little bit visible
+			if (!TeleportDirectionMeshComponent->IsVisible())
+			{
+				SetVisibility(true);
+			}
+		}
+	}
+}
+
+float AVREditorTeleporter::CalculateAnimatedScaleFactor() const
+{
+	const float AnimationOvershootAmount = 0.7f;	// @todo vreditor tweak
+	const float EasedAlpha = UVREditorMode::OvershootEaseOut(FadeAlpha, AnimationOvershootAmount);
+
+	// Animate vertically more than horizontally; just looks a little better
+	const float Scale = FMath::Max(0.001f, EasedAlpha);
+	return Scale * Scale;
+}
+
+float AVREditorTeleporter::GetSlideDelta(UVREditorMotionControllerInteractor* Interactor, const bool Axis)
+{
+	FVector2D SlideDelta = FVector2D::ZeroVector;
+	if (VRMode->GetHMDDeviceType() == EHMDDeviceType::DT_SteamVR)
+	{
+		SlideDelta = FVector2D(Interactor->GetTrackpadSlideDelta(0), Interactor->GetTrackpadSlideDelta(1));
+	}
+	else
+	{
+		SlideDelta = Interactor->GetTrackpadPosition();
+	}
+	SlideDelta.Normalize();
+
+	const bool OtherAxis = !Axis;
+	float Result = 0.0f;
+	if (!FMath::IsNearlyZero(SlideDelta[Axis], VREd::TeleportSlideBuffer->GetFloat()) && 
+		 FMath::IsNearlyZero(SlideDelta[OtherAxis], VREd::TeleportSlideBuffer->GetFloat()))
+	{
+		Result = SlideDelta.X;
+	}
+	return Result;
 }
