Merging //UE4/Private-Geometry to Dev-VREditor (//UE4/Dev-VREditor) 3339870

diff --git a/Engine/Source/Editor/VREditor/Teleporter/VREditorTeleporter.h b/Engine/Source/Editor/VREditor/Teleporter/VREditorTeleporter.h
--- a/Engine/Source/Editor/VREditor/Teleporter/VREditorTeleporter.h
+++ b/Engine/Source/Editor/VREditor/Teleporter/VREditorTeleporter.h
@@ -11,51 +11,86 @@ class UViewportInteractor;
 class UVREditorMode;
 
 /**
-* VR Editor teleport manager
-*/
+ * VR Editor teleport manager and the visual representation of the teleport
+ */
 UCLASS()
-class UVREditorTeleporter : public UObject
+class AVREditorTeleporter: public AActor
 {
-	GENERATED_UCLASS_BODY()
+	GENERATED_BODY()
 
 public:
 
-	/** Initializes the teleport system */
-	void Init( class UVREditorMode* InMode );
+	/** Default constructor that sets up CDO properties */
+	AVREditorTeleporter();
 
-	/** Shuts down the teleport system */
+	/** Deconstructor */
+	virtual ~AVREditorTeleporter();
+
+	/** Initializes the teleporter */
+	void Init(class UVREditorMode* InMode);
+
+	/** Shuts down the teleporter */
 	void Shutdown();
-	
-	/** Returns true if we're actively teleporting */
-	bool IsTeleporting() const
+
+private:
+
+	//~ Begin AActor interface
+	virtual void Tick(const float DeltaTime) override;
+	virtual bool IsEditorOnly() const final
 	{
-		return bIsTeleporting;
+		return true;
 	}
+	//~ End AActor interface
 
-private:
-	/** Ticks the current teleport */
-	void Tick( const float DeltaTime );
-	
 	/** Start teleporting, does a ray trace with the hand passed and calculates the locations for lerp movement in Teleport */
-	void StartTeleport( class UViewportInteractor* Interactor );
+	void StartTeleport(class UViewportInteractor* Interactor);
 
 	/** Called when the user presses a button on their motion controller device */
-	void OnInteractorAction( class FEditorViewportClient& ViewportClient, UViewportInteractor* Interactor,
-		const struct FViewportActionKeyInput& Action, bool& bOutIsInputCaptured, bool& bWasHandled );
+	void OnPreviewInputAction(class FEditorViewportClient& ViewportClient, UViewportInteractor* Interactor,
+		const struct FViewportActionKeyInput& Action, bool& bOutIsInputCaptured, bool& bWasHandled);
+
+	/** Move the roomspace using lerp towards the new location */
+	void Teleport(const float DeltaTime);
 
-	/** Actually teleport to a position */
-	void Teleport( const float DeltaTime );
+	/** Updating aiming with teleport to end of laser and pulling and pushing */
+	void UpdateTeleportAim(const float DeltaTime);
 
-	//
-	// Fields
-	//
+	/** Helper function to push and pull the teleporter along the laser */
+	FVector UpdatePushPullTeleporter(class UVREditorMotionControllerInteractor* VREditorInteractor, const FVector& LaserPointerStart, const FVector& LaserPointerEnd, const bool bEnablePushPull = true);
 
-	/** Owning mode */
+	/** Hide or show the teleporter visuals */
+	void SetVisibility(const bool bVisible);
+
+	/** Sets the color for the teleporter visuals */
+	void SetColor(const FLinearColor& Color);
+
+	/** Figures out the new transforms for all the visuals based on the new location and the transforms of the HMD and motioncontrollers */
+	void UpdateVisuals(const FVector& NewLocation);
+
+	/** If we want to start showing or hiding the meshes */
+	void Show(const bool bShow);
+
+	void UpdateFadingState(const float DeltaTime);
+
+	/** Calculated the scale for animation */
+	float CalculateAnimatedScaleFactor() const;
+
+	/** Get slide delta to push/pull or scale the teleporter */
+	float GetSlideDelta(UVREditorMotionControllerInteractor* Interactor, const bool Axis);
+	
+	/** The owning VR mode */
 	UPROPERTY()
 	UVREditorMode* VRMode;
 
-	/** If we are currently teleporting */
-	bool bIsTeleporting;
+	enum EState
+	{
+		None = 0,				// When not aiming for a teleport
+		Aiming = 1,	// Aiming at the end of the laser, user can also scale using touchpad/analog stick
+		Teleporting = 3			// Currently teleporting from one location to another with lerp
+	};
+
+	/** The current teleport state */
+	EState TeleportingState;
 
 	/** The current lerp of the teleport between the TeleportStartLocation and the TeleportGoalLocation */
 	float TeleportLerpAlpha;
@@ -66,11 +101,54 @@ private:
 	/** The calculated goal location in StartTeleport to move the Roomspace to */
 	FVector TeleportGoalLocation;
 
-	/** So we can teleport on a light press and make sure a second teleport by a fullpress doesn't happen */
-	bool bJustTeleported;
+	/** Visuals for the feet location of the teleporter with the same direction of the HMD yaw */
+	UPROPERTY()
+	UStaticMeshComponent* TeleportDirectionMeshComponent;
+
+	/** Visuals for teleport HMD */
+	UPROPERTY()
+	UStaticMeshComponent* HMDMeshComponent;
+
+	/** Visuals for teleport left motion controller */
+	UPROPERTY()
+	UStaticMeshComponent* LeftMotionControllerMeshComponent;
+
+	/** Visuals for teleport right motion controller */
+	UPROPERTY()
+	UStaticMeshComponent* RightMotionControllerMeshComponent;
+	
+	/** Dynamic material for teleport visuals */
+	UPROPERTY()
+	UMaterialInstanceDynamic* TeleportMID;
 
-	/** Teleport sound */
+	/** The interactor that started aiming to teleport */
 	UPROPERTY()
-	class USoundCue* TeleportSound;
+	UViewportInteractor* InteractorTryingTeleport;
+
+	/** When offset between the hoverlocation of the laser and the calculated  teleport*/
+	FVector OffsetDistance;
+
+	/** The goal world to meters scale. This value is used to scale the visuals and will be used to set world to meters scale after teleporting */
+	float TeleportGoalScale;
+
+	/** The current length of the laser where the teleport should be at */
+	float DragRayLength;
+	
+	/** The current drag velocity to push or pull the teleport along the laser */
+	float DragRayLengthVelocity;
+
+	/** If the teleporter has been pushed by the trackpad */
+	bool bPushedFromEndOfLaser;
+
+	/** Is this the first time aiming for teleporting, we don't have to smooth movement for initial tick */
+	bool bInitialTeleportAim;
+
+	/** Fade alpha, for visibility transitions */
+	float FadeAlpha;
+
+	/** If the teleporter should fade in */
+	TOptional<bool> bShouldBeVisible;
 
-};
+	/** Delay to start the actual moving to the end location */
+	uint32 TeleportTickDelay;
+};
\ No newline at end of file
diff --git a/Engine/Source/Editor/VREditor/UI/VREditorDockableWindow.cpp b/Engine/Source/Editor/VREditor/UI/VREditorDockableWindow.cpp
--- a/Engine/Source/Editor/VREditor/UI/VREditorDockableWindow.cpp
+++ b/Engine/Source/Editor/VREditor/UI/VREditorDockableWindow.cpp
@@ -11,128 +11,83 @@
 #include "VREditorWidgetComponent.h"
 #include "ViewportWorldInteraction.h"
 #include "VREditorInteractor.h"
+#include "VREditorAssetContainer.h"
 
 namespace VREd
 {
 	static FAutoConsoleVariable DockWindowThickness( TEXT( "VREd.DockWindowTickness" ), 1.0f, TEXT( "How thick the window is" ) );
 	static FAutoConsoleVariable DockUISelectionBarVerticalOffset( TEXT( "VREd.DockUISelectionBarVerticalOffset" ), 2.0f, TEXT( "Z Distance between the selectionbar and the UI" ) );
-	static FAutoConsoleVariable DockUICloseButtonOffsetFromCorner( TEXT( "VREd.DockUICloseButtonOffsetFromCorner" ), 1.5f, TEXT( "How far away from the corner (along each 2D axis) of the UI the close button draws" ) );
 	static FAutoConsoleVariable DockUIFadeAnimationDuration( TEXT( "VREd.DockUIFadeAnimationDuration" ), 0.15f, TEXT( "How quick the fade animation should complete in" ) );
-	static FAutoConsoleVariable DockUIHoverScale( TEXT( "VREd.DockUIHoverScale" ), 1.15f, TEXT( "How big the selection bar gets when you hover over it" ) );
+	static FAutoConsoleVariable DockUIHoverScale( TEXT( "VREd.DockUIHoverScale" ), 1.1f, TEXT( "How big the selection bar gets when you hover over it" ) );
 	static FAutoConsoleVariable DockUIHoverAnimationDuration( TEXT( "VREd.DockUIHoverAnimationDuration" ), 0.15f, TEXT( "How quick the hover animation should complete in" ) );
+	static FAutoConsoleVariable DockUIDragSmoothingAmount( TEXT( "VREd.DockUIDragSmoothingAmount" ), 0.85f, TEXT( "How much to smooth out motion when dragging UI (frame rate sensitive)" ) );
 }
 
 
-AVREditorDockableWindow::AVREditorDockableWindow()
-	: Super(),
-	DockSelectDistance( 0.0f )
+AVREditorDockableWindow::AVREditorDockableWindow() : 
+	Super(),
+	bIsLaserAimingTowardUI(false),
+	AimingAtMeFadeAlpha(0.0f),
+	bIsHoveringOverSelectionBar(false),
+	SelectionBarHoverAlpha(0.0f),
+	bIsHoveringOverCloseButton(false),
+	CloseButtonHoverAlpha(0.0f),
+	DockSelectDistance(0.0f)
 {
-	{
-		UStaticMesh* WindowMesh = nullptr;
-		{
-			static ConstructorHelpers::FObjectFinder<UStaticMesh> ObjectFinder( TEXT( "/Engine/VREditor/UI/SM_ContentWindow_01" ) );
-			WindowMesh = ObjectFinder.Object;
-			check( WindowMesh != nullptr );
-		}
+	UVREditorAssetContainer* AssetContainer = LoadObject<UVREditorAssetContainer>(nullptr, *UVREditorMode::AssetContainerPath);
 
-		WindowMeshComponent = CreateDefaultSubobject<UStaticMeshComponent>( TEXT( "WindowMesh" ) );
-		WindowMeshComponent->SetStaticMesh( WindowMesh );
-		WindowMeshComponent->SetMobility( EComponentMobility::Movable );
-		WindowMeshComponent->SetupAttachment( RootComponent );
-
-		WindowMeshComponent->SetCollisionEnabled( ECollisionEnabled::QueryOnly );
-		WindowMeshComponent->SetCollisionResponseToAllChannels( ECollisionResponse::ECR_Ignore );
-		WindowMeshComponent->SetCollisionResponseToChannel( COLLISION_GIZMO, ECollisionResponse::ECR_Block );
-		WindowMeshComponent->SetCollisionObjectType( COLLISION_GIZMO );
-
-		WindowMeshComponent->bGenerateOverlapEvents = false;
-		WindowMeshComponent->SetCanEverAffectNavigation( false );
-		WindowMeshComponent->bCastDynamicShadow = false;
-		WindowMeshComponent->bCastStaticShadow = false;
-		WindowMeshComponent->bAffectDistanceFieldLighting = false;
+	{
+		UStaticMesh* WindowMesh = AssetContainer->WindowMesh;
+		WindowMeshComponent->SetStaticMesh(WindowMesh);
+		check(WindowMeshComponent != nullptr);
 	}
 
+	UMaterialInterface* HoverMaterial = AssetContainer->WindowMaterial;
+	UMaterialInterface* TranslucentHoverMaterial = AssetContainer->WindowMaterial;
+
+	const FRotator RelativeRotation(30.f, 0.f, 0.f);
 	{
-		UStaticMesh* SelectionMesh = nullptr;
-		{
-			static ConstructorHelpers::FObjectFinder<UStaticMesh> ObjectFinder( TEXT( "/Engine/VREditor/UI/SelectionBarMesh" ) );
-			SelectionMesh = ObjectFinder.Object;
-			check( SelectionMesh != nullptr );
-		}
+		UStaticMesh* SelectionMesh = AssetContainer->WindowSelectionBarMesh;
 
 		SelectionBarMeshComponent = CreateDefaultSubobject<UStaticMeshComponent>( TEXT( "SelectionBarMesh" ) );
 		SelectionBarMeshComponent->SetStaticMesh( SelectionMesh );
 		SelectionBarMeshComponent->SetMobility( EComponentMobility::Movable );
 		SelectionBarMeshComponent->SetupAttachment( RootComponent );
 
-		SelectionBarMeshComponent->SetCollisionEnabled( ECollisionEnabled::QueryOnly );
 		SelectionBarMeshComponent->bGenerateOverlapEvents = false;
 		SelectionBarMeshComponent->SetCanEverAffectNavigation( false );
 		SelectionBarMeshComponent->bCastDynamicShadow = false;
 		SelectionBarMeshComponent->bCastStaticShadow = false;
 		SelectionBarMeshComponent->bAffectDistanceFieldLighting = false;
+		SelectionBarMeshComponent->SetRelativeRotation(RelativeRotation);
+
 
-		UMaterial* HoverMaterial = nullptr;
-		{
-			static ConstructorHelpers::FObjectFinder<UMaterial> ObjectFinder( TEXT( "/Engine/VREditor/TransformGizmo/TransformGizmoMaterial" ) );
-			HoverMaterial = ObjectFinder.Object;
-			check( HoverMaterial != nullptr );
-		}
 		SelectionBarMID = UMaterialInstanceDynamic::Create( HoverMaterial, GetTransientPackage() );
 		check( SelectionBarMID != nullptr );
 		SelectionBarMeshComponent->SetMaterial( 0, SelectionBarMID );
-
-		UMaterial* TranslucentHoverMaterial = nullptr;
-		{
-			static ConstructorHelpers::FObjectFinder<UMaterial> ObjectFinder( TEXT( "/Engine/VREditor/TransformGizmo/TranslucentTransformGizmoMaterial" ) );
-			TranslucentHoverMaterial = ObjectFinder.Object;
-			check( TranslucentHoverMaterial != nullptr );
-		}
 		SelectionBarTranslucentMID = UMaterialInstanceDynamic::Create( TranslucentHoverMaterial, GetTransientPackage() );
 		check( SelectionBarTranslucentMID != nullptr );
 		SelectionBarMeshComponent->SetMaterial( 1, SelectionBarTranslucentMID );
 	}
 
 	{
-		UStaticMesh* CloseButtonMesh = nullptr;
-		{
-			static ConstructorHelpers::FObjectFinder<UStaticMesh> ObjectFinder( TEXT( "/Engine/VREditor/UI/CloseButtonMesh" ) );
-			CloseButtonMesh = ObjectFinder.Object;
-			check( CloseButtonMesh != nullptr );
-		}
+		UStaticMesh* CloseButtonMesh = AssetContainer->WindowCloseButtonMesh;
 
 		CloseButtonMeshComponent= CreateDefaultSubobject<UStaticMeshComponent>( TEXT( "CloseButtonMesh" ) );
 		CloseButtonMeshComponent->SetStaticMesh( CloseButtonMesh );
 		CloseButtonMeshComponent->SetMobility( EComponentMobility::Movable );
 		CloseButtonMeshComponent->SetupAttachment( RootComponent );
 
-		CloseButtonMeshComponent->SetCollisionEnabled( ECollisionEnabled::QueryOnly );
-		CloseButtonMeshComponent->SetCollisionResponseToAllChannels( ECollisionResponse::ECR_Ignore );
-		CloseButtonMeshComponent->SetCollisionResponseToChannel( COLLISION_GIZMO, ECollisionResponse::ECR_Block );
-		CloseButtonMeshComponent->SetCollisionObjectType( COLLISION_GIZMO );
-
 		CloseButtonMeshComponent->bGenerateOverlapEvents = false;
 		CloseButtonMeshComponent->SetCanEverAffectNavigation( false );
 		CloseButtonMeshComponent->bCastDynamicShadow = false;
 		CloseButtonMeshComponent->bCastStaticShadow = false;
 		CloseButtonMeshComponent->bAffectDistanceFieldLighting = false;
+		CloseButtonMeshComponent->SetRelativeRotation(RelativeRotation);
 
-		UMaterial* HoverMaterial = nullptr;
-		{
-			static ConstructorHelpers::FObjectFinder<UMaterial> ObjectFinder( TEXT( "/Engine/VREditor/TransformGizmo/TransformGizmoMaterial" ) );
-			HoverMaterial = ObjectFinder.Object;
-			check( HoverMaterial != nullptr );
-		}
 		CloseButtonMID = UMaterialInstanceDynamic::Create( HoverMaterial, GetTransientPackage() );
 		check( CloseButtonMID != nullptr );
 		CloseButtonMeshComponent->SetMaterial( 0, CloseButtonMID );
-
-		UMaterial* TranslucentHoverMaterial = nullptr;
-		{
-			static ConstructorHelpers::FObjectFinder<UMaterial> ObjectFinder( TEXT( "/Engine/VREditor/TransformGizmo/TranslucentTransformGizmoMaterial" ) );
-			TranslucentHoverMaterial = ObjectFinder.Object;
-			check( TranslucentHoverMaterial != nullptr );
-		}
 		CloseButtonTranslucentMID = UMaterialInstanceDynamic::Create( TranslucentHoverMaterial, GetTransientPackage() );
 		check( CloseButtonTranslucentMID != nullptr );
 		CloseButtonMeshComponent->SetMaterial( 1, CloseButtonTranslucentMID );
@@ -151,7 +106,6 @@ AVREditorDockableWindow::AVREditorDockableWindow()
 
 AVREditorDockableWindow::~AVREditorDockableWindow()
 {
-	WindowMeshComponent = nullptr;
 	SelectionBarMeshComponent = nullptr;
 	CloseButtonMeshComponent = nullptr;
 	SelectionBarMID = nullptr;
@@ -170,24 +124,38 @@ void AVREditorDockableWindow::SetupWidgetComponent()
 }
 
 
+void AVREditorDockableWindow::SetCollision(const ECollisionEnabled::Type InCollisionType, const ECollisionResponse InCollisionResponse, const ECollisionChannel InCollisionChannel)
+{
+	AVREditorFloatingUI::SetCollision(InCollisionType, InCollisionResponse, InCollisionChannel);
+
+	SelectionBarMeshComponent->SetCollisionEnabled(InCollisionType);
+	SelectionBarMeshComponent->SetCollisionResponseToAllChannels(InCollisionResponse);
+	SelectionBarMeshComponent->SetCollisionObjectType(InCollisionChannel);
+	CloseButtonMeshComponent->SetCollisionEnabled(InCollisionType);
+	CloseButtonMeshComponent->SetCollisionResponseToAllChannels(InCollisionResponse);
+}
+
+
 void AVREditorDockableWindow::TickManually( float DeltaTime )
 {
 	Super::TickManually( DeltaTime );
 
 	if( WidgetComponent->IsVisible() )
 	{
+		const FVector2D Size = GetSize();
 		const float WorldScaleFactor = GetOwner().GetOwner().GetWorldScaleFactor();
 		const FVector AnimatedScale = CalculateAnimatedScale();
 
-
 		// Update whether the user is aiming toward us or not
 		bIsLaserAimingTowardUI = false;
+
+		if (!GetOwner().IsDraggingDockUI())
 		{
 			const FTransform UICapsuleTransform = this->GetActorTransform();
 
-			const FVector UICapsuleStart = FVector( 0.0f, 0.0f, -GetSize().Y * 0.4f ) * WorldScaleFactor * AnimatedScale;
-			const FVector UICapsuleEnd = FVector( 0.0f, 0.0f, GetSize().Y * 0.4f ) * WorldScaleFactor * AnimatedScale;
-			const float UICapsuleLocalRadius = GetSize().X * 0.5f * WorldScaleFactor * AnimatedScale.Y;
+			const FVector UICapsuleStart = FVector( 0.0f, 0.0f, -Size.Y * 0.4f ) * WorldScaleFactor * AnimatedScale;
+			const FVector UICapsuleEnd = FVector( 0.0f, 0.0f, Size.Y * 0.4f ) * WorldScaleFactor * AnimatedScale;
+			const float UICapsuleLocalRadius = Size.X * 0.5f * WorldScaleFactor * AnimatedScale.Y;
 			const float MinDistanceToUICapsule = 10.0f * WorldScaleFactor * AnimatedScale.Y;	// @todo vreditor tweak
 			const FVector UIForwardVector = FVector::ForwardVector;
 			const float MinDotForAimingAtOtherHand = -1.1f;	// @todo vreditor tweak
@@ -226,8 +194,8 @@ void AVREditorDockableWindow::TickManually( float DeltaTime )
 
 			const FVector WindowMeshScale = FVector(
 				VREd::DockWindowThickness->GetFloat(),
-				GetSize().X / WindowMeshSize,
-				GetSize().Y / WindowMeshSize ) * AnimatedScale * WorldScaleFactor;
+				Size.X / WindowMeshSize,
+				Size.Y / WindowMeshSize ) * AnimatedScale * WorldScaleFactor;
 			WindowMeshComponent->SetRelativeScale3D( WindowMeshScale );			
 		}
 
@@ -244,7 +212,7 @@ void AVREditorDockableWindow::TickManually( float DeltaTime )
 			SelectionBarHoverAlpha = FMath::Clamp( SelectionBarHoverAlpha, 0.0f, 1.0f );
 
 			// How big the selection bar should be
-			const FVector SelectionBarSize( 1.0f, GetSize().X * 0.8f, 6.0f );	// @todo vreditor tweak
+			const FVector SelectionBarSize(20.0f, Size.X * 0.8f, Size.X * 0.1f);
 			FVector SelectionBarScale = SelectionBarSize * AnimatedScale * WorldScaleFactor;
 			SelectionBarScale *= FMath::Lerp( 1.0f, VREd::DockUIHoverScale->GetFloat(), SelectionBarHoverAlpha );
 
@@ -253,9 +221,9 @@ void AVREditorDockableWindow::TickManually( float DeltaTime )
 
 			SelectionBarMeshComponent->SetRelativeScale3D( SelectionBarScale );
 			const FVector SelectionBarRelativeLocation = FVector(
-				0.0f,
-				0.0f,
-				-( GetSize().Y * 0.5f + SelectionBarSize.Z * 0.5f + VREd::DockUISelectionBarVerticalOffset->GetFloat() ) ) * AnimatedScale * WorldScaleFactor;
+				4.0f,
+				((Size.X * 0.5f) - (SelectionBarSize.Y * 0.5f)),
+				-(Size.Y * 0.5f + SelectionBarSize.Z + VREd::DockUISelectionBarVerticalOffset->GetFloat())) * AnimatedScale * WorldScaleFactor;
 			SelectionBarMeshComponent->SetRelativeLocation( SelectionBarRelativeLocation );
 
 			SetSelectionBarColor( GetOwner().GetOwner().GetColor( bIsHoveringOverSelectionBar ? UVREditorMode::EColors::UISelectionBarHoverColor : UVREditorMode::EColors::UISelectionBarColor ) );
@@ -274,20 +242,22 @@ void AVREditorDockableWindow::TickManually( float DeltaTime )
 			CloseButtonHoverAlpha = FMath::Clamp( CloseButtonHoverAlpha, 0.0f, 1.0f );
 
 			// How big the close button should be
-			const FVector CloseButtonSize( 1.0f, 2.0f, 2.0f );	// @todo vreditor tweak
+			const FVector CloseButtonSize(20.0f, Size.X * 0.1f, Size.X * 0.1f);
 			FVector CloseButtonScale = CloseButtonSize * AnimatedScale * WorldScaleFactor * EasedAimingAtMeFadeAlpha;
 			CloseButtonScale *= FMath::Lerp( 1.0f, VREd::DockUIHoverScale->GetFloat(), CloseButtonHoverAlpha );
 
-
 			CloseButtonMeshComponent->SetRelativeScale3D( CloseButtonScale );
 			const FVector CloseButtonRelativeLocation = FVector(
-				0.0f,
-				-( GetSize().X * 0.5f + CloseButtonSize.Y * 0.5f + VREd::DockUICloseButtonOffsetFromCorner->GetFloat() ),
-				( GetSize().Y * 0.5f + CloseButtonSize.Z * 0.5f + VREd::DockUICloseButtonOffsetFromCorner->GetFloat() ) ) * AnimatedScale * WorldScaleFactor;
-			CloseButtonMeshComponent->SetRelativeLocation( CloseButtonRelativeLocation );
+				4.0f,
+				-((Size.X * 0.5f) - (CloseButtonSize.Y * 0.5f)),
+				-(Size.Y * 0.5f + CloseButtonSize.Z + VREd::DockUISelectionBarVerticalOffset->GetFloat())) * AnimatedScale * WorldScaleFactor;
+			CloseButtonMeshComponent->SetRelativeLocation(CloseButtonRelativeLocation);
 
 			SetCloseButtonColor( GetOwner().GetOwner().GetColor( bIsHoveringOverCloseButton ? UVREditorMode::EColors::UICloseButtonHoverColor : UVREditorMode::EColors::UICloseButtonColor ) );
 		}
+
+		const FVector NewScale(GetScale() * AnimatedScale * VRMode->GetWorldScaleFactor());
+		SetWidgetComponentScale(NewScale);
 	}
 } 
 
@@ -321,9 +291,15 @@ float AVREditorDockableWindow::GetDockSelectDistance() const
 	return DockSelectDistance;
 }
 
+void AVREditorDockableWindow::SetDockSelectDistance(const float InDockDistance)
+{
+	DockSelectDistance = InDockDistance;
+}
+
 void AVREditorDockableWindow::OnPressed( UViewportInteractor* Interactor, const FHitResult& InHitResult, bool& bOutResultedInDrag )
 {
 	bOutResultedInDrag = false;
+
 	UVREditorInteractor* VRInteractor = Cast<UVREditorInteractor>( Interactor );
 	if( VRInteractor != nullptr )
 	{
@@ -331,17 +307,14 @@ void AVREditorDockableWindow::OnPressed( UViewportInteractor* Interactor, const
 		{
 			// Close the window
 			const bool bShouldShow = false;
-			const bool bShowOnHand = false;
-			const bool bRefreshQuickMenu = true;
-			GetOwner().ShowEditorUIPanel( this, VRInteractor, bShouldShow, bShowOnHand, bRefreshQuickMenu );
+			const bool bSpawnInFront = false;
+			GetOwner().ShowEditorUIPanel(this, VRInteractor, bShouldShow, bSpawnInFront);
 		}
-		else if( InHitResult.Component == GetSelectionBarMeshComponent() )
+		else if(InHitResult.Component == GetSelectionBarMeshComponent() && !GetOwner().IsDraggingPanelFromOpen())
 		{
 			bOutResultedInDrag = true;
-
-			FVector LaserPointerStart, LaserPointerEnd;
-			DockSelectDistance = ( InHitResult.TraceStart - InHitResult.Location ).Size();
-			GetOwner().StartDraggingDockUI( this, VRInteractor, DockSelectDistance );
+			SetDockSelectDistance((InHitResult.TraceStart - InHitResult.Location ).Size());
+			GetOwner().StartDraggingDockUI(this, VRInteractor, DockSelectDistance);
 		}
 	}
 }
@@ -369,10 +342,10 @@ void AVREditorDockableWindow::OnHoverLeave( UViewportInteractor* Interactor, con
 	UActorComponent* OtherInteractorHoveredComponent = nullptr;
 	if( Interactor->GetOtherInteractor() != nullptr )
 	{
-		OtherInteractorHoveredComponent = Interactor->GetOtherInteractor()->GetHoverComponent();
+		OtherInteractorHoveredComponent = Interactor->GetOtherInteractor()->GetLastHoverComponent();
 	}
 
-	if ( OtherInteractorHoveredComponent != SelectionBarMeshComponent && NewComponent != SelectionBarMeshComponent )
+	if ( OtherInteractorHoveredComponent != SelectionBarMeshComponent && NewComponent != SelectionBarMeshComponent && !GetDragOperationComponent()->IsDragging() )
 	{
 		bIsHoveringOverSelectionBar = false;
 	}
@@ -423,23 +396,41 @@ void UDockableWindowDragOperation::ExecuteDrag( UViewportInteractor* Interactor,
 	{
 		UVREditorUISystem& UISystem = DockableWindow->GetOwner();
 
-		const bool bIsAbsolute = ( UISystem.GetOwner().GetHMDDeviceType() == EHMDDeviceType::DT_SteamVR );
-		float SlideDeltaY = VRInteractor->GetSlideDelta();
-
-		float NewUIScale = DockableWindow->GetScale() + SlideDeltaY;
-		if ( NewUIScale <= UISystem.GetMinDockWindowSize() )
+		if (UISystem.CanScalePanel())
 		{
-			NewUIScale = UISystem.GetMinDockWindowSize();
+			float NewUIScale = DockableWindow->GetScale() + VRInteractor->GetSlideDelta();
+			if (NewUIScale <= UISystem.GetMinDockWindowSize())
+			{
+				NewUIScale = UISystem.GetMinDockWindowSize();
+			}
+			else if (NewUIScale >= UISystem.GetMaxDockWindowSize())
+			{
+				NewUIScale = UISystem.GetMaxDockWindowSize();
+			}
+			DockableWindow->SetScale(NewUIScale);
 		}
-		else if ( NewUIScale >= UISystem.GetMaxDockWindowSize() )
+
+		const FTransform UIToWorld = UISystem.MakeDockableUITransform( DockableWindow, VRInteractor, DockableWindow->GetDockSelectDistance() );
+		FTransform SmoothedUIToWorld = UIToWorld;
+		if( LastUIToWorld.IsSet() )
 		{
-			NewUIScale = UISystem.GetMaxDockWindowSize();
+			SmoothedUIToWorld.Blend( UIToWorld, LastUIToWorld.GetValue(), VREd::DockUIDragSmoothingAmount->GetFloat() );
 		}
 
-		DockableWindow->SetScale( NewUIScale );
+		// Update interactor hover location while dragging the interactor
+		const FTransform LaserImpactToWorld = UISystem.MakeDockableUITransformOnLaser( DockableWindow, VRInteractor, DockableWindow->GetDockSelectDistance() );
+		FTransform SmoothedLaserImpactToWorld = LaserImpactToWorld;
+		if( LastLaserImpactToWorld.IsSet() )
+		{
+			SmoothedLaserImpactToWorld.Blend( LaserImpactToWorld, LastLaserImpactToWorld.GetValue(), VREd::DockUIDragSmoothingAmount->GetFloat() );
+		}
 
-		const FTransform DockedUIToWorld = UISystem.MakeDockableUITransform( DockableWindow, VRInteractor, DockableWindow->GetDockSelectDistance() );
-		DockableWindow->SetActorTransform( DockedUIToWorld );
+		DockableWindow->SetActorTransform( SmoothedUIToWorld );
 		DockableWindow->UpdateRelativeRoomTransform();
+
+		Interactor->SetHoverLocation(SmoothedLaserImpactToWorld.GetLocation());
+
+		LastUIToWorld = SmoothedUIToWorld;
+		LastLaserImpactToWorld = SmoothedLaserImpactToWorld;
 	}
 }
diff --git a/Engine/Source/Editor/VREditor/UI/VREditorDockableWindow.h b/Engine/Source/Editor/VREditor/UI/VREditorDockableWindow.h
--- a/Engine/Source/Editor/VREditor/UI/VREditorDockableWindow.h
+++ b/Engine/Source/Editor/VREditor/UI/VREditorDockableWindow.h
@@ -40,9 +40,12 @@ public:
 	/** Gets the selection bar component */
 	UStaticMeshComponent* GetSelectionBarMeshComponent() const;
 
-	/** Gets the distance between the interactor and the window when sta	rting drag */
+	/** Gets the distance between the interactor and the window when starting drag */
 	float GetDockSelectDistance() const;
 
+	/** Set the distance between the interactor and the window when starting drag */
+	void SetDockSelectDistance(const float InDockDistance);
+
 	// IViewportInteractableInterface
 	virtual void OnPressed( UViewportInteractor* Interactor, const FHitResult& InHitResult, bool& bOutResultedInDrag ) override;
 	virtual void OnHover( UViewportInteractor* Interactor ) override;
@@ -59,6 +62,7 @@ protected:
 
 	// AVREditorFloatingUI overrides
 	virtual void SetupWidgetComponent() override;
+	virtual void SetCollision(const ECollisionEnabled::Type InCollisionType, const ECollisionResponse InCollisionResponse, const ECollisionChannel InCollisionChannel) override;
 
 private:
 
@@ -68,10 +72,6 @@ private:
 	/** Set the color on the dynamic materials of the close button */
 	void SetCloseButtonColor( const FLinearColor& LinearColor );
 
-	/** The dockable window mesh */
-	UPROPERTY()
-	class UStaticMeshComponent* WindowMeshComponent;
-
 	/** Mesh underneath the window for easy selecting and dragging */
 	UPROPERTY()
 	class UStaticMeshComponent* SelectionBarMeshComponent;
@@ -133,4 +133,10 @@ public:
 
 	// IViewportDragOperation
 	virtual void ExecuteDrag( class UViewportInteractor* Interactor, IViewportInteractableInterface* Interactable ) override;
+
+	/** Last frame's UIToWorld transform */
+	TOptional<FTransform> LastUIToWorld;
+
+	/** last frame's LaserImpactToWorld transform */
+	TOptional<FTransform> LastLaserImpactToWorld;
 };
diff --git a/Engine/Source/Editor/VREditor/UI/VREditorFloatingUI.cpp b/Engine/Source/Editor/VREditor/UI/VREditorFloatingUI.cpp
--- a/Engine/Source/Editor/VREditor/UI/VREditorFloatingUI.cpp
+++ b/Engine/Source/Editor/VREditor/UI/VREditorFloatingUI.cpp
@@ -6,12 +6,12 @@
 #include "VREditorMode.h"
 #include "Components/WidgetComponent.h"
 #include "VREditorWidgetComponent.h"
-
+#include "Components/StaticMeshComponent.h"
+#include "VRModeSettings.h"
 
 namespace VREd
 {
 	static FAutoConsoleVariable UIFadeSpeed( TEXT( "VREd.UIFadeSpeed" ), 6.0f, TEXT( "How fast UI should fade in and out" ) );
-	static FAutoConsoleVariable UIBrightness( TEXT( "VREd.UIBrightness" ), 0.3f, TEXT( "How bright the UI should be" ) );
 }
 
 
@@ -25,7 +25,6 @@ AVREditorFloatingUI::AVREditorFloatingUI()
 	UserWidgetClass( nullptr ),
 	bShouldBeVisible(),
 	FadeAlpha( 1.0f ),
-	bCollisionOnShowUI( true ),
 	FadeDelay( 0.0f ),
 	InitialScale( 1.0f )
 {
@@ -37,8 +36,21 @@ AVREditorFloatingUI::AVREditorFloatingUI()
 	WidgetComponent = CreateDefaultSubobject<UVREditorWidgetComponent>( TEXT( "WidgetComponent" ), bTransient );
 	WidgetComponent->SetEditTimeUsable(true);
 	WidgetComponent->SetupAttachment( SceneComponent );
-
+	WidgetComponent->PrimaryComponentTick.bTickEvenWhenPaused = true;
 	InitialScale = Scale;
+
+	{
+		WindowMeshComponent = CreateDefaultSubobject<UStaticMeshComponent>(TEXT("WindowMesh"));
+		WindowMeshComponent->SetMobility(EComponentMobility::Movable);
+		WindowMeshComponent->SetupAttachment(RootComponent);
+		
+		WindowMeshComponent->bGenerateOverlapEvents = false;
+		WindowMeshComponent->SetCanEverAffectNavigation(false);
+		WindowMeshComponent->bCastDynamicShadow = false;
+		WindowMeshComponent->bCastStaticShadow = false;
+		WindowMeshComponent->bAffectDistanceFieldLighting = false;
+		WindowMeshComponent->bSelectable = false;
+	}
 }
 
 
@@ -86,13 +98,24 @@ void AVREditorFloatingUI::SetupWidgetComponent()
 	//WidgetComponent->SetMaxInteractionDistance( 10000.0f );
 
 	// Default to visible
-	ShowUI( true );
+	ShowUI( false );
 
 	// Set initial opacity
 	UpdateFadingState( 0.0f );
 
 	// Set initial transform
 	UpdateTransformIfDocked();
+
+	// Update the window border mesh
+	{
+		const float WindowMeshSize = 100.0f;	// Size of imported mesh, we need to inverse compensate for
+
+		const FVector WindowMeshScale = FVector(
+			1.0f,
+			GetSize().X / WindowMeshSize,
+			GetSize().Y / WindowMeshSize) * GetOwner().GetOwner().GetWorldScaleFactor();
+		WindowMeshComponent->SetRelativeScale3D(WindowMeshScale);
+	}
 }
 
 void AVREditorFloatingUI::SetSlateWidget( UVREditorUISystem& InitOwner, const TSharedRef<SWidget>& InitSlateWidget, const FIntPoint InitResolution, const float InitScale, const EDockedTo InitDockedTo )
@@ -158,16 +181,17 @@ void AVREditorFloatingUI::Destroyed()
 
 void AVREditorFloatingUI::SetTransform( const FTransform& Transform )
 {
-	const FVector AnimatedScale = CalculateAnimatedScale();
-
-	FTransform AnimatedTransform = Transform;
-	AnimatedTransform.SetScale3D( AnimatedTransform.GetScale3D() * AnimatedScale );
+	if (!bHidden)
+	{
+		const FVector AnimatedScale = CalculateAnimatedScale();
+		FTransform AnimatedTransform = Transform;
+		AnimatedTransform.SetScale3D(AnimatedTransform.GetScale3D() * AnimatedScale);
 
-	const float Aspect = ( float ) Resolution.X / ( float ) Resolution.Y;
+		RootComponent->SetWorldLocation(AnimatedTransform.GetLocation());
+		RootComponent->SetWorldRotation(AnimatedTransform.GetRotation());
 
-	RootComponent->SetWorldLocation( AnimatedTransform.GetLocation() );
-	RootComponent->SetWorldRotation( AnimatedTransform.GetRotation() );
-	WidgetComponent->SetWorldScale3D( FVector( 1.0f / AnimatedTransform.GetScale3D().X, 1.0f / ( float ) Resolution.X, 1.0f / ( float ) Resolution.Y / Aspect ) * AnimatedTransform.GetScale3D() );
+		SetWidgetComponentScale(AnimatedTransform.GetScale3D());
+	}
 }
 
 
@@ -216,16 +240,17 @@ void AVREditorFloatingUI::UpdateFadingState( const float DeltaTime )
 		}
 
  		// Set material color
-		// NOTE: We intentionally make the UI quite dark here, so it doesn't bloom out in HDR
-		const float UIBrightness = FadeAlpha * VREd::UIBrightness->GetFloat();
+		const float UIBrightness = FadeAlpha * GetDefault<UVRModeSettings>()->UIBrightness;
 		WidgetComponent->SetTintColorAndOpacity( FLinearColor( UIBrightness, UIBrightness, UIBrightness ).CopyWithNewOpacity( FadeAlpha ) );
+
 	}
 }
 
 FVector AVREditorFloatingUI::CalculateAnimatedScale() const
 {
 	const float AnimationOvershootAmount = 0.7f;	// @todo vreditor tweak
-	const float EasedAlpha = UVREditorMode::OvershootEaseOut( FadeAlpha, AnimationOvershootAmount );
+	float EasedAlpha = UVREditorMode::OvershootEaseOut(FadeAlpha, AnimationOvershootAmount);
+	EasedAlpha = FMath::Clamp(EasedAlpha, 0.01f, 1.0f + AnimationOvershootAmount);
 
 	// Animate vertically more than horizontally; just looks a little better
 	const float ZScale = FMath::Max( 0.001f, EasedAlpha );
@@ -238,9 +263,23 @@ FVector AVREditorFloatingUI::CalculateAnimatedScale() const
 	return AnimatedScale;
 }
 
-void AVREditorFloatingUI::SetCollisionOnShow( const bool bInCollisionOnShow )
+void AVREditorFloatingUI::SetCollision(const ECollisionEnabled::Type InCollisionType, const ECollisionResponse InCollisionResponse, const ECollisionChannel InCollisionChannel)
+{
+	WidgetComponent->SetCollisionEnabled(InCollisionType);
+	WidgetComponent->SetCollisionResponseToAllChannels(InCollisionResponse);
+	WidgetComponent->SetCollisionObjectType(InCollisionChannel);
+
+	if (WindowMeshComponent)
+	{
+		WindowMeshComponent->SetCollisionEnabled(InCollisionType);
+		WindowMeshComponent->SetCollisionResponseToAllChannels(InCollisionResponse);
+		WindowMeshComponent->SetCollisionObjectType(InCollisionChannel);
+	}
+}
+
+UVREditorBaseUserWidget* AVREditorFloatingUI::GetUserWidget()
 {
-	bCollisionOnShowUI = bInCollisionOnShow;
+	return UserWidget;
 }
 
 float AVREditorFloatingUI::GetInitialScale() const
@@ -248,18 +287,12 @@ float AVREditorFloatingUI::GetInitialScale() const
 	return InitialScale;
 }
 
-
 void AVREditorFloatingUI::ShowUI( const bool bShow, const bool bAllowFading, const float InitFadeDelay )
 {
 	if( !bShouldBeVisible.IsSet() || bShow != bShouldBeVisible.GetValue() )
 	{
 		bShouldBeVisible = bShow;
 
-		if( bCollisionOnShowUI )
-		{
-			SetActorEnableCollision( bShow );
-		}
-
 		if( !bAllowFading )
 		{
 			SetActorHiddenInGame( !bShow );
@@ -267,8 +300,19 @@ void AVREditorFloatingUI::ShowUI( const bool bShow, const bool bAllowFading, con
 			FadeAlpha = bShow ? 1.0f : 0.0f;
 		}
 
+		// Set collision on components
+		if (bShow)
+		{
+			SetCollision(ECollisionEnabled::QueryOnly, ECollisionResponse::ECR_Block, ECollisionChannel::ECC_WorldStatic);
+		}
+		else
+		{
+			SetCollision(ECollisionEnabled::NoCollision, ECollisionResponse::ECR_Ignore, ECollisionChannel::ECC_Visibility);
+		}
+
+
 		FadeDelay = InitFadeDelay;
-	}
+	}		
 }
 
 
@@ -283,12 +327,20 @@ float AVREditorFloatingUI::GetScale() const
 	return Scale;
 }
 
-void AVREditorFloatingUI::SetScale( const float NewSize )
+void AVREditorFloatingUI::SetScale( const float NewSize, const bool bScaleWidget /*= true*/)
 {
 	Scale = NewSize;
 
-	const float WorldScaleFactor = Owner->GetOwner().GetWorldScaleFactor();
-	const FVector NewScale( Scale * WorldScaleFactor );
-	const float Aspect = (float)Resolution.X / (float)Resolution.Y;
-	WidgetComponent->SetWorldScale3D( FVector( 1.0f / NewScale.X, 1.0f / (float)Resolution.X, 1.0f / (float)Resolution.Y / Aspect ) * NewScale );
+	if (bScaleWidget)
+	{
+		const float WorldScaleFactor = Owner->GetOwner().GetWorldScaleFactor();
+		const FVector NewScale( Scale * WorldScaleFactor );
+		SetWidgetComponentScale(NewScale);
+	}
 }
+
+void AVREditorFloatingUI::SetWidgetComponentScale(const FVector& InScale)
+{
+	const float Aspect = (float)Resolution.X / (float)Resolution.Y;
+	WidgetComponent->SetWorldScale3D(FVector(1.0f / InScale.X, 1.0f / (float)Resolution.X, 1.0f / (float)Resolution.Y / Aspect) * InScale);
+}
\ No newline at end of file
diff --git a/Engine/Source/Editor/VREditor/UI/VREditorFloatingUI.h b/Engine/Source/Editor/VREditor/UI/VREditorFloatingUI.h
--- a/Engine/Source/Editor/VREditor/UI/VREditorFloatingUI.h
+++ b/Engine/Source/Editor/VREditor/UI/VREditorFloatingUI.h
@@ -45,6 +45,12 @@ public:
 		return WidgetComponent;
 	}
 
+	/** Returns the mesh component for this UI, or nullptr if not spawned right now */
+	class UStaticMeshComponent* GetMeshComponent()
+	{
+		return WindowMeshComponent;
+	}
+
 	/** Returns the actual size of the UI in either axis, after scaling has been applied.  Does not take into animation or world scaling */
 	FVector2D GetSize() const;
 
@@ -52,7 +58,10 @@ public:
 	float GetScale() const;
 
 	/** Sets a new size for the UI */
-	void SetScale( const float NewSize );
+	void SetScale(const float NewSize, const bool bScaleWidget = true);
+
+	/** Set the widget scale */
+	void SetWidgetComponentScale(const FVector& InScale);
 
 	/** Sets the UI transform */
 	virtual void SetTransform( const FTransform& Transform ) override;
@@ -86,20 +95,22 @@ public:
 		return CastChecked<T>( UserWidget );
 	}
 
-	/** Sets if the collision is on when showing the UI*/
-	void SetCollisionOnShow( const bool bInCollisionOnShow );
+	/** Gets the current user widget of this floating UI, return nullptr if using slate widget */
+	UVREditorBaseUserWidget* GetUserWidget();
 
 	/** Gets the initial size of this UI */
 	float GetInitialScale() const;
 
+	/** Called to finish setting everything up, after a widget has been assigned */
+	virtual void SetupWidgetComponent();
 
 protected:
 
 	/** Returns a scale to use for this widget that takes into account animation */
 	FVector CalculateAnimatedScale() const;
 
-	/** Called to finish setting everything up, after a widget has been assigned */
-	virtual void SetupWidgetComponent();
+	/** Set collision on components */
+	virtual void SetCollision(const ECollisionEnabled::Type InCollisionType, const ECollisionResponse InCollisionResponse, const ECollisionChannel InCollisionChannel);
 
 private:
 
@@ -118,6 +129,10 @@ protected:
 	UPROPERTY()
 	class UVREditorWidgetComponent* WidgetComponent;
 
+	/** The floating window mesh */
+	UPROPERTY()
+	class UStaticMeshComponent* WindowMeshComponent;
+
 	/** Resolution we should draw this UI at, regardless of scale */
 	FIntPoint Resolution;
 
@@ -137,9 +152,6 @@ private:
 	/** Fade alpha, for visibility transitions */
 	float FadeAlpha;
 
-	/** Collision on when showing UI */
-	bool bCollisionOnShowUI;
-
 	/** Delay to fading in or out. Set on ShowUI and cleared on finish of fade in/out */
 	float FadeDelay;
 
