Merging //UE4/Private-Geometry to Dev-VREditor (//UE4/Dev-VREditor) 3339870

diff --git a/Engine/Source/Editor/VREditor/UI/VREditorUISystem.cpp b/Engine/Source/Editor/VREditor/UI/VREditorUISystem.cpp
--- a/Engine/Source/Editor/VREditor/UI/VREditorUISystem.cpp
+++ b/Engine/Source/Editor/VREditor/UI/VREditorUISystem.cpp
@@ -22,11 +22,11 @@
 #include "Framework/MultiBox/MultiBoxBuilder.h"
 #include "Framework/Application/SlateApplication.h"
 #include "Framework/Docking/TabManager.h"
+#include "Framework/Commands/UICommandInfo.h"
 #include "Layout/WidgetPath.h"
 
 #include "Engine/EngineTypes.h"
 #include "Editor/UnrealEdEngine.h"
-#include "Kismet/GameplayStatics.h"
 #include "GameFramework/WorldSettings.h"
 #include "Sound/SoundCue.h"
 #include "EditorModeManager.h"
@@ -37,6 +37,7 @@
 #include "VREditorBaseActor.h"
 #include "VREditorBaseUserWidget.h"
 #include "VREditorFloatingUI.h"
+#include "VREditorRadialFloatingUI.h"
 #include "VREditorDockableWindow.h"
 #include "ViewportInteractionTypes.h"
 #include "IHeadMountedDisplay.h"
@@ -44,7 +45,9 @@
 #include "VREditorInteractor.h"
 #include "VREditorMotionControllerInteractor.h"
 #include "IVREditorModule.h"
+#include "VREditorAssetContainer.h"
 #include "Settings/EditorExperimentalSettings.h"
+#include "DebuggerCommands.h"
 
 // UI
 #include "Components/WidgetComponent.h"
@@ -82,13 +85,19 @@
 #include "Widgets/Layout/SDPIScaler.h"
 #include "Widgets/Docking/SDockTab.h"
 #include "Widgets/Colors/SColorPicker.h"
+#include "Widgets/Input/SCheckBox.h"
 #include "MultiBox.h"
 #include "LevelEditorActions.h"
 #include "VREditorActions.h"
 #include "GenericCommands.h"
-#include "SVRRadialPanel.h"
 
 #include "ISequencer.h"
+#include "Engine/StaticMesh.h"
+#include "EdMode.h"
+#include "SImage.h"
+
+#include "DrawDebugHelpers.h"
+#include "VRModeSettings.h"
 
 #define LOCTEXT_NAMESPACE "VREditor"
 
@@ -100,46 +109,30 @@ namespace VREd
 	static FAutoConsoleVariable SequencerUIResolutionY(TEXT("VREd.SequencerUIResolutionY"), 600, TEXT("Vertical resolution to use for Sequencer UI render targets"));
 	static FAutoConsoleVariable DefaultEditorUIResolutionX( TEXT( "VREd.DefaultEditorUIResolutionX" ), 1024, TEXT( "Horizontal resolution to use for VR editor UI render targets" ) );
 	static FAutoConsoleVariable DefaultEditorUIResolutionY( TEXT( "VREd.DefaultEditorUIResolutionY" ), 1024, TEXT( "Vertical resolution to use for VR editor UI render targets" ) );
-	static FAutoConsoleVariable QuickMenuUIResolutionX( TEXT( "VREd.QuickMenuUIResolutionX" ), 1440, TEXT( "Horizontal resolution to use for Quick Menu VR UI render targets" ) );
-	static FAutoConsoleVariable QuickMenuUIResolutionY( TEXT( "VREd.QuickMenuUIResolutionY" ), 1440, TEXT( "Vertical resolution to use for Quick Menu VR UI render targets" ) );
-	static FAutoConsoleVariable ContentBrowserUISize( TEXT( "VREd.ContentBrowserUISize" ), 60.0f, TEXT( "How big content browser UIs should be" ) );
-	static FAutoConsoleVariable EditorUISize( TEXT( "VREd.EditorUISize" ), 50.0f, TEXT( "How big editor UIs should be" ) );
+	static FAutoConsoleVariable DefaultRadialElementResolutionX(TEXT("VREd.DefaultRadialElementResolutionX"), 350, TEXT("Horizontal resolution to use for VR editor radial UI render targets"));
+	static FAutoConsoleVariable DefaultRadialElementResolutionY(TEXT("VREd.DefaultRadialElementResolutionY"), 350, TEXT("Vertical resolution to use for VR editor radial UI render targets"));
+	static FAutoConsoleVariable QuickMenuUIResolutionX( TEXT( "VREd.QuickMenuUIResolutionX" ), 1200, TEXT( "Horizontal resolution to use for Quick Menu VR UI render targets" ) );
+	static FAutoConsoleVariable QuickMenuUIResolutionY( TEXT( "VREd.QuickMenuUIResolutionY" ), 1075, TEXT( "Vertical resolution to use for Quick Menu VR UI render targets" ) );
+	static FAutoConsoleVariable ContentBrowserUISize( TEXT( "VREd.ContentBrowserUISize" ), 80.0f, TEXT( "How big content browser UIs should be" ) );
+	static FAutoConsoleVariable EditorUISize( TEXT( "VREd.EditorUISize" ), 70.0f, TEXT( "How big editor UIs should be" ) );
 	static FAutoConsoleVariable ContentBrowserUIScale( TEXT( "VREd.ContentBrowserUIScale" ), 2.0f, TEXT( "How much to scale up (or down) the content browser for VR" ) );
 	static FAutoConsoleVariable EditorUIScale( TEXT( "VREd.EditorUIScale" ), 2.0f, TEXT( "How much to scale up (or down) editor UIs for VR" ) );
-	static FAutoConsoleVariable EditorUIOffsetFromHand( TEXT( "VREd.EditorUIOffsetFromHand" ), 12.0f, TEXT( "How far to offset editor UIs from the origin of the hand mesh" ) );
 	static FAutoConsoleVariable AssetEditorUIResolutionX( TEXT( "VREd.AssetEditorUIResolutionX" ), 1920, TEXT( "Horizontal resolution to use for VR editor asset editor UI render targets" ) );
 	static FAutoConsoleVariable AssetEditorUIResolutionY( TEXT( "VREd.AssetEditorUIResolutionY" ), 1080, TEXT( "Vertical resolution to use for VR editor asset editor UI render targets" ) );
 	static FAutoConsoleVariable RadialMenuFadeDelay( TEXT( "VREd.RadialMenuFadeDelay" ), 0.2f, TEXT( "The delay for the radial menu after selecting a button" ) );
 	static FAutoConsoleVariable UIAbsoluteScrollSpeed( TEXT( "VREd.UIAbsoluteScrollSpeed" ), 8.0f, TEXT( "How fast the UI scrolls when dragging the touchpad" ) );
 	static FAutoConsoleVariable UIRelativeScrollSpeed( TEXT( "VREd.UIRelativeScrollSpeed" ), 0.75f, TEXT( "How fast the UI scrolls when holding an analog stick" ) );
 	static FAutoConsoleVariable MinUIScrollDeltaForInertia( TEXT( "VREd.MinUIScrollDeltaForInertia" ), 0.25f, TEXT( "Minimum amount of touch pad input before inertial UI scrolling kicks in" ) );
-	static FAutoConsoleVariable MinDockDragDistance( TEXT( "VREd.MinDockDragDistance" ), 10.0f, TEXT( "Minimum amount of distance needed to drag dock from hand" ) );
-	static FAutoConsoleVariable DoubleClickTime( TEXT( "VREd.DoubleClickTime" ), 0.25f, TEXT( "Minimum duration between clicks for a double click event to fire" ) );
 	static FAutoConsoleVariable UIPressHapticFeedbackStrength( TEXT( "VREd.UIPressHapticFeedbackStrength" ), 0.4f, TEXT( "Strenth of haptic when clicking on the UI" ) );
 	static FAutoConsoleVariable UIAssetEditorSummonedOnHandHapticFeedbackStrength( TEXT( "VREd.UIAssetEditorSummonedOnHandHapticFeedbackStrength" ), 1.0f, TEXT( "Strenth of haptic to play to remind a user which hand an asset editor was spawned on" ) );
 	static FAutoConsoleVariable MaxDockWindowSize( TEXT( "VREd.MaxDockWindowSize" ), 250, TEXT( "Maximum size for dockable windows" ) );
 	static FAutoConsoleVariable MinDockWindowSize( TEXT( "VREd.MinDockWindowSize" ), 40, TEXT( "Minimum size for dockable windows" ) );
-	static FAutoConsoleVariable DockSnapTime( TEXT( "VREd.DockSnapTime" ), 0.1f, TEXT( "Lerp time for snapping dock to hand and back" ) );
-	static FAutoConsoleVariable DockDragSpeed( TEXT( "VREd.DockDragSpeed" ), 0.05f, TEXT( "Dock dragging time delay" ) );
+	static FAutoConsoleVariable UIPanelOpenDistance( TEXT( "VREd.UIPanelOpenDistance" ), 20.f, TEXT( "Distance to spawn a panel from the hand in centimeters" ) );
+	static FAutoConsoleVariable UIPanelOpenRotationPitchOffset( TEXT( "VREd.UIPanelOpenRotationPitchOffset" ), 45.0f, TEXT( "The pitch rotation offset in degrees when spawning a panel in front of the motioncontroller" ) );
 	static FAutoConsoleVariable ColorPickerDockSpawnOffset( TEXT( "VREd.ColorPickerDockSpawnOffset" ),  3.f, TEXT( "Offset of where the color picker spawns" ) );
-
-	// Tutorial UI commands
-	static FAutoConsoleVariable TutorialUIResolutionX( TEXT( "VREd.TutorialUI.Resolution.X" ), 1920, TEXT( "The X resolution for the tutorial UI panel" ) );
-	static FAutoConsoleVariable TutorialUIResolutionY( TEXT( "VREd.TutorialUI.Resolution.Y" ), 1080, TEXT( "The Y resolution for the tutorial UI panel" ) );
-	static FAutoConsoleVariable TutorialUISize( TEXT( "VREd.TutorialUI.Size" ), 200, TEXT( "The room space size for the tutorial UI panel" ) );
-	static FAutoConsoleVariable TutorialUIYaw( TEXT( "VREd.TutorialUI.Yaw" ), 180, TEXT( "The yaw rotation for the tutorial UI panel" ) );
-	static FAutoConsoleVariable TutorialUIPitch( TEXT( "VREd.TutorialUI.Pitch" ), 45, TEXT( "The pitch rotation for the tutorial UI panel" ) );
-	static FAutoConsoleVariable TutorialUILocationX( TEXT( "VREd.TutorialUI.Location.X" ), 200, TEXT( "The X location for the tutorial UI panel" ) );
-	static FAutoConsoleVariable TutorialUILocationY( TEXT( "VREd.TutorialUI.Location.Y" ), 0, TEXT( "The Y location for the tutorial UI panel" ) );
-	static FAutoConsoleVariable TutorialUILocationZ( TEXT( "VREd.TutorialUI.Location.Z" ), 40, TEXT( "The Z location for the tutorial UI panel" ) );
-
-	static const FTransform DefaultModesTransform(			FRotator( -30,  130, 0 ), FVector( 40, -45,  40 ), FVector( 1 ) );	// Top left
-	static const FTransform DefaultWorldOutlinerTransform(	FRotator( -30, -130, 0 ), FVector( 40,  45,  40 ), FVector( 1 ) );	// Top right
-	static const FTransform DefaultContentBrowserTransform( FRotator(  20,  160, 0 ), FVector( 50, -65, -40 ), FVector( 1 ) );  // Bottom left
-	static const FTransform DefaultActorDetailsTransform(	FRotator(  20, -160, 0 ), FVector( 50,  65, -40 ), FVector( 1 ) );	// Bottom right
-
+	static FAutoConsoleVariable EnableQuickMenu(TEXT("VREd.EnableQuickMenu"), 0, TEXT("Whether the quick menu should be enabled. Disabled by default."));
+	static FAutoConsoleVariable SteamVRTrackpadDeadzone(TEXT("VREd.SteamVRTrackpadDeadzone"), 0.3f, TEXT("The deadzone for the Vive motion controller trackpad"));
 	static const FTransform DefaultColorPickerTransform(	FRotator( -10, 180, 0), FVector( 30, 35, 0 ), FVector( 1 ) );
-	
 }
 
 
@@ -150,56 +143,47 @@ UVREditorUISystem::UVREditorUISystem( const FObjectInitializer& Initializer ) :
 	QuickMenuUI( nullptr ),
 	QuickRadialMenu( nullptr ),
 	RadialMenuHideDelayTime( 0.0f ),
-	TutorialWidgetClass( nullptr ),
 	InteractorDraggingUI( nullptr ),
 	DraggingUIOffsetTransform( FTransform::Identity ),
 	DraggingUI( nullptr ),
 	ColorPickerUI( nullptr ),
 	bRefocusViewport(false),
-	StartDragUISound( nullptr ),
-	StopDragUISound( nullptr ),
-	HideUISound( nullptr ),
-	ShowUISound( nullptr ),
-	ButtonPressSound( nullptr ),
-	bDraggedDockFromHandPassedThreshold( false ),
 	LastDraggingHoverLocation( FVector::ZeroVector ),
-	bSetDefaultLayout( true )
+	bSequencerOpenedFromRadialMenu(false),
+	bDragPanelFromOpen(false),
+	DragPanelFromOpenTime(0.0f),
+	bPanelCanScale(true)
 {
 	EditorUIPanels.SetNumZeroed( (int32)EEditorUIPanel::TotalCount );
-	
-	// Set default layout transform in correct order of enum
-	DefaultWindowTransforms.Add( VREd::DefaultContentBrowserTransform );
-	DefaultWindowTransforms.Add( VREd::DefaultWorldOutlinerTransform );
-	DefaultWindowTransforms.Add( VREd::DefaultActorDetailsTransform );
-	DefaultWindowTransforms.Add( VREd::DefaultModesTransform );
 }
 
 void UVREditorUISystem::Init()
 {
 	// Register to find out about VR events
+	GetOwner().GetWorldInteraction().OnPreviewInputAction().AddUObject( this, &UVREditorUISystem::OnPreviewInputAction );
 	GetOwner().GetWorldInteraction().OnViewportInteractionInputAction().AddUObject( this, &UVREditorUISystem::OnVRAction );
 	GetOwner().GetWorldInteraction().OnViewportInteractionHoverUpdate().AddUObject( this, &UVREditorUISystem::OnVRHoverUpdate );
 
-	TutorialWidgetClass = LoadClass<UVREditorBaseUserWidget>( nullptr, TEXT( "/Engine/VREditor/Tutorial/UI_VR_Tutorial_00.UI_VR_Tutorial_00_C" ) );
-	check( TutorialWidgetClass != nullptr );
-
-	// Load sounds
-	StartDragUISound = LoadObject<USoundCue>( nullptr, TEXT( "/Engine/VREditor/Sounds/VR_negative_Cue" ) );
-	check( StartDragUISound != nullptr );
-
-	StopDragUISound = LoadObject<USoundCue>( nullptr, TEXT( "/Engine/VREditor/Sounds/VR_negative_Cue" ) );
-	check( StopDragUISound != nullptr );
-
-	HideUISound = LoadObject<USoundCue>( nullptr, TEXT( "/Engine/VREditor/Sounds/VR_close_Cue" ) );
-	check( HideUISound != nullptr );
+	FSlateApplication::Get().OnDragDropCheckOverride.BindUObject(this, &UVREditorUISystem::CheckForVRDragDrop);
 
-	ShowUISound = LoadObject<USoundCue>( nullptr, TEXT( "/Engine/VREditor/Sounds/VR_open_Cue" ) );
-	check( ShowUISound != nullptr );
-
-	ButtonPressSound = LoadObject<USoundCue>(nullptr, TEXT("/Engine/VREditor/Sounds/VR_click1_Cue"));
-	check(ButtonPressSound != nullptr);
+	// TODO: asymmetry hardcoded right now
+	if (GetDefault<UVRModeSettings>()->InteractorHand == EInteractorHand::Left)
+	{
+		LaserInteractor = Cast<UVREditorMotionControllerInteractor>(GetOwner().GetHandInteractor(EControllerHand::Left));
+		UIInteractor = Cast<UVREditorMotionControllerInteractor>(GetOwner().GetHandInteractor(EControllerHand::Right));
+	}
+	else
+	{
+		UIInteractor = Cast<UVREditorMotionControllerInteractor>(GetOwner().GetHandInteractor(EControllerHand::Left));
+		LaserInteractor = Cast<UVREditorMotionControllerInteractor>(GetOwner().GetHandInteractor(EControllerHand::Right));
+	}
+	UIInteractor->SetControllerType(EControllerType::UI);
+	LaserInteractor->SetControllerType(EControllerType::Laser);
 
 	// Create all of our UI panels
+	bRadialMenuIsNumpad = false;
+	RadialMenuHandler = NewObject<UVRRadialMenuHandler>();
+	RadialMenuHandler->Init(this);
 	CreateUIs();
 
 	// Bind the color picker creation & destruction overrides
@@ -211,409 +195,375 @@ void UVREditorUISystem::Init()
 
 	FVREditorActionCallbacks::GizmoCoordinateSystemText = FVREditorActionCallbacks::GetGizmoCoordinateSystemText();
 	FVREditorActionCallbacks::GizmoModeText = FVREditorActionCallbacks::GetGizmoModeText();
+	FVREditorActionCallbacks::UpdateSelectingCandidateActorsText(VRMode);
+	FVREditorActionCallbacks::SelectingCandidateActorsText = FVREditorActionCallbacks::GetSelectingCandidateActorsText();
+
+	GLevelEditorModeTools().OnEditorModeChanged().AddUObject(this, &UVREditorUISystem::HandleEditorModeChanged);
 }
 
 
 void UVREditorUISystem::Shutdown()
 {
+	FSlateApplication::Get().OnDragDropCheckOverride.Unbind();
+
 	if ( VRMode != nullptr && &VRMode->GetWorldInteraction() != nullptr )
 	{
+		GetOwner().GetWorldInteraction().OnPreviewInputAction().RemoveAll( this );
 		GetOwner().GetWorldInteraction().OnViewportInteractionInputAction().RemoveAll( this );
 		GetOwner().GetWorldInteraction().OnViewportInteractionHoverUpdate().RemoveAll( this );
 	}
 
+	GLevelEditorModeTools().OnEditorModeChanged().RemoveAll(this);
+
 	// Unbind the color picker creation & destruction overrides
 	SColorPicker::OnColorPickerNonModalCreateOverride.Unbind();
 	SColorPicker::OnColorPickerDestroyOverride.Unbind();
 	FGlobalTabmanager::Get()->OnOverrideDockableAreaRestore_Handler.Unbind();
 
-	CleanUpActorsBeforeMapChangeOrSimulate();
-
-	DefaultWindowTransforms.Empty();
+	CleanupActorsBeforeShutdown();
 
-	TutorialWidgetClass = nullptr;
-	StartDragUISound = nullptr;
-	StopDragUISound = nullptr;
-	HideUISound = nullptr;
-	ShowUISound = nullptr;
-	ButtonPressSound = nullptr;
 	VRMode = nullptr;
 	DraggingUI = nullptr;
 	ColorPickerUI = nullptr;
 }
 
-void UVREditorUISystem::OnVRAction( FEditorViewportClient& ViewportClient, UViewportInteractor* Interactor,
-	const FViewportActionKeyInput& Action, bool& bOutIsInputCaptured, bool& bWasHandled )
+
+void UVREditorUISystem::OnPreviewInputAction(FEditorViewportClient& ViewportClient, UViewportInteractor* Interactor,
+	const FViewportActionKeyInput& Action, bool& bOutIsInputCaptured, bool& bWasHandled)
 {
-	UVREditorMotionControllerInteractor* VREditorInteractor = Cast<UVREditorMotionControllerInteractor>( Interactor );
-	if( VREditorInteractor && 
-		!( VREditorInteractor->GetDraggingMode() == EViewportInteractionDraggingMode::TransformablesWithGizmo ||
-		  VREditorInteractor->GetDraggingMode() == EViewportInteractionDraggingMode::Interactable ) )
+	UVREditorMotionControllerInteractor* VREditorInteractor = Cast<UVREditorMotionControllerInteractor>(Interactor);
+	
+	// If we are releasing a UI panel that started drag from opening it.
+	if (VREditorInteractor && UIInteractor != nullptr && DraggingUI != nullptr && InteractorDraggingUI != nullptr && UIInteractor == InteractorDraggingUI && VREditorInteractor == UIInteractor && 
+		 Action.Event == EInputEvent::IE_Released && 
+		((VRMode->GetHMDDeviceType() != EHMDDeviceType::DT_SteamVR && Action.ActionType == ViewportWorldActionTypes::SelectAndMove) || (VRMode->GetHMDDeviceType() == EHMDDeviceType::DT_SteamVR && Action.ActionType == VRActionTypes::ConfirmRadialSelection)))
+	{
+		bDragPanelFromOpen = false;
+		UViewportDragOperationComponent* DragOperationComponent = DraggingUI->GetDragOperationComponent();
+		if (DragOperationComponent)
+		{
+			DragOperationComponent->ClearDragOperation();
+		}
+		const bool bPlaySound = false;
+		StopDraggingDockUI(InteractorDraggingUI, bPlaySound);
+	}
+	
+	// UI Interactor Preview actions
+	if (VREditorInteractor &&
+		VREditorInteractor == UIInteractor && 
+		(VREditorInteractor->GetDraggingMode() != EViewportInteractionDraggingMode::World || (VREditorInteractor->GetOtherInteractor() != nullptr && VREditorInteractor->GetOtherInteractor()->GetDraggingMode() != EViewportInteractionDraggingMode::World && VREditorInteractor->GetDraggingMode() == EViewportInteractionDraggingMode::AssistingDrag )))
 	{
-		if ( !bWasHandled )
+		if (Action.Event == EInputEvent::IE_Pressed && 
+			VRMode->GetHMDDeviceType() == EHMDDeviceType::DT_OculusRift 
+			&& Action.ActionType == VRActionTypes::ConfirmRadialSelection
+			&& GIsDemoMode)
 		{
-			// If the numpad is currently showing
-			if (bRadialMenuIsNumpad)
+			ResetAll();
+			bWasHandled = true;
+		}
+
+		if (!bWasHandled)
+		{
+			if (IsShowingRadialMenu(VREditorInteractor))
 			{
-				// Modifier button is backspace
-				if (Action.ActionType == VRActionTypes::Modifier)
+				// If the numpad is currently showing and we press a button (only on press to avoid duplicate calls)
+				if (bRadialMenuIsNumpad && Action.Event == IE_Pressed)
+				{
+					// Modifier button is backspace
+					if (Action.ActionType == VRActionTypes::Modifier)
+					{
+						const bool bRepeat = false;
+						FVREditorActionCallbacks::SimulateBackspace();
+						bWasHandled = true;
+					}
+					// Side triggers function as enter keys
+					if (Action.ActionType == ViewportWorldActionTypes::WorldMovement)
+					{
+						const bool bRepeat = false;
+						FVREditorActionCallbacks::SimulateKeyDown(EKeys::Enter, bRepeat);
+						FVREditorActionCallbacks::SimulateKeyUp(EKeys::Enter);
+						// After pressing enter, dismiss the numpad
+						SwapRadialMenu();
+						if (!bRadialMenuVisibleAtSwap)
+						{
+							HideRadialMenu(VREditorInteractor);
+						}
+						bWasHandled = true;
+					}
+				}
+				if (Action.ActionType == VRActionTypes::Modifier && 
+					Action.Event == IE_Pressed &&
+					!bRadialMenuIsNumpad &&
+					!bWasHandled)
 				{
-					const bool bRepeat = false;
-					FVREditorActionCallbacks::SimulateKeyDown( EKeys::BackSpace, bRepeat );
-					FVREditorActionCallbacks::SimulateKeyUp( EKeys::BackSpace );
-					bWasHandled = true;
+					if (RadialMenuHandler && RadialMenuHandler->GetCurrentMenuGenerator().GetHandle() != RadialMenuHandler->GetHomeMenuGenerator().GetHandle())
+					{
+						RadialMenuHandler->BackOutMenu();
+					}
 				}
-				// Side triggers function as enter keys
-				if (Action.ActionType == ViewportWorldActionTypes::WorldMovement)
+				if (!bWasHandled)
 				{
-					const bool bRepeat = false;
-					FVREditorActionCallbacks::SimulateKeyDown( EKeys::Enter, bRepeat );
-					FVREditorActionCallbacks::SimulateKeyUp( EKeys::Enter );
-					// After pressing enter, dismiss the numpad
-					SwapRadialMenu();
-					if (!bRadialMenuVisibleAtSwap)
+					if ((VRMode->GetHMDDeviceType() == EHMDDeviceType::DT_OculusRift && Action.ActionType == ViewportWorldActionTypes::SelectAndMove) || 
+						(VRMode->GetHMDDeviceType() == EHMDDeviceType::DT_SteamVR && Action.ActionType == VRActionTypes::ConfirmRadialSelection))
 					{
-						HideRadialMenu( VREditorInteractor );
+						// If the radial menu is showing, select the currently highlighted button by pressing the trigger
+						if (QuickRadialMenu->GetCurrentlyHoveredButton().IsValid())
+						{
+							if ((Action.Event == IE_Pressed))
+							{
+								QuickRadialMenu->SimulateLeftClick();
+								bOutIsInputCaptured = true;
+							}
+							if (Action.Event == IE_Released)
+							{
+								bOutIsInputCaptured = false;
+							}
+							bWasHandled = true;
+						}
 					}
-					bWasHandled = true;
 				}
-			}
-			if ( Action.ActionType == VRActionTypes::ConfirmRadialSelection )
+			}	
+			else if (Action.ActionType == VRActionTypes::Modifier &&
+				Action.Event == EInputEvent::IE_Pressed && 
+				VREditorInteractor->GetDraggingMode() != EViewportInteractionDraggingMode::AssistingDrag &&
+				VREditorInteractor->GetDraggingMode() != EViewportInteractionDraggingMode::TransformablesFreely)
 			{
-				const EViewportInteractionDraggingMode DraggingMode = Interactor->GetDraggingMode();
+				const bool bForceRefresh = false;
+				TryToSpawnRadialMenu(VREditorInteractor, bForceRefresh);
+				bWasHandled = true;
+			}
+		}
+	}
 
-				// We may be activating sequencer scrub mode
-				if (DraggingMode == EViewportInteractionDraggingMode::Nothing && 
-					IsShowingEditorUIPanel(EEditorUIPanel::SequencerUI) && 
-					Action.Event == IE_Pressed && 
-					GetOwner().GetHMDDeviceType() == EHMDDeviceType::DT_SteamVR )
-				{
-					SequencerPlayback(Interactor, VREditorInteractor);
-				}
+	// Laser Interaction Preview actions
+	if (VREditorInteractor &&
+		VREditorInteractor == LaserInteractor)
+	{
 
-				// Toggle the radial menu by pressing the touchpad or thumb stick, as long as we're not busy dragging something
-				if ( Action.Event == IE_Pressed && 
-					 DraggingMode != EViewportInteractionDraggingMode::TransformablesFreely &&
-					 DraggingMode != EViewportInteractionDraggingMode::World &&
-					 DraggingMode != EViewportInteractionDraggingMode::AssistingDrag )
+		if (Action.ActionType == ViewportWorldActionTypes::SelectAndMove &&
+			VREditorInteractor->GetDraggingMode() == EViewportInteractionDraggingMode::Nothing)
+		{
+			FVector LaserPointerStart, LaserPointerEnd;
+			// If we are clicking on an Actor but not a widget component, send a fake mouse click event to toggle focus
+			if (VREditorInteractor->GetLaserPointer(LaserPointerStart, LaserPointerEnd))
+			{
+				FHitResult HitResult = VREditorInteractor->GetHitResultFromLaserPointer();
+				if (HitResult.Actor.IsValid())
 				{
-					if ( !IsShowingRadialMenu(VREditorInteractor) )
+					UWidgetComponent* WidgetComponent = Cast<UWidgetComponent>(HitResult.GetComponent());
+
+					if (WidgetComponent == nullptr)
 					{
-						// Don't force the UI to spawn if using the touchpad/thumbstick to spawn the menu
-						const bool bForceOverUI = false;
-						TryToSpawnRadialMenu(VREditorInteractor, bForceOverUI);
-					}
-					else
-					{	
-						// If the numpad is showing, either swap to the radial menu or dismiss the menu entirely
+						const bool bIsRightClicking =
+							(Action.Event == IE_Pressed && VREditorInteractor->IsModifierPressed()) ||
+							(Action.Event == IE_Released && VREditorInteractor->IsRightClickingOnUI());
+						TSet<FKey> PressedButtons;
+						FPointerEvent PointerEvent(
+							1 + (uint8)VREditorInteractor->GetControllerSide(),
+							FVector2D::ZeroVector,
+							FVector2D::ZeroVector,
+							PressedButtons,
+							bIsRightClicking ? EKeys::RightMouseButton : EKeys::LeftMouseButton,
+							0.0f,	// Wheel delta
+							FModifierKeysState());
+
+						FWidgetPath EmptyWidgetPath;
+						FReply Reply = FSlateApplication::Get().RoutePointerDownEvent(EmptyWidgetPath, PointerEvent);
 						if (bRadialMenuIsNumpad)
 						{
+							// If clicking somewhere outside UI so the widget loses focus
 							SwapRadialMenu();
 							if (!bRadialMenuVisibleAtSwap)
 							{
 								HideRadialMenu(VREditorInteractor);
 							}
 						}
-						// Always dismiss the regular radial menu
-						else
-						{
-							HideRadialMenu(VREditorInteractor);
-						}
 					}
-					bWasHandled = true;
 				}
 			}
-			else 
-			if ( Action.ActionType == ViewportWorldActionTypes::SelectAndMove_LightlyPressed )
+			// Clicking on UI
 			{
-				// If the radial menu is showing, select the currently highlighted button by pressing the trigger
-				if ( IsShowingRadialMenu( VREditorInteractor ) )
+				FVector LaserPointerStart, LaserPointerEnd;
+				if (VREditorInteractor->GetLaserPointer(LaserPointerStart, LaserPointerEnd))
 				{
-					TSharedPtr<SRadialBox> RadialBox = StaticCastSharedPtr<SRadialBox>(RadialWidget);
-					if(RadialBox->GetCurrentlyHoveredButton().IsValid())
-					{
-						if((Action.Event == IE_Pressed))
-						{
-							RadialBox->SimulateLeftClickDown();
-							bOutIsInputCaptured = true;
-						}
-						if (Action.Event == IE_Released)
-						{
-							RadialBox->SimulateLeftClickUp();
-							bOutIsInputCaptured = false;
-						}
-						bWasHandled = true;
-					}
-					else
+					FHitResult HitResult = VREditorInteractor->GetHitResultFromLaserPointer();
+					if (HitResult.Actor.IsValid())
 					{
-						FVector LaserPointerStart, LaserPointerEnd;
-						if (VREditorInteractor->GetLaserPointer( LaserPointerStart, LaserPointerEnd ))
+						// Only allow clicks to our own widget components
+						UVREditorWidgetComponent* WidgetComponent = Cast<UVREditorWidgetComponent>(HitResult.GetComponent());
+						if (WidgetComponent != nullptr)
 						{
-							FHitResult HitResult = VREditorInteractor->GetHitResultFromLaserPointer();
-							if (HitResult.Actor.IsValid())
+							// Always mark the event as handled so that the editor doesn't try to select the widget component
+							bWasHandled = true;
+
+							if (Action.Event != IE_Repeat)
 							{
-								UWidgetComponent* WidgetComponent = Cast<UWidgetComponent>( HitResult.GetComponent() );
-								// If we are clicking on an Actor but not a widget component, send a fake mouse click event to toggle focus
-								if (WidgetComponent == nullptr)
+								// If the Modifier button is held down, treat this like a right click instead of a left click
+								const bool bIsRightClicking =
+									(Action.Event == IE_Pressed && VREditorInteractor->IsModifierPressed()) ||
+									(Action.Event == IE_Released && VREditorInteractor->IsRightClickingOnUI());
+
+								FVector2D LastLocalHitLocation = WidgetComponent->GetLastLocalHitLocation();
+
+								FVector2D LocalHitLocation;
+								WidgetComponent->GetLocalHitLocation(HitResult.ImpactPoint, LocalHitLocation);
+
+								// If we weren't already hovering over this widget, then we'll reset the last hit location
+								if (WidgetComponent != VREditorInteractor->GetLastHoveredWidgetComponent())
 								{
-									const bool bIsRightClicking =
-										(Action.Event == IE_Pressed && VREditorInteractor->IsModifierPressed()) ||
-										(Action.Event == IE_Released && VREditorInteractor->IsRightClickingOnUI());
-									TSet<FKey> PressedButtons;
-									FPointerEvent PointerEvent(
-										1 + (uint8)VREditorInteractor->GetControllerSide(),
-										FVector2D::ZeroVector,
-										FVector2D::ZeroVector,
-										PressedButtons,
-										bIsRightClicking ? EKeys::RightMouseButton : EKeys::LeftMouseButton,
-										0.0f,	// Wheel delta
-										FModifierKeysState() );
+									LastLocalHitLocation = LocalHitLocation;
 
-									FWidgetPath EmptyWidgetPath;
-									FReply Reply = FSlateApplication::Get().RoutePointerDownEvent( EmptyWidgetPath, PointerEvent );
-									bWasHandled = true;
-									if (bRadialMenuIsNumpad)
+									if (UVREditorWidgetComponent* VRWidgetComponent = Cast<UVREditorWidgetComponent>(VREditorInteractor->GetLastHoveredWidgetComponent()))
 									{
-										// If clicking somewhere outside UI so the widget loses focus
-										SwapRadialMenu();
-										if (!bRadialMenuVisibleAtSwap)
-										{
-											HideRadialMenu( VREditorInteractor );
-										}
+										VRWidgetComponent->SetIsHovering(false);
+										OnHoverEndEffect(VRWidgetComponent);
 									}
 								}
-							}
-						}
-					}
-				}
 
-				if(!bWasHandled)
-				{
-					FVector LaserPointerStart, LaserPointerEnd;
-					if ( VREditorInteractor->GetLaserPointer( LaserPointerStart, LaserPointerEnd ) )
-					{
-						FHitResult HitResult = VREditorInteractor->GetHitResultFromLaserPointer();
-						if ( HitResult.Actor.IsValid() )
-						{
-							// Only allow clicks to our own widget components
-							UVREditorWidgetComponent* WidgetComponent = Cast<UVREditorWidgetComponent>( HitResult.GetComponent() );
-							if ( WidgetComponent != nullptr )
-							{
-								// Always mark the event as handled so that the editor doesn't try to select the widget component
-								bWasHandled = true;
-
-								if ( Action.Event != IE_Repeat )
+								FWidgetPath WidgetPathUnderFinger = FWidgetPath(WidgetComponent->GetHitWidgetPath(HitResult.ImpactPoint, /*bIgnoreEnabledStatus*/ false));
+								if (WidgetPathUnderFinger.IsValid())
 								{
-									// If the Modifier button is held down, treat this like a right click instead of a left click
-									const bool bIsRightClicking =
-										( Action.Event == IE_Pressed && VREditorInteractor->IsModifierPressed() ) ||
-										( Action.Event == IE_Released && VREditorInteractor->IsRightClickingOnUI() );
+									TSet<FKey> PressedButtons;
+									if (Action.Event == IE_Pressed)
+									{
+										PressedButtons.Add(bIsRightClicking ? EKeys::RightMouseButton : EKeys::LeftMouseButton);
+									}
 
-									FVector2D LastLocalHitLocation = WidgetComponent->GetLastLocalHitLocation();
+									FPointerEvent PointerEvent(
+										1 + (uint8)VREditorInteractor->GetControllerSide(),
+										LocalHitLocation,
+										LastLocalHitLocation,
+										PressedButtons,
+										bIsRightClicking ? EKeys::RightMouseButton : EKeys::LeftMouseButton,
+										0.0f,	// Wheel delta
+										FModifierKeysState());
 
-									FVector2D LocalHitLocation;
-									WidgetComponent->GetLocalHitLocation( HitResult.ImpactPoint, LocalHitLocation );
+									VREditorInteractor->SetLastHoveredWidgetComponent(WidgetComponent);
 
-									// If we weren't already hovering over this widget, then we'll reset the last hit location
-									if ( WidgetComponent != VREditorInteractor->GetHoveringOverWidgetComponent() )
+									if (UVREditorWidgetComponent* VRWidgetComponent = Cast<UVREditorWidgetComponent>(VREditorInteractor->GetLastHoveredWidgetComponent()))
 									{
-										LastLocalHitLocation = LocalHitLocation;
-
-										if ( UVREditorWidgetComponent* VRWidgetComponent = Cast<UVREditorWidgetComponent>( VREditorInteractor->GetHoveringOverWidgetComponent() ) )
-										{
-											VRWidgetComponent->SetIsHovering( false );
-										}
+										VRWidgetComponent->SetIsHovering(true);
+										OnHoverBeginEffect(VRWidgetComponent);
 									}
 
-									FWidgetPath WidgetPathUnderFinger = FWidgetPath( WidgetComponent->GetHitWidgetPath( HitResult.ImpactPoint, /*bIgnoreEnabledStatus*/ false ) );
-									if ( WidgetPathUnderFinger.IsValid() )
+									FReply Reply = FReply::Unhandled();
+									if (Action.Event == IE_Pressed)
 									{
-										TSet<FKey> PressedButtons;
-										if ( Action.Event == IE_Pressed )
-										{
-											PressedButtons.Add( bIsRightClicking ? EKeys::RightMouseButton : EKeys::LeftMouseButton );
-										}
-
-										FPointerEvent PointerEvent(
-											1 + (uint8)VREditorInteractor->GetControllerSide(),
-											LocalHitLocation,
-											LastLocalHitLocation,
-											PressedButtons,
-											bIsRightClicking ? EKeys::RightMouseButton : EKeys::LeftMouseButton,
-											0.0f,	// Wheel delta
-											FModifierKeysState() );
-
-										VREditorInteractor->SetHoveringOverWidgetComponent( WidgetComponent );
-
-										if ( UVREditorWidgetComponent* VRWidgetComponent = Cast<UVREditorWidgetComponent>( VREditorInteractor->GetHoveringOverWidgetComponent() ) )
+										const double CurrentTime = FPlatformTime::Seconds();
+										if (CurrentTime - VREditorInteractor->GetLastUIPressTime() <= GetDefault<UVRModeSettings>()->DoubleClickTime)
 										{
-											VRWidgetComponent->SetIsHovering( true );
+											// Trigger a double click event!
+											Reply = FSlateApplication::Get().RoutePointerDoubleClickEvent(WidgetPathUnderFinger, PointerEvent);
 										}
-
-										FReply Reply = FReply::Unhandled();
-										if ( Action.Event == IE_Pressed )
+										else
 										{
-											const double CurrentTime = FPlatformTime::Seconds();
-											if (CurrentTime - VREditorInteractor->GetLastClickReleaseTime() <= VREd::DoubleClickTime->GetFloat())
-											{
-												// Trigger a double click event!
-												Reply = FSlateApplication::Get().RoutePointerDoubleClickEvent(WidgetPathUnderFinger, PointerEvent);
-											}
-											else
+											// If we are clicking on an editable text field and the radial menu is not a numpad, show the numpad
+											if (WidgetPathUnderFinger.Widgets.Last().Widget->GetTypeAsString() == TEXT("SEditableText") && (bRadialMenuIsNumpad == false))
 											{
-												// If we are clicking on an editable text field and the radial menu is not a numpad, show the numpad
-												if (WidgetPathUnderFinger.Widgets.Last().Widget->GetTypeAsString() == TEXT("SEditableText") && (bRadialMenuIsNumpad == false))
+												if (!QuickRadialMenu->bHidden)
 												{
-													if (!QuickRadialMenu->bHidden)
-													{
-														bRadialMenuVisibleAtSwap = true;
-													}
-													else
-													{
-														bRadialMenuVisibleAtSwap = false;
-														// Force the radial menu to spawn even if the laser is over UI
-														const bool bForceOverUI = true;
-														TryToSpawnRadialMenu( VREditorInteractor, bForceOverUI );
-													}
-													SwapRadialMenu();
+													bRadialMenuVisibleAtSwap = true;
 												}
-												Reply = FSlateApplication::Get().RoutePointerDownEvent(WidgetPathUnderFinger, PointerEvent);
-											}
-
-											// In case of selecting a level in the content browser the VREditormode is closed, this makes sure nothing happens after that.
-											if( !IVREditorModule::Get().IsVREditorModeActive() )
-											{
-												bWasHandled = true;
-												return;
+												else
+												{
+													bRadialMenuVisibleAtSwap = false;
+													// Force the radial menu to spawn even if the laser is over UI
+													const bool bForceRefresh = false;
+													TryToSpawnRadialMenu(UIInteractor, bForceRefresh);
+												}
+												SwapRadialMenu();
 											}
-										
-											VREditorInteractor->SetIsClickingOnUI( true );
-											VREditorInteractor->SetIsRightClickingOnUI( bIsRightClicking );
-											VREditorInteractor->SetLastClickReleaseTime( CurrentTime );
-											bOutIsInputCaptured = true;
-
-											// Play a haptic effect on press
-											VREditorInteractor->PlayHapticEffect( VREd::UIPressHapticFeedbackStrength->GetFloat() );
+											Reply = FSlateApplication::Get().RoutePointerDownEvent(WidgetPathUnderFinger, PointerEvent);
 										}
-										else if ( Action.Event == IE_Released )
+
+										// In case of selecting a level in the content browser the VREditormode is closed, this makes sure nothing happens after that.
+										if (!IVREditorModule::Get().IsVREditorModeActive())
 										{
-											Reply = FSlateApplication::Get().RoutePointerUpEvent( WidgetPathUnderFinger, PointerEvent );
+											bWasHandled = true;
+											return;
 										}
+
+										VREditorInteractor->SetIsClickingOnUI(true);
+										VREditorInteractor->SetIsRightClickingOnUI(bIsRightClicking);
+										VREditorInteractor->SetLastUIPressTime(CurrentTime);
+										bOutIsInputCaptured = true;
+
+										// Play a haptic effect on press
+										VREditorInteractor->PlayHapticEffect(VREd::UIPressHapticFeedbackStrength->GetFloat());
+									}
+									else if (Action.Event == IE_Released)
+									{
+										Reply = FSlateApplication::Get().RoutePointerUpEvent(WidgetPathUnderFinger, PointerEvent);
 									}
 								}
 							}
-							else
-							{
-								// If we didn't handle the input in any other way, send an empty mouse down event so Slate focus is handled correctly
-								const bool bIsRightClicking =
-									(Action.Event == IE_Pressed && VREditorInteractor->IsModifierPressed()) ||
-									(Action.Event == IE_Released && VREditorInteractor->IsRightClickingOnUI());
-								TSet<FKey> PressedButtons;
-								FPointerEvent PointerEvent(
-									1 + (uint8)VREditorInteractor->GetControllerSide(),
-									FVector2D::ZeroVector,
-									FVector2D::ZeroVector,
-									PressedButtons,
-									bIsRightClicking ? EKeys::RightMouseButton : EKeys::LeftMouseButton,
-									0.0f,	// Wheel delta
-									FModifierKeysState() );
-
-								FWidgetPath EmptyWidgetPath;
-								FReply Reply = FSlateApplication::Get().RoutePointerDownEvent( EmptyWidgetPath, PointerEvent );
-							}
 						}
-					}
-				}
-				if ( Action.Event == IE_Released )
-				{
-					bool bWasRightClicking = false;
-					if ( VREditorInteractor->IsClickingOnUI() )
-					{
-						if ( VREditorInteractor->IsRightClickingOnUI() )
+						else
 						{
-							bWasRightClicking = true;
+							// If we didn't handle the input in any other way, send an empty mouse down event so Slate focus is handled correctly
+							const bool bIsRightClicking =
+								(Action.Event == IE_Pressed && VREditorInteractor->IsModifierPressed()) ||
+								(Action.Event == IE_Released && VREditorInteractor->IsRightClickingOnUI());
+							TSet<FKey> PressedButtons;
+							FPointerEvent PointerEvent(
+								1 + (uint8)VREditorInteractor->GetControllerSide(),
+								FVector2D::ZeroVector,
+								FVector2D::ZeroVector,
+								PressedButtons,
+								bIsRightClicking ? EKeys::RightMouseButton : EKeys::LeftMouseButton,
+								0.0f,	// Wheel delta
+								FModifierKeysState());
+
+							FWidgetPath EmptyWidgetPath;
+							FReply Reply = FSlateApplication::Get().RoutePointerDownEvent(EmptyWidgetPath, PointerEvent);
 						}
-						VREditorInteractor->SetIsClickingOnUI( false );
-						VREditorInteractor->SetIsRightClickingOnUI( false );
-						bOutIsInputCaptured = false;
 					}
+				}
+			}
 
-					if ( !bWasHandled )
+			if (Action.Event == IE_Released)
+			{
+				bool bWasRightClicking = false;
+				if (VREditorInteractor->IsClickingOnUI())
+				{
+					if (VREditorInteractor->IsRightClickingOnUI())
 					{
-						TSet<FKey> PressedButtons;
-						FPointerEvent PointerEvent(
-							1 + (uint8)VREditorInteractor->GetControllerSide(),
-							FVector2D::ZeroVector,
-							FVector2D::ZeroVector,
-							PressedButtons,
-							bWasRightClicking ? EKeys::RightMouseButton : EKeys::LeftMouseButton,
-							0.0f,	// Wheel delta
-							FModifierKeysState() );
-
-						FWidgetPath EmptyWidgetPath;
-
-						VREditorInteractor->SetIsClickingOnUI( false );
-						VREditorInteractor->SetIsRightClickingOnUI( false );
-						FReply Reply = FSlateApplication::Get().RoutePointerUpEvent( EmptyWidgetPath, PointerEvent );
+						bWasRightClicking = true;
 					}
+					VREditorInteractor->SetIsClickingOnUI(false);
+					VREditorInteractor->SetIsRightClickingOnUI(false);
+					bOutIsInputCaptured = false;
+				}
+
+				if (!bWasHandled)
+				{
+					TSet<FKey> PressedButtons;
+					FPointerEvent PointerEvent(
+						1 + (uint8)VREditorInteractor->GetControllerSide(),
+						FVector2D::ZeroVector,
+						FVector2D::ZeroVector,
+						PressedButtons,
+						bWasRightClicking ? EKeys::RightMouseButton : EKeys::LeftMouseButton,
+						0.0f,	// Wheel delta
+						FModifierKeysState());
+
+					FWidgetPath EmptyWidgetPath;
+
+					VREditorInteractor->SetIsClickingOnUI(false);
+					VREditorInteractor->SetIsRightClickingOnUI(false);
+					FReply Reply = FSlateApplication::Get().RoutePointerUpEvent(EmptyWidgetPath, PointerEvent);
 				}
 			}
 		}
 	}
+
 }
 
 
-void UVREditorUISystem::SequencerPlayback(UViewportInteractor* Interactor, UVREditorMotionControllerInteractor* VREditorInteractor)
+void UVREditorUISystem::OnVRAction( FEditorViewportClient& ViewportClient, UViewportInteractor* Interactor,
+	const FViewportActionKeyInput& Action, bool& bOutIsInputCaptured, bool& bWasHandled )
 {
-	ISequencer* CurrentSequencer = GetOwner().GetCurrentSequencer();
-	if (Interactor->GetOtherInteractor() != nullptr && CurrentSequencer != nullptr)
-	{
-		if (!VREditorInteractor->IsTouchingTrackpad())
-		{
-			// Pause: For releasing touchpad on Vive
-			FVREditorActionCallbacks::SimulateKeyDown(EKeys::K, false);
-			FVREditorActionCallbacks::SimulateKeyUp(EKeys::K);
 
-		}
-		UVREditorMotionControllerInteractor* OtherMotionController = CastChecked<UVREditorMotionControllerInteractor>(Interactor->GetOtherInteractor());
-		if (!OtherMotionController->IsTouchingTrackpad())
-		{
-			// Basic Sequencer Playback
-			
-			if (VREditorInteractor->GetTrackpadPosition().Y > 0.5f && FMath::Abs(VREditorInteractor->GetTrackpadPosition().X) < 0.15)
-			{
-				// FF
-				
-				CurrentSequencer->OnPlay(false, 2.0f);
-				
-			}
-			else if (VREditorInteractor->GetTrackpadPosition().Y <= -0.5f && FMath::Abs(VREditorInteractor->GetTrackpadPosition().X) < 0.15)
-			{
-				// Fast Reverse
-				CurrentSequencer->OnPlay(false, -2.0f);
-			}
-			else if (VREditorInteractor->GetTrackpadPosition().X > 0.5f)
-			{
-				// Play
-				CurrentSequencer->OnPlay(false, 1.0f);
-			}
-			else if (VREditorInteractor->GetTrackpadPosition().X < -0.5f)
-			{
-				// Reverse
-				CurrentSequencer->OnPlay(false, -1.0f);
-			}
-			
-			else
-			{
-				// Pause: Center of touchpad, center of thumbstick
-				CurrentSequencer->Pause();
-			}
-		}
-		else
-		{
-			// Two-Handed Sequencer Playback
-		}
-	}
 }
 
 void UVREditorUISystem::OnVRHoverUpdate(UViewportInteractor* Interactor, FVector& HoverImpactPoint, bool& bWasHandled)
@@ -621,14 +571,7 @@ void UVREditorUISystem::OnVRHoverUpdate(UViewportInteractor* Interactor, FVector
 	UVREditorMotionControllerInteractor* VREditorInteractor = Cast<UVREditorMotionControllerInteractor>( Interactor );
 	if( VREditorInteractor != nullptr )
 	{
-		// We may be activating sequencer scrub mode
-		if (Interactor->GetDraggingMode() == EViewportInteractionDraggingMode::Nothing &&
-			IsShowingEditorUIPanel(EEditorUIPanel::SequencerUI) &&
-			GetOwner().GetHMDDeviceType() == EHMDDeviceType::DT_OculusRift)
-		{
-			SequencerPlayback(Interactor, VREditorInteractor);
-		}
-		if( !bWasHandled && Interactor->GetDraggingMode() != EViewportInteractionDraggingMode::Interactable )
+		if( !bWasHandled && Interactor->GetDraggingMode() == EViewportInteractionDraggingMode::Nothing )
 		{
 			FVector LaserPointerStart, LaserPointerEnd;
 			if( Interactor->GetLaserPointer( LaserPointerStart, LaserPointerEnd ) )
@@ -636,14 +579,12 @@ void UVREditorUISystem::OnVRHoverUpdate(UViewportInteractor* Interactor, FVector
 				FHitResult HitResult = Interactor->GetHitResultFromLaserPointer();
 				if( HitResult.Actor.IsValid() )
 				{
-					// If a button is currently highlighted on the radial menu, don't highlight another button with the laser
-					if ((HitResult.Actor == QuickRadialMenu) && (RadialWidget.Get() != nullptr))
+					// The laser should make the quick radial menu stay active
+					if ((HitResult.Actor == QuickRadialMenu) && (QuickRadialMenu != nullptr))
 					{
-						const TSharedPtr<SRadialBox>& RadialBox = StaticCastSharedPtr<SRadialBox>(RadialWidget);
-						if (RadialBox->GetCurrentlyHoveredButton().Get() != nullptr)
-						{
-							return;
-						}
+						RadialMenuModifierSpawnTime = FTimespan::FromSeconds(FPlatformTime::Seconds());
+						FVector2D ReturnToCenter = FVector2D::ZeroVector;
+						QuickRadialMenu->HighlightSlot(ReturnToCenter);
 					}
 
 					// Only allow clicks to our own widget components
@@ -656,13 +597,14 @@ void UVREditorUISystem::OnVRHoverUpdate(UViewportInteractor* Interactor, FVector
 						WidgetComponent->GetLocalHitLocation( HitResult.ImpactPoint, LocalHitLocation );
 
 						// If we weren't already hovering over this widget, then we'll reset the last hit location
-						if( WidgetComponent != VREditorInteractor->GetHoveringOverWidgetComponent() )
+						if( WidgetComponent != VREditorInteractor->GetLastHoveredWidgetComponent() )
 						{
 							LastLocalHitLocation = LocalHitLocation;
 
-							if ( UVREditorWidgetComponent* VRWidgetComponent = Cast<UVREditorWidgetComponent>( VREditorInteractor->GetHoveringOverWidgetComponent() ) )
+							if ( UVREditorWidgetComponent* VRWidgetComponent = Cast<UVREditorWidgetComponent>( VREditorInteractor->GetLastHoveredWidgetComponent() ) )
 							{
 								VRWidgetComponent->SetIsHovering(false);
+								OnHoverEndEffect(VRWidgetComponent);
 							}
 						}
 
@@ -670,7 +612,21 @@ void UVREditorUISystem::OnVRHoverUpdate(UViewportInteractor* Interactor, FVector
 						FWidgetPath WidgetPathUnderFinger = FWidgetPath( WidgetComponent->GetHitWidgetPath( HitResult.ImpactPoint, /*bIgnoreEnabledStatus*/ false ) );
 						if ( WidgetPathUnderFinger.IsValid() )
 						{
+							HoverImpactPoint = HitResult.ImpactPoint;
+							VREditorInteractor->GetInteractorData().LastHoverLocationOverUI = HitResult.ImpactPoint;
+							VREditorInteractor->SetLastHoveredWidgetComponent( WidgetComponent );
+							VREditorInteractor->SetIsHoveringOverUI( true );
+							
 							TSet<FKey> PressedButtons;
+							if( VREditorInteractor->IsClickingOnUI() == true )
+							{
+								PressedButtons.Add( EKeys::LeftMouseButton );
+							}
+							else if( VREditorInteractor->IsRightClickingOnUI() == true )
+							{
+								PressedButtons.Add( EKeys::RightMouseButton );
+							}
+
 							FPointerEvent PointerEvent(
 								1 + (uint8)VREditorInteractor->GetControllerSide(),
 								LocalHitLocation,
@@ -682,13 +638,11 @@ void UVREditorUISystem::OnVRHoverUpdate(UViewportInteractor* Interactor, FVector
 							FSlateApplication::Get().RoutePointerMoveEvent(WidgetPathUnderFinger, PointerEvent, false);
 							
 							bWasHandled = true;
-							HoverImpactPoint = HitResult.ImpactPoint;
-							VREditorInteractor->SetHoveringOverWidgetComponent( WidgetComponent );
-							VREditorInteractor->SetIsHoveringOverUI( true );
 
 							if ( UVREditorWidgetComponent* VRWidgetComponent = Cast<UVREditorWidgetComponent>(WidgetComponent) )
 							{
 								VRWidgetComponent->SetIsHovering(true);
+								OnHoverBeginEffect(VRWidgetComponent);
 							}
 
 							// Route the mouse scrolling
@@ -783,15 +737,16 @@ void UVREditorUISystem::OnVRHoverUpdate(UViewportInteractor* Interactor, FVector
 		}
 
 		// If nothing was hovered, make sure we tell Slate about that
-		if( !bWasHandled && VREditorInteractor->GetHoveringOverWidgetComponent() != nullptr )
+		if( !bWasHandled && VREditorInteractor->GetLastHoveredWidgetComponent() != nullptr )
 		{
-			if ( UVREditorWidgetComponent* VRWidgetComponent = Cast<UVREditorWidgetComponent>( VREditorInteractor->GetHoveringOverWidgetComponent() ) )
+			if ( UVREditorWidgetComponent* VRWidgetComponent = Cast<UVREditorWidgetComponent>( VREditorInteractor->GetLastHoveredWidgetComponent() ) )
 			{
 				VRWidgetComponent->SetIsHovering( false );
+				OnHoverEndEffect(VRWidgetComponent);
 			}
 
-			const FVector2D LastLocalHitLocation = VREditorInteractor->GetHoveringOverWidgetComponent()->GetLastLocalHitLocation();
-			VREditorInteractor->SetHoveringOverWidgetComponent( nullptr );
+			const FVector2D LastLocalHitLocation = VREditorInteractor->GetLastHoveredWidgetComponent()->GetLastLocalHitLocation();
+			VREditorInteractor->SetLastHoveredWidgetComponent( nullptr );
 
 			TSet<FKey> PressedButtons;
 			FPointerEvent PointerEvent(
@@ -802,9 +757,12 @@ void UVREditorUISystem::OnVRHoverUpdate(UViewportInteractor* Interactor, FVector
 				PressedButtons,
 				FModifierKeysState() );
 
+
 			FSlateApplication::Get().RoutePointerMoveEvent( FWidgetPath(), PointerEvent, false );
 		}
+
 	}
+
 }
 
 void UVREditorUISystem::Tick( FEditorViewportClient* ViewportClient, const float DeltaTime )
@@ -815,9 +773,26 @@ void UVREditorUISystem::Tick( FEditorViewportClient* ViewportClient, const float
 		FSlateApplication::Get().SetUserFocus(0, ViewportClient->GetEditorViewportWidget());
 	}
 
+	const FTimespan CompareTime = FMath::Max(UIInteractor->GetLastActiveTrackpadUpdateTime(), RadialMenuModifierSpawnTime);
+	if (!bRadialMenuIsNumpad &&
+		FTimespan::FromSeconds(FPlatformTime::Seconds()) - CompareTime > FTimespan::FromSeconds(1.5f))
+	{
+		HideRadialMenu(UIInteractor);
+	}
+
+	// When dragging a panel with the UI interactor for a while, we want to close the radial menu.
+	if (bDragPanelFromOpen)
+	{
+		DragPanelFromOpenTime += DeltaTime;
+		if (DragPanelFromOpenTime >= 1.0 && UIInteractor != nullptr && DraggingUI != nullptr && InteractorDraggingUI != nullptr && UIInteractor == InteractorDraggingUI)
+		{
+			HideRadialMenu(InteractorDraggingUI, false);
+		}
+	}
+
 	// Figure out if one hand is "aiming toward" the other hand.  We'll fade in a UI on the hand being
 	// aimed at when the user does this.
-	if( QuickMenuUI != nullptr )
+	if( QuickMenuUI != nullptr && VREd::EnableQuickMenu->GetInt() == 1)
 	{
 		UVREditorInteractor* HandInteractorWithQuickMenu = nullptr;
 		if( QuickMenuUI->IsUIVisible() )
@@ -830,28 +805,27 @@ void UVREditorUISystem::Tick( FEditorViewportClient* ViewportClient, const float
 		UViewportInteractor* HandInteractorThatNeedsQuickMenu = nullptr;
 		for( UViewportInteractor* Interactor : GetOwner().GetWorldInteraction().GetInteractors() )
 		{
-			bool bShowQuickMenuOnArm = false;
-
 			UViewportInteractor* OtherInteractor = Interactor->GetOtherInteractor();
-			if( OtherInteractor != nullptr )
+			if( Interactor->GetDraggingMode() == EViewportInteractionDraggingMode::Nothing &&
+				OtherInteractor != nullptr )
 			{
 				// @todo vreditor tweak: Weird to hard code this here. Probably should be an accessor on the hand itself, and based on the actual device type
 				const FTransform UICapsuleTransform = OtherInteractor->GetTransform();
-				const FVector UICapsuleStart = FVector( -3.0f, 0.0f, 0.0f ) * WorldScaleFactor;
-				const FVector UICapsuleEnd = FVector( -18.0f, 0.0f, 0.0f ) * WorldScaleFactor;
-				const float UICapsuleLocalRadius = 8.0f * WorldScaleFactor;
-				const float MinDistanceToUICapsule = 8.0f * WorldScaleFactor; // @todo vreditor tweak
+				const FVector UICapsuleStart = FVector( -16.0f, 0.0f, 0.0f ) * WorldScaleFactor;
+				const FVector UICapsuleEnd = FVector( -23.0f, 0.0f, 0.0f ) * WorldScaleFactor;
+				const float UICapsuleLocalRadius = 10.0f * WorldScaleFactor;
+				const float MinDistanceToUICapsule = 10.0f * WorldScaleFactor; // @todo vreditor tweak
 				const FVector UIForwardVector = FVector::UpVector;
 				const float MinDotForAimingAtOtherHand = 0.25f;	// @todo vreditor tweak
 
-				if (GetOwner().IsHandAimingTowardsCapsule(Interactor, UICapsuleTransform, UICapsuleStart, UICapsuleEnd, UICapsuleLocalRadius, MinDistanceToUICapsule, UIForwardVector, MinDotForAimingAtOtherHand) && (QuickRadialMenu->bHidden))
-				{
-					bShowQuickMenuOnArm = true;
-				}
-
-				if( bShowQuickMenuOnArm )
+				// Don't spawn a quick menu on a hand that has a radial menu visible -- they will cover each other up!
+				UVREditorInteractor* OtherVRInteractor = Cast<UVREditorInteractor>( OtherInteractor );
+				if( OtherVRInteractor != nullptr && !IsShowingRadialMenu( OtherVRInteractor ) )
 				{
-					HandInteractorThatNeedsQuickMenu = OtherInteractor;
+					if( GetOwner().IsHandAimingTowardsCapsule( Interactor, UICapsuleTransform, UICapsuleStart, UICapsuleEnd, UICapsuleLocalRadius, MinDistanceToUICapsule, UIForwardVector, MinDotForAimingAtOtherHand ) )
+					{
+						HandInteractorThatNeedsQuickMenu = OtherInteractor;
+					}
 				}
 			}
 		}
@@ -862,22 +836,16 @@ void UVREditorUISystem::Tick( FEditorViewportClient* ViewportClient, const float
 		{
 			// Despawn
 			HandInteractorWithQuickMenu->SetHasUIOnForearm( false );
-			HandInteractorWithQuickMenu->SetHasUIInFront( false );
 			QuickMenuUI->ShowUI( false );
 			// Reset all the currently animating buttons to their minimum scale and stop animating them
 			for (FVRButton& VRButton : VRButtons)
 			{
 				if(VRButton.AnimationDirection != EVREditorAnimationState::None)
 				{
-					
-					TSharedRef<SWidget> PlaceholderWidget = (VRButton.BaseSlot->DetachWidget()).ToSharedRef();
-					TSharedRef<SWidget> WidgetToTransfer = (VRButton.HoverSlot->DetachWidget()).ToSharedRef();
-					VRButton.BaseSlot->AttachWidget(WidgetToTransfer);
-					VRButton.HoverSlot->AttachWidget(PlaceholderWidget);
 					VRButton.CurrentScale = VRButton.MinScale;
-					VRButton.ButtonBorder->SetRenderTransform(FSlateRenderTransform::FTransform2D(VRButton.CurrentScale));
+					VRButton.ButtonWidget->SetRelativeScale3D(VRButton.CurrentScale*VRButton.OriginalRelativeScale*(GetOwner().GetWorldScaleFactor()));
 					VRButton.AnimationDirection = EVREditorAnimationState::None;
-		}
+				}
 			}
 		}
 
@@ -890,7 +858,6 @@ void UVREditorUISystem::Tick( FEditorViewportClient* ViewportClient, const float
 				const AVREditorFloatingUI::EDockedTo DockTo = ( MotionControllerHandInteractorThatNeedsQuickMenu->GetControllerSide() == EControllerHand::Left ) ? AVREditorFloatingUI::EDockedTo::LeftArm : AVREditorFloatingUI::EDockedTo::RightArm;
 				QuickMenuUI->SetDockedTo( DockTo );
 				QuickMenuUI->ShowUI( true );
-				MotionControllerHandInteractorThatNeedsQuickMenu->SetHasUIInFront( true );
 				MotionControllerHandInteractorThatNeedsQuickMenu->SetHasUIOnForearm( true );
 			}
 		}
@@ -932,7 +899,10 @@ void UVREditorUISystem::Tick( FEditorViewportClient* ViewportClient, const float
 				break;
 			}
 		}
-		VRButton.ButtonBorder->SetRenderTransform(FSlateRenderTransform::FTransform2D(VRButton.CurrentScale));	
+		if (VRButton.ButtonWidget != nullptr)
+		{
+			VRButton.ButtonWidget->SetRelativeScale3D(VRButton.CurrentScale*VRButton.OriginalRelativeScale*(GetOwner().GetWorldScaleFactor()));
+		}
 	}
 	
 	// Tick all of our floating UIs
@@ -943,6 +913,7 @@ void UVREditorUISystem::Tick( FEditorViewportClient* ViewportClient, const float
 			FloatingUIPtr->TickManually( DeltaTime );
 		}
 	}
+	QuickRadialMenu->TickManually(DeltaTime);
 }
 
 void UVREditorUISystem::Render( const FSceneView* SceneView, FViewport* Viewport, FPrimitiveDrawInterface* PDI )
@@ -959,27 +930,35 @@ void UVREditorUISystem::CreateUIs()
 		const bool bWithSceneComponent = false;
 		// @todo vreditor: Tweak
 		{
-			QuickMenuUI = UViewportWorldInteraction::SpawnTransientSceneActor< AVREditorFloatingUI >(GetOwner().GetWorld(), TEXT("QuickMenu"), bWithSceneComponent);
+			QuickMenuUI = GetOwner().SpawnTransientSceneActor<AVREditorFloatingUI>(TEXT("QuickMenu"), bWithSceneComponent);
 			QuickMenuUI->SetSlateWidget(*this, BuildQuickMenuWidget(), Resolution, 40.0f, AVREditorFloatingUI::EDockedTo::Nothing);
 			QuickMenuUI->ShowUI(false);
-			QuickMenuUI->SetRelativeOffset(FVector(-12.0f, 0.0f, 5.0f));
-			QuickMenuUI->GetWidgetComponent()->SetOpacityFromTexture(1.0f);
-			QuickMenuUI->GetWidgetComponent()->SetBackgroundColor(FLinearColor::Transparent);
-			QuickMenuUI->GetWidgetComponent()->SetBlendMode(EWidgetBlendMode::Masked);
+			QuickMenuUI->SetRelativeOffset(FVector(-20.0f, 0.0f, 5.0f));
+			QuickMenuUI->GetWidgetComponent()->SetBackgroundColor(FLinearColor(0.15f, 0.15f, 0.15f));
+			UStaticMesh* WindowMesh = nullptr;
+			{
+				WindowMesh = LoadObject<UStaticMesh>(nullptr, TEXT("/Engine/VREditor/UI/SM_ContentWindow_01"));
+				check(WindowMesh != nullptr);
+			}
+			QuickMenuUI->GetMeshComponent()->SetStaticMesh(WindowMesh);
 			FloatingUIs.Add(QuickMenuUI);
 
 		}
 		// Create the radial UI
 		{
-			QuickRadialMenu = UViewportWorldInteraction::SpawnTransientSceneActor< AVREditorFloatingUI >(GetOwner().GetWorld(), TEXT("QuickRadialmenu"), bWithSceneComponent);
-			QuickRadialMenu->SetSlateWidget(*this, BuildRadialMenuWidget(), DefaultResolution, 40.0f, AVREditorFloatingUI::EDockedTo::Nothing);
-			QuickRadialMenu->ShowUI(false);
-			QuickRadialMenu->SetActorEnableCollision(false);
-			QuickRadialMenu->SetRelativeOffset(FVector(0.0f, 0.0f, 5.0f));
-			QuickRadialMenu->GetWidgetComponent()->SetOpacityFromTexture(1.0f);
-			QuickRadialMenu->GetWidgetComponent()->SetBackgroundColor(FLinearColor::Transparent);
-			QuickRadialMenu->GetWidgetComponent()->SetBlendMode(EWidgetBlendMode::Masked);
-			FloatingUIs.Add(QuickRadialMenu);
+			QuickRadialMenu = GetOwner().SpawnTransientSceneActor<AVREditorRadialFloatingUI>(TEXT("QuickRadialmenu"), bWithSceneComponent);
+			FVector RelativeOffset = FVector::ZeroVector;
+			if (VRMode->GetHMDDeviceType() == EHMDDeviceType::Type::DT_SteamVR)
+			{
+				RelativeOffset = FVector(-5.0f, 0.0f, 5.0f);
+			}
+			else if (VRMode->GetHMDDeviceType() == EHMDDeviceType::Type::DT_OculusRift)
+			{
+				RelativeOffset = FVector(0.0f, 0.0f, 3.0f);
+			}
+
+			QuickRadialMenu->SetRelativeOffset(RelativeOffset);
+			QuickRadialMenu->ShowUI(false, 0.0f, false);
 		}
 	}
 	// Make some editor UIs!
@@ -996,9 +975,6 @@ void UVREditorUISystem::CreateUIs()
 			Config.InitialAssetViewType = EAssetViewType::Tile;
 			Config.bCanShowDevelopersFolder = false;
 			Config.bCanShowFolders = false;
-			Config.bUsePathPicker = false;
-			Config.bCanShowFilters = false;
-			Config.bCanShowAssetSearch = false;
 			Config.bUseSourcesView = true;
 			Config.bExpandSourcesView = true;
 			Config.ThumbnailLabel = EThumbnailLabel::NoLabel;
@@ -1007,10 +983,18 @@ void UVREditorUISystem::CreateUIs()
 			Config.bUsePathPicker = true;
 			Config.bShowAssetPathTree = true;
 			Config.bAlwaysShowCollections = false;
+			Config.bShowBottomToolbar = true;
+			Config.bCanShowLockButton = false;
+
+			if( GIsDemoMode )
+			{
+				Config.bShowAssetPathTree = false;
+				Config.bAlwaysShowCollections = true;
 
+				Config.SelectedCollectionName.Name = TEXT( "Demo" );
+				Config.SelectedCollectionName.Type = ECollectionShareType::CST_Shared;
+			}
 
-			Config.bShowBottomToolbar = false;
-			Config.bCanShowLockButton = false;
 			TSharedRef<SWidget> ContentBrowser = ContentBrowserSingleton.CreateContentBrowser( "VRContentBrowser", nullptr, &Config );
 
 			TSharedRef<SWidget> WidgetToDraw =
@@ -1022,7 +1006,7 @@ void UVREditorUISystem::CreateUIs()
 			;
 
 			const bool bWithSceneComponent = false;
-			AVREditorFloatingUI* ContentBrowserUI = UViewportWorldInteraction::SpawnTransientSceneActor< AVREditorDockableWindow >( GetOwner().GetWorld(), TEXT( "ContentBrowserUI" ), bWithSceneComponent );
+			AVREditorFloatingUI* ContentBrowserUI = GetOwner().SpawnTransientSceneActor<AVREditorDockableWindow>(TEXT( "ContentBrowserUI" ), bWithSceneComponent);
 			ContentBrowserUI->SetSlateWidget( *this, WidgetToDraw, Resolution, VREd::ContentBrowserUISize->GetFloat(), AVREditorFloatingUI::EDockedTo::Nothing );
 			ContentBrowserUI->ShowUI( false );
 			FloatingUIs.Add( ContentBrowserUI );
@@ -1041,12 +1025,16 @@ void UVREditorUISystem::CreateUIs()
 				SNew( SDPIScaler )
 				.DPIScale( VREd::EditorUIScale->GetFloat() )
 				[
-					SceneOutliner
+					SNew(SBorder)
+					.BorderImage(FEditorStyle::Get().GetBrush("ToolPanel.GroupBorder"))
+					[
+						SceneOutliner
+					]
 				]
 			;
 
 			const bool bWithSceneComponent = false;
-			AVREditorFloatingUI* WorldOutlinerUI = UViewportWorldInteraction::SpawnTransientSceneActor< AVREditorDockableWindow >( GetOwner().GetWorld(), TEXT( "WorldOutlinerUI" ), bWithSceneComponent );
+			AVREditorFloatingUI* WorldOutlinerUI = GetOwner().SpawnTransientSceneActor<AVREditorDockableWindow>(TEXT("WorldOutlinerUI"), bWithSceneComponent);
 			WorldOutlinerUI->SetSlateWidget( *this, WidgetToDraw, DefaultResolution, VREd::EditorUISize->GetFloat(), AVREditorFloatingUI::EDockedTo::Nothing );
 			WorldOutlinerUI->ShowUI( false );
 			FloatingUIs.Add( WorldOutlinerUI );
@@ -1069,7 +1057,7 @@ void UVREditorUISystem::CreateUIs()
 			;
 
 			const bool bWithSceneComponent = false;
-			AVREditorFloatingUI* ActorDetailsUI = UViewportWorldInteraction::SpawnTransientSceneActor< AVREditorDockableWindow >( GetOwner().GetWorld(), TEXT( "ActorDetailsUI" ), bWithSceneComponent );
+			AVREditorFloatingUI* ActorDetailsUI = GetOwner().SpawnTransientSceneActor<AVREditorDockableWindow>(TEXT("ActorDetailsUI"), bWithSceneComponent);
 			ActorDetailsUI->SetSlateWidget( *this, WidgetToDraw, DefaultResolution, VREd::EditorUISize->GetFloat(), AVREditorFloatingUI::EDockedTo::Nothing );
 			ActorDetailsUI->ShowUI( false );
 			FloatingUIs.Add( ActorDetailsUI );
@@ -1086,51 +1074,33 @@ void UVREditorUISystem::CreateUIs()
 				SNew( SDPIScaler )
 				.DPIScale( VREd::EditorUIScale->GetFloat() )
 				[
-					Modes
+					SNew(SBorder)
+					.BorderImage(FEditorStyle::Get().GetBrush("ToolPanel.GroupBorder"))
+					[
+						Modes
+					]
 				]
 			;
 
 			const bool bWithSceneComponent = false;
-			AVREditorFloatingUI* ModesUI = UViewportWorldInteraction::SpawnTransientSceneActor< AVREditorDockableWindow >( GetOwner().GetWorld(), TEXT( "ModesUI" ), bWithSceneComponent );
+			AVREditorFloatingUI* ModesUI = GetOwner().SpawnTransientSceneActor<AVREditorDockableWindow>(TEXT("ModesUI"), bWithSceneComponent);
 			ModesUI->SetSlateWidget( *this, WidgetToDraw, DefaultResolution, VREd::EditorUISize->GetFloat(), AVREditorFloatingUI::EDockedTo::Nothing );
 			ModesUI->ShowUI( false );
 			FloatingUIs.Add( ModesUI );
 
 			// @todo vreditor placement: This is required to force the modes UI to refresh -- otherwise it looks empty
-			GLevelEditorModeTools().ActivateMode( FBuiltinEditorModes::EM_Placement );
+			GLevelEditorModeTools().ActivateDefaultMode();
 
 			EditorUIPanels[ (int32)EEditorUIPanel::Modes ] = ModesUI;			
 		}
 
-		// Create the tutorial dockable window
 		{
-			const bool bShowAtStart = 
-				GetOwner().IsActuallyUsingVR() &&	// Don't show tutorial at startup when in "forced VR" mode, because you can't interact with UI using the mouse currently
-				(GetDefault<UEditorExperimentalSettings>()->bAlwaysShowVRTutorial); // Only show tutorial at startup if the user has this enabled
+			const FIntPoint Resolution = FIntPoint(VREd::AssetEditorUIResolutionX->GetInt(), VREd::AssetEditorUIResolutionY->GetInt());
 
-			AVREditorFloatingUI* TutorialUI = UViewportWorldInteraction::SpawnTransientSceneActor< AVREditorDockableWindow >( GetOwner().GetWorld(), TEXT( "TutorialUI" ), false );
-			TutorialUI->SetUMGWidget( *this, TutorialWidgetClass, FIntPoint( VREd::TutorialUIResolutionX->GetFloat(), VREd::TutorialUIResolutionY->GetFloat() ), VREd::TutorialUISize->GetFloat(), bShowAtStart ? AVREditorFloatingUI::EDockedTo::Room : AVREditorFloatingUI::EDockedTo::Nothing );
-
-			if( bShowAtStart )
-			{
-				TutorialUI->SetRelativeOffset( FVector( VREd::TutorialUILocationX->GetFloat(), VREd::TutorialUILocationY->GetFloat(), VREd::TutorialUILocationZ->GetFloat() ) );
-				TutorialUI->SetLocalRotation( FRotator( VREd::TutorialUIPitch->GetFloat(), VREd::TutorialUIYaw->GetFloat(), 0 ) );
-			}
-			
-			TutorialUI->ShowUI( bShowAtStart );
-
-			FloatingUIs.Add( TutorialUI );
-
-			EditorUIPanels[ (int32)EEditorUIPanel::Tutorial ] = TutorialUI;		
-		}
-
-		{
-			const FIntPoint Resolution = FIntPoint(VREd::AssetEditorUIResolutionX->GetInt(), VREd::AssetEditorUIResolutionY->GetInt());
-
-			const bool bWithSceneComponent = false;
-			AVREditorDockableWindow* TabManagerUI = UViewportWorldInteraction::SpawnTransientSceneActor< AVREditorDockableWindow >( GetOwner().GetWorld(), TEXT( "AssetEditor" ), bWithSceneComponent );
-			TabManagerUI->SetSlateWidget( *this, SNullWidget::NullWidget, Resolution, VREd::EditorUISize->GetFloat(), AVREditorFloatingUI::EDockedTo::Nothing );
-			TabManagerUI->ShowUI( false );
+			const bool bWithSceneComponent = false;
+			AVREditorDockableWindow* TabManagerUI = GetOwner().SpawnTransientSceneActor<AVREditorDockableWindow>(TEXT("AssetEditor"), bWithSceneComponent);
+			TabManagerUI->SetSlateWidget( *this, SNullWidget::NullWidget, Resolution, VREd::EditorUISize->GetFloat(), AVREditorFloatingUI::EDockedTo::Nothing );
+			TabManagerUI->ShowUI( false );
 
 			// @todo vreditor: Could use "Hovering" instead for better performance with many UIs, but needs to be manually refreshed in some cases
 			TabManagerUI->GetWidgetComponent()->SetDrawingPolicy( EVREditorWidgetDrawingPolicy::Always );
@@ -1143,6 +1113,7 @@ void UVREditorUISystem::CreateUIs()
 			TSharedRef<SWindow> TabManagerWindowRef = TabManagerWindow.ToSharedRef();
 			ProxyTabManager = MakeShareable( new FProxyTabmanager( TabManagerWindowRef ) );
 
+			ProxyTabManager->OnIsTabSupported.AddUObject( this, &UVREditorUISystem::IsProxyTabSupported );
 			ProxyTabManager->OnTabOpened.AddUObject( this, &UVREditorUISystem::OnProxyTabLaunched );
 			ProxyTabManager->OnAttentionDrawnToTab.AddUObject( this, &UVREditorUISystem::OnAttentionDrawnToTab );
 
@@ -1175,7 +1146,7 @@ void UVREditorUISystem::CreateUIs()
 			;
 
 			const bool bWithSceneComponent = false;
-			AVREditorFloatingUI* WorldSettingsUI = UViewportWorldInteraction::SpawnTransientSceneActor< AVREditorDockableWindow >( World, TEXT( "WorldSettingsUI" ), bWithSceneComponent );
+			AVREditorFloatingUI* WorldSettingsUI = GetOwner().SpawnTransientSceneActor<AVREditorDockableWindow>(TEXT("WorldSettingsUI"), bWithSceneComponent);
 			WorldSettingsUI->SetSlateWidget( *this, WidgetToDraw, DefaultResolution, VREd::EditorUISize->GetFloat(), AVREditorFloatingUI::EDockedTo::Nothing );
 			WorldSettingsUI->ShowUI( false );
 			FloatingUIs.Add( WorldSettingsUI );
@@ -1196,7 +1167,7 @@ void UVREditorUISystem::CreateUIs()
 			;
 
 			const bool bWithSceneComponent = false;
-			AVREditorFloatingUI* SequencerUI = UViewportWorldInteraction::SpawnTransientSceneActor< AVREditorDockableWindow >( GetOwner().GetWorld(), TEXT("SequencerUI"), bWithSceneComponent);
+			AVREditorFloatingUI* SequencerUI = GetOwner().SpawnTransientSceneActor<AVREditorDockableWindow>(TEXT("SequencerUI"), bWithSceneComponent);
 			SequencerUI->SetSlateWidget(*this, WidgetToDraw, FIntPoint(VREd::SequencerUIResolutionX->GetFloat(), VREd::SequencerUIResolutionY->GetFloat()), VREd::ContentBrowserUISize->GetFloat(), AVREditorFloatingUI::EDockedTo::Nothing);
 			SequencerUI->ShowUI(false);
 			FloatingUIs.Add(SequencerUI);
@@ -1206,7 +1177,7 @@ void UVREditorUISystem::CreateUIs()
 	}
 }
 
-void UVREditorUISystem::CleanUpActorsBeforeMapChangeOrSimulate()
+void UVREditorUISystem::CleanupActorsBeforeShutdown()
 {
 	// If we have a sequence tab open, reset its widget and close the associated Sequencer
 	if (GetOwner().GetCurrentSequencer() != nullptr)
@@ -1226,6 +1197,7 @@ void UVREditorUISystem::CleanUpActorsBeforeMapChangeOrSimulate()
 
 	FloatingUIs.Reset();
 	EditorUIPanels.Reset();
+	QuickRadialMenu->Destroy(false, false);
 	QuickRadialMenu = nullptr;
 	QuickMenuUI = nullptr;
 
@@ -1266,132 +1238,164 @@ bool UVREditorUISystem::IsShowingEditorUIPanel( const EEditorUIPanel EditorUIPan
 }
 
 
-void UVREditorUISystem::ShowEditorUIPanel( const UWidgetComponent* WidgetComponent, UVREditorInteractor* Interactor, const bool bShouldShow, const bool OnHand, const bool bRefreshQuickMenu, const bool bPlaySound )
+void UVREditorUISystem::ShowEditorUIPanel(const UWidgetComponent* WidgetComponent, UVREditorInteractor* Interactor, const bool bShouldShow, const bool bSpawnInFront, const bool bDragFromOpen, const bool bPlaySound)
 {
 	AVREditorFloatingUI* Panel = nullptr;
-	for( AVREditorFloatingUI* CurrentPanel : EditorUIPanels )
+	for (AVREditorFloatingUI* CurrentPanel : EditorUIPanels)
 	{
-		if( CurrentPanel && CurrentPanel->GetWidgetComponent() == WidgetComponent )
+		if (CurrentPanel && CurrentPanel->GetWidgetComponent() == WidgetComponent)
 		{
 			Panel = CurrentPanel;
 			break;
 		}
 	}
 
-	ShowEditorUIPanel( Panel, Interactor, bShouldShow, OnHand, bRefreshQuickMenu, bPlaySound );
+	ShowEditorUIPanel(Panel, Interactor, bShouldShow, bSpawnInFront, bDragFromOpen, bPlaySound);
 }
 
 
-void UVREditorUISystem::ShowEditorUIPanel( const EEditorUIPanel EditorUIPanel, UVREditorInteractor* Interactor, const bool bShouldShow, const bool OnHand, const bool bRefreshQuickMenu, const bool bPlaySound )
+void UVREditorUISystem::ShowEditorUIPanel(const EEditorUIPanel EditorUIPanel, UVREditorInteractor* Interactor, const bool bShouldShow, const bool bSpawnInFront, const bool bDragFromOpen, const bool bPlaySound)
 {
-	AVREditorFloatingUI* Panel = EditorUIPanels[ (int32)EditorUIPanel ];
-	ShowEditorUIPanel( Panel, Interactor, bShouldShow, OnHand, bRefreshQuickMenu, bPlaySound );
+	AVREditorFloatingUI* Panel = EditorUIPanels[(int32)EditorUIPanel];
+	ShowEditorUIPanel(Panel, Interactor, bShouldShow, bSpawnInFront, bDragFromOpen, bPlaySound);
 }
 
-void UVREditorUISystem::ShowEditorUIPanel( AVREditorFloatingUI* Panel, UVREditorInteractor* Interactor, const bool bShouldShow, const bool OnHand, const bool bRefreshQuickMenu, const bool bPlaySound )
+void UVREditorUISystem::ShowEditorUIPanel(AVREditorFloatingUI* Panel, UVREditorInteractor* Interactor, const bool bShouldShow, const bool bSpawnInFront, const bool bDragFromOpen, const bool bPlaySound)
 {
 	if( Panel != nullptr )
 	{
 		AVREditorFloatingUI::EDockedTo DockedTo = Panel->GetDockedTo();
-		if( OnHand || DockedTo == AVREditorFloatingUI::EDockedTo::Nothing )
+		if(bShouldShow && (bSpawnInFront || DockedTo == AVREditorFloatingUI::EDockedTo::Nothing || bDragFromOpen))
 		{
-			// Hide any panels that are already shown on this hand
-			for( int32 PanelIndex = 0; PanelIndex < (int32)EEditorUIPanel::TotalCount; ++PanelIndex )
+			Panel->SetScale(Panel->GetInitialScale(), false);
+
+			// Make sure to set the panel to dock to the room if it is docked to nothing. Otherwise it will show up in world space.
+			if (DockedTo == AVREditorFloatingUI::EDockedTo::Nothing)
 			{
-				AVREditorFloatingUI* OtherPanel = EditorUIPanels[PanelIndex];
-				if( OtherPanel != nullptr && OtherPanel->IsUIVisible() && OtherPanel->GetDockedTo() != AVREditorFloatingUI::EDockedTo::Room )
-				{
-					UVREditorInteractor* OtherPanelDockToVREditorInteractor = GetOwner().GetHandInteractor( OtherPanel->GetDockedTo() == AVREditorFloatingUI::EDockedTo::LeftHand ? EControllerHand::Left : EControllerHand::Right );
-					if( OtherPanelDockToVREditorInteractor == Interactor )
-					{
-						OtherPanel->ShowUI( false );
-						OtherPanel->SetDockedTo( AVREditorFloatingUI::EDockedTo::Nothing );
-						OtherPanelDockToVREditorInteractor->SetHasUIInFront( false );
-					}
-				}
+				Panel->SetDockedTo(AVREditorBaseActor::EDockedTo::Room);
 			}
-			
-			UVREditorMotionControllerInteractor* MotionControllerInteractor = Cast<UVREditorMotionControllerInteractor>( Interactor );
-			if ( MotionControllerInteractor )
+
+			// Set the initial transform when opening a panel.
+			FTransform HandTransform;
+			FVector HandForward;
+			Interactor->GetTransformAndForwardVector(HandTransform, HandForward);
+
+			const float OpenDistance = VREd::UIPanelOpenDistance->GetFloat() * GetOwner().GetWorldScaleFactor();
+			FTransform WindowToWorldTransform;
+			WindowToWorldTransform.SetLocation(HandTransform.GetLocation() + (HandForward.GetSafeNormal() * OpenDistance));
+
+			const FQuat Rotation = (HandTransform.GetLocation() - WindowToWorldTransform.GetLocation()).ToOrientationQuat() * FRotator(VREd::UIPanelOpenRotationPitchOffset->GetFloat(), 0.f, 0.f).Quaternion();
+			WindowToWorldTransform.SetRotation(Rotation);
+
+			if (bDragFromOpen)
 			{
-				const AVREditorFloatingUI::EDockedTo NewDockedTo = MotionControllerInteractor->GetControllerSide() == EControllerHand::Left ? AVREditorFloatingUI::EDockedTo::LeftHand : AVREditorFloatingUI::EDockedTo::RightHand;
-				Panel->SetDockedTo( NewDockedTo );
+				AVREditorDockableWindow* DockableWindow = Cast<AVREditorDockableWindow>(Panel);
+				if (DockableWindow != nullptr)
+				{	
+					bDragPanelFromOpen = true;
+					DragPanelFromOpenTime = 0.0f;
+					bPanelCanScale = false;
+
+					DockableWindow->SetActorTransform(WindowToWorldTransform);
+					DockableWindow->SetDockSelectDistance(OpenDistance);
+
+					const bool bPlayStartDragSound = false;
+					StartDraggingDockUI(DockableWindow, Interactor, OpenDistance, bPlayStartDragSound);
+					VRMode->GetWorldInteraction().SetDraggedInteractable(DockableWindow, Interactor);
+				}
 			}
-			
-			if ( bShouldShow )
+			else
 			{
-				Panel->SetScale( Panel->GetInitialScale() );
-			}
+				if (DockedTo != AVREditorFloatingUI::EDockedTo::Room)
+				{
+					Panel->SetDockedTo(AVREditorBaseActor::EDockedTo::Room);
+				}
 
-			const FVector EditorUIRelativeOffset( Panel->GetSize().Y * 0.5f + VREd::EditorUIOffsetFromHand->GetFloat(), 0.0f, 0.0f );
-			Panel->SetRelativeOffset( EditorUIRelativeOffset );
-			Panel->SetLocalRotation( FRotator( 90.0f, 180.0f, 0.0f ) ); // @todo vreditor: needs initial rotation
+				const FTransform RoomToWorld = GetOwner().GetRoomTransform();
+				const FTransform WorldToRoom = RoomToWorld.Inverse();
+				const FTransform WindowToRoomTransform = WindowToWorldTransform * WorldToRoom;
+				const FVector RoomSpaceWindowLocation = WindowToRoomTransform.GetLocation() / GetOwner().GetWorldScaleFactor();
+				const FQuat RoomSpaceWindowRotation = WindowToRoomTransform.GetRotation();
+				Panel->SetRelativeOffset(RoomSpaceWindowLocation);
+				Panel->SetLocalRotation(RoomSpaceWindowRotation.Rotator());
+			}
 		}
 
-		// If we are closing the sequencer panel, then also null out the sequencer widget and close the Sequencer instance
-		if (!bShouldShow && Panel == EditorUIPanels[(int32)EEditorUIPanel::SequencerUI] && GetOwner().GetCurrentSequencer() != nullptr)
+		// Handle special cases for closing a panel.
+		if (!bShouldShow)
 		{
-			Panel->SetSlateWidget(*this, SNullWidget::NullWidget, FIntPoint(VREd::SequencerUIResolutionX->GetFloat(), VREd::SequencerUIResolutionY->GetFloat()), VREd::EditorUISize->GetFloat(), AVREditorFloatingUI::EDockedTo::Nothing);
-			FVREditorActionCallbacks::CloseSequencer(GetOwner().GetCurrentSequencer()->GetRootMovieSceneSequence());
+			// If we are closing the sequencer panel, then also null out the sequencer widget and close the Sequencer instance
+			if (Panel == EditorUIPanels[(int32)EEditorUIPanel::SequencerUI])
+			{
+				if (GetOwner().GetCurrentSequencer() != nullptr)
+				{
+					Panel->SetSlateWidget(*this, SNullWidget::NullWidget, FIntPoint(VREd::SequencerUIResolutionX->GetFloat(), VREd::SequencerUIResolutionY->GetFloat()), VREd::EditorUISize->GetFloat(), AVREditorFloatingUI::EDockedTo::Nothing);
+					FVREditorActionCallbacks::CloseSequencer(GetOwner().GetCurrentSequencer()->GetRootMovieSceneSequence());
+				}
+			}
+			else if(Panel == EditorUIPanels[ (int32)EEditorUIPanel::Modes ] )
+			{
+				// Quit active mode and go back to Placement Mode when closing the Modes panel.
+				GLevelEditorModeTools().DeactivateAllModes();
+			}
 		}
 
 		Panel->ShowUI(bShouldShow);
 
-		if ( Panel->GetDockedTo() == AVREditorFloatingUI::EDockedTo::LeftHand || Panel->GetDockedTo() == AVREditorFloatingUI::EDockedTo::RightHand )
-		{
-			UVREditorInteractor* PanelDockedToVREditorInteractor = GetOwner().GetHandInteractor( Panel->GetDockedTo() == AVREditorFloatingUI::EDockedTo::LeftHand ? EControllerHand::Left : EControllerHand::Right );
-			PanelDockedToVREditorInteractor->SetHasUIInFront( bShouldShow );
-		}
-
 		if( bPlaySound )
 		{
-			UGameplayStatics::PlaySound2D( GetOwner().GetWorld(), bShouldShow ? ShowUISound : HideUISound );
+			const UVREditorAssetContainer& AssetContainer = VRMode->GetAssetContainer();
+			VRMode->PlaySound(bShouldShow ? AssetContainer.DockableWindowOpenSound : AssetContainer.DockableWindowCloseSound, Panel->GetActorLocation());
 		}
-
 	}
 }
 
 
 bool UVREditorUISystem::IsShowingRadialMenu(const UVREditorInteractor* Interactor ) const
 {
-	const EControllerHand DockedToHand = QuickRadialMenu->GetDockedTo() == AVREditorFloatingUI::EDockedTo::LeftArm ? EControllerHand::Left : EControllerHand::Right;
-	const UVREditorInteractor* DockedToInteractor = GetOwner().GetHandInteractor( DockedToHand );
-	return DockedToInteractor == Interactor && !QuickRadialMenu->bHidden;
+	if (QuickRadialMenu != nullptr)
+	{
+		const EControllerHand DockedToHand = QuickRadialMenu->GetDockedTo() == AVREditorFloatingUI::EDockedTo::LeftArm ? EControllerHand::Left : EControllerHand::Right;
+		const UVREditorInteractor* DockedToInteractor = GetOwner().GetHandInteractor(DockedToHand);
+		return DockedToInteractor == Interactor && !QuickRadialMenu->bHidden;
+	}
+	else
+	{
+		return false;
+	}
 }
 
 
-
-
-
-void UVREditorUISystem::TryToSpawnRadialMenu( UVREditorInteractor* Interactor, const bool bForceOverUI )
+void UVREditorUISystem::TryToSpawnRadialMenu( UVREditorInteractor* Interactor, const bool bForceRefresh, const bool bPlaySound /*= true*/ )
 {
 	if( Interactor != nullptr )
 	{
 		UVREditorMotionControllerInteractor* MotionControllerInteractor = Cast<UVREditorMotionControllerInteractor>( Interactor );
 		if( MotionControllerInteractor != nullptr )
 		{
-			UVREditorInteractor* DockedToInteractor = nullptr;
-			if( !QuickRadialMenu->bHidden )
-			{
-				const EControllerHand ControllerHand = QuickRadialMenu->GetDockedTo() == AVREditorFloatingUI::EDockedTo::LeftArm ? EControllerHand::Left : EControllerHand::Right;
-				DockedToInteractor = GetOwner().GetHandInteractor( ControllerHand );
-			}
 
 			EViewportInteractionDraggingMode DraggingMode = Interactor->GetDraggingMode();
 
-			bool bNeedsSpawn =
-				( QuickRadialMenu->bHidden || DockedToInteractor != Interactor ) &&
+			bool bNeedsSpawn = (bForceRefresh || !QuickRadialMenu->IsUIVisible()) &&
+				MotionControllerInteractor == UIInteractor && 
 				DraggingMode != EViewportInteractionDraggingMode::TransformablesAtLaserImpact &&	 // Don't show radial menu if the hand is busy dragging something around
 				DraggingMode != EViewportInteractionDraggingMode::TransformablesFreely &&
-				DraggingMode != EViewportInteractionDraggingMode::TransformablesWithGizmo &&
 				DraggingMode != EViewportInteractionDraggingMode::World &&
 				DraggingMode != EViewportInteractionDraggingMode::AssistingDrag &&
-				( InteractorDraggingUI == nullptr || InteractorDraggingUI != Interactor ) &&
-				(!Interactor->IsHoveringOverUI() || bForceOverUI);	// Don't show radial menu when aiming at a UI unless forcing it (too much clutter)
+				bDragPanelFromOpen == false;
 
 			if (bNeedsSpawn)
 			{
-				const AVREditorFloatingUI::EDockedTo DockedTo = MotionControllerInteractor->GetControllerSide() == EControllerHand::Left ? AVREditorFloatingUI::EDockedTo::LeftArm : AVREditorFloatingUI::EDockedTo::RightArm;
+				RadialMenuModifierSpawnTime = FTimespan::FromSeconds(FPlatformTime::Seconds());
+				if(bRadialMenuIsNumpad)
+				{ 
+					BuildNumPadWidget();
+				}
+				else
+				{
+					BuildRadialMenuWidget();
+				}
+				const AVREditorFloatingUI::EDockedTo DockedTo = UIInteractor->GetControllerSide() == EControllerHand::Left ? AVREditorFloatingUI::EDockedTo::LeftArm : AVREditorFloatingUI::EDockedTo::RightArm;
 				QuickRadialMenu->SetDockedTo(DockedTo);
 				QuickRadialMenu->ShowUI(true);
 			}
@@ -1400,12 +1404,12 @@ void UVREditorUISystem::TryToSpawnRadialMenu( UVREditorInteractor* Interactor, c
 }
 
 
-void UVREditorUISystem::HideRadialMenu( UVREditorInteractor* Interactor )
+void UVREditorUISystem::HideRadialMenu( UVREditorInteractor* Interactor, const bool bPlaySound /*= true*/ )
 {
 	// Only hide the radial menu if the passed interactor is actually the interactor with the radial menu
 	if( IsShowingRadialMenu( Interactor ) )
 	{
-		QuickRadialMenu->ShowUI( false, true, VREd::RadialMenuFadeDelay->GetFloat() );
+		QuickRadialMenu->ShowUI( false, true, VREd::RadialMenuFadeDelay->GetFloat(), bPlaySound );
 	}
 }
 
@@ -1415,54 +1419,36 @@ FTransform UVREditorUISystem::MakeDockableUITransformOnLaser( AVREditorDockableW
 	FTransform HandTransform;
 	FVector HandForward;
 	Interactor->GetTransformAndForwardVector( HandTransform, HandForward );
-	const FTransform InteractorTransform = Interactor->GetTransform();
+	FTransform InteractorTransform = Interactor->GetTransform();
 
-	const FVector NewLocation = InteractorTransform.GetLocation() + ( HandForward.GetSafeNormal() * DockSelectDistance );
-	FRotator NewRotation = ( InteractorTransform.GetLocation() - NewLocation ).ToOrientationRotator();
-	NewRotation.Roll = -InteractorTransform.GetRotation().Rotator().Roll;
+	// Use the smoothed laser pointer direction for computing an offset direction.  It looks a lot better!
+	FVector LaserPointerStart, LaserPointerEnd;
+	const bool bEvenIfBlocked = true;
+	Interactor->GetLaserPointer( /* Out */ LaserPointerStart, /* Out */ LaserPointerEnd, bEvenIfBlocked );
 
-	const FTransform LaserImpactToWorld( NewRotation, NewLocation );
+	const FVector LaserPointerDirection = ( LaserPointerEnd - LaserPointerStart ).GetSafeNormal();
+	const FVector NewLocation = InteractorTransform.GetLocation() + ( LaserPointerDirection * DockSelectDistance );
+
+	const FTransform LaserImpactToWorld( InteractorTransform.GetRotation(), NewLocation );
 	return LaserImpactToWorld;
 }
 
 
 FTransform UVREditorUISystem::MakeDockableUITransform( AVREditorDockableWindow* InitDraggingDockUI, UVREditorInteractor* Interactor, const float DockSelectDistance )
 {
-	const FTransform UIOnLaserToWorld = MakeDockableUITransformOnLaser( DraggingUI, Interactor, DockSelectDistance );
-	const FTransform UIToUIOnLaser = DraggingUIOffsetTransform;
-	
-	const FTransform UIToWorld = UIToUIOnLaser * UIOnLaserToWorld;
+	const FTransform LaserImpactToWorld = MakeDockableUITransformOnLaser( DraggingUI, Interactor, DockSelectDistance );
+	const FTransform UIToLaserImpact = DraggingUIOffsetTransform;
+	const FTransform UIToWorld = UIToLaserImpact * LaserImpactToWorld;
 
-	if( DraggingUI != nullptr &&  DraggingUI->GetDockedTo() == AVREditorFloatingUI::EDockedTo::Dragging && InteractorDraggingUI && Interactor == InteractorDraggingUI )
+	// Enable the ability to scale the UI when opening from the radial menu after the trackpad was reset.
+	if (!bPanelCanScale && bDragPanelFromOpen && UIInteractor != nullptr && DraggingUI != nullptr && InteractorDraggingUI != nullptr && UIInteractor == InteractorDraggingUI)
 	{
-		UVREditorMotionControllerInteractor* OtherInteractor = Cast<UVREditorMotionControllerInteractor>( InteractorDraggingUI->GetOtherInteractor() );
-		if ( OtherInteractor )
+		UVREditorMotionControllerInteractor* MotionController = Cast<UVREditorMotionControllerInteractor>(InteractorDraggingUI);
+		if (MotionController != nullptr)
 		{
-			const float WorldScaleFactor = GetOwner().GetWorldScaleFactor();
-			const FVector DraggingUILocation = InteractorDraggingUI->GetHoverLocation();
-			const FVector OtherHandLocation = OtherInteractor->GetTransform().GetLocation();
-			LastDraggingHoverLocation = UIOnLaserToWorld.GetLocation();
-
-			const float Distance = FVector::Dist( UIOnLaserToWorld.GetLocation(), OtherHandLocation ) / WorldScaleFactor;
-
-			// if dragged passed the threshold since starting dragging
-			if ( Distance > VREd::MinDockDragDistance->GetFloat() && !bDraggedDockFromHandPassedThreshold &&
-				( DraggingUI->GetPreviouslyDockedTo() == AVREditorFloatingUI::EDockedTo::LeftHand || DraggingUI->GetPreviouslyDockedTo() == AVREditorFloatingUI::EDockedTo::RightHand ) )
-			{
-				bDraggedDockFromHandPassedThreshold = true;
-			}
-
-			// Snapping to a hand when in range
-			if ( Distance <= VREd::MinDockDragDistance->GetFloat() && ( bDraggedDockFromHandPassedThreshold || DraggingUI->GetPreviouslyDockedTo() == AVREditorFloatingUI::EDockedTo::Room ) )
+			if (MotionController->GetTrackpadPosition().IsNearlyZero(0.1f))
 			{
-				const FVector EditorUIRelativeOffset( DraggingUI->GetSize().Y * 0.5f + VREd::EditorUIOffsetFromHand->GetFloat(), 0.0f, 0.0f );
-				const FTransform MovingUIGoalTransform = DraggingUI->MakeUITransformLockedToHand( OtherInteractor, false, EditorUIRelativeOffset, FRotator( 90.0f, 180.0f, 0.0f ) );
-				const AVREditorFloatingUI::EDockedTo NewDockedTo = OtherInteractor->GetControllerSide() == EControllerHand::Left ? AVREditorFloatingUI::EDockedTo::LeftHand : AVREditorFloatingUI::EDockedTo::RightHand;
-				DraggingUI->MoveTo( MovingUIGoalTransform, VREd::DockDragSpeed->GetFloat(), NewDockedTo );
-			}
-			else
-			{
-				DraggingUI->MoveTo( UIToWorld, VREd::DockDragSpeed->GetFloat(), AVREditorFloatingUI::EDockedTo::Room );
+				bPanelCanScale = true;
 			}
 		}
 	}
@@ -1471,42 +1457,25 @@ FTransform UVREditorUISystem::MakeDockableUITransform( AVREditorDockableWindow*
 }
 
 
-void UVREditorUISystem::StartDraggingDockUI( AVREditorDockableWindow* InitDraggingDockUI, UVREditorInteractor* Interactor, const float DockSelectDistance )
+void UVREditorUISystem::StartDraggingDockUI(AVREditorDockableWindow* InitDraggingDockUI, UVREditorInteractor* Interactor, const float DockSelectDistance, const bool bPlaySound /*= true*/)
 {
-	AVREditorFloatingUI::EDockedTo DockTo = InitDraggingDockUI->GetDockedTo();
-	if( DockTo == AVREditorFloatingUI::EDockedTo::LeftHand || DockTo == AVREditorFloatingUI::EDockedTo::RightHand )
-	{
-		bDraggedDockFromHandPassedThreshold = false;
-
-		// Set the laser of the interactor the dockable window was pulled from
-		UViewportInteractor* OtherInteractor = Interactor->GetOtherInteractor();
-		if (OtherInteractor)
-		{
-			UVREditorInteractor* OtherVREditorInteractor = Cast<UVREditorInteractor>(OtherInteractor);
-			if (OtherVREditorInteractor)
-			{
-				OtherVREditorInteractor->SetHasUIInFront(false);
-			}
-		}
-	}
-
-	
-
 	InteractorDraggingUI = Interactor;
-
 	FTransform UIToWorld = InitDraggingDockUI->GetActorTransform();
 	UIToWorld.SetScale3D( FVector( 1.0f ) );
 	const FTransform WorldToUI = UIToWorld.Inverse();
 
-	const FTransform UIOnLaserToWorld = MakeDockableUITransformOnLaser( InitDraggingDockUI, Interactor, DockSelectDistance );
-	const FTransform UIOnLaserToUI = UIOnLaserToWorld * WorldToUI;
-	const FTransform UIToUIOnLaser = UIOnLaserToUI.Inverse();
-	DraggingUIOffsetTransform = UIToUIOnLaser;
+	const FTransform LaserImpactToWorld = MakeDockableUITransformOnLaser( InitDraggingDockUI, Interactor, DockSelectDistance );
+	const FTransform LaserImpactToUI = LaserImpactToWorld * WorldToUI;
+	const FTransform UIToLaserImpact = LaserImpactToUI.Inverse();
+	DraggingUIOffsetTransform = UIToLaserImpact;
 
 	DraggingUI = InitDraggingDockUI;
 	DraggingUI->SetDockedTo( AVREditorFloatingUI::EDockedTo::Dragging );
 
-	UGameplayStatics::PlaySound2D( GetOwner().GetWorld(), StartDragUISound );
+	if (bPlaySound)
+	{
+		VRMode->PlaySound(VRMode->GetAssetContainer().DockableWindowDragSound, LaserImpactToWorld.GetLocation()); //@todo VREditor: Removed sounds here and put in DockableWindow
+	}
 }
 
 AVREditorDockableWindow* UVREditorUISystem::GetDraggingDockUI() const
@@ -1514,7 +1483,7 @@ AVREditorDockableWindow* UVREditorUISystem::GetDraggingDockUI() const
 	return DraggingUI;
 }
 
-void UVREditorUISystem::StopDraggingDockUI( UVREditorInteractor* VREditorInteractor )
+void UVREditorUISystem::StopDraggingDockUI( UVREditorInteractor* VREditorInteractor, const bool bPlaySound /*= true*/ )
 {
 	if ( IsInteractorDraggingDockUI( VREditorInteractor ) )
 	{
@@ -1525,28 +1494,19 @@ void UVREditorUISystem::StopDraggingDockUI( UVREditorInteractor* VREditorInterac
 			UVREditorMotionControllerInteractor* MotionControllerOtherInteractor = Cast<UVREditorMotionControllerInteractor>( OtherInteractor );
 			if ( MotionControllerOtherInteractor )
 			{
-				bool bOnHand = false;
-				const float Distance = FVector::Dist( LastDraggingHoverLocation, MotionControllerOtherInteractor->GetTransform().GetLocation() ) / GetOwner().GetWorldScaleFactor();
-				if ( Distance > VREd::MinDockDragDistance->GetFloat() )
-				{
-					DraggingUI->SetDockedTo( AVREditorFloatingUI::EDockedTo::Room );
-				}
-				else
+				// Reset the the panel that we were dragging to room space.
+				DraggingUI->SetDockedTo( AVREditorFloatingUI::EDockedTo::Room );
+
+				// Play stop dragging sound.
+				if (bPlaySound)
 				{
-					bOnHand = true;
-					const FVector EditorUIRelativeOffset( DraggingUI->GetSize().Y * 0.5f + VREd::EditorUIOffsetFromHand->GetFloat(), 0.0f, 0.0f );
-					const FTransform MovingUIGoalTransform = DraggingUI->MakeUITransformLockedToHand( MotionControllerOtherInteractor, false, EditorUIRelativeOffset, FRotator( 90.0f, 180.0f, 0.0f ) );
-					const AVREditorFloatingUI::EDockedTo NewDockedTo = MotionControllerOtherInteractor->GetControllerSide() == EControllerHand::Left ? AVREditorFloatingUI::EDockedTo::LeftHand : AVREditorFloatingUI::EDockedTo::RightHand;
-					DraggingUI->MoveTo( MovingUIGoalTransform, VREd::DockSnapTime->GetFloat(), NewDockedTo );
+					VRMode->PlaySound(VRMode->GetAssetContainer().DockableWindowDropSound, DraggingUI->GetActorLocation());
 				}
 
-				ShowEditorUIPanel( DraggingUI, MotionControllerOtherInteractor, true, bOnHand );
-
+				// We are not dragging anymore.
 				DraggingUI = nullptr;
 				InteractorDraggingUI->SetDraggingMode( EViewportInteractionDraggingMode::Nothing );
 				InteractorDraggingUI = nullptr;
-
-				UGameplayStatics::PlaySound2D( GetOwner().GetWorld(), StopDragUISound );
 			}
 		}
 	}
@@ -1562,10 +1522,14 @@ bool UVREditorUISystem::IsInteractorDraggingDockUI( const UVREditorInteractor* I
 	return InteractorDraggingUI && InteractorDraggingUI == Interactor;
 }
 
+bool UVREditorUISystem::IsDraggingPanelFromOpen() const
+{
+	return bDragPanelFromOpen;
+}
+
 void UVREditorUISystem::TogglePanelsVisibility()
 {
 	bool bAnyPanelsVisible = false;
-	AVREditorFloatingUI* PanelOnHand = nullptr;
 
 	// Check if there is any panel visible and if any is docked to a hand
 	for (AVREditorFloatingUI* Panel : EditorUIPanels)
@@ -1573,74 +1537,28 @@ void UVREditorUISystem::TogglePanelsVisibility()
 		if (Panel && Panel->IsUIVisible())
 		{
 			bAnyPanelsVisible = true;
-
-			const AVREditorFloatingUI::EDockedTo DockedTo = Panel->GetDockedTo();
-			if (DockedTo == AVREditorFloatingUI::EDockedTo::LeftHand || DockedTo == AVREditorFloatingUI::EDockedTo::RightHand)
-			{
-				PanelOnHand = Panel;
-				break;
-			}
+			break;
 		}
 	}
 
 	// Hide if there is any UI visible
 	const bool bShowUI = !bAnyPanelsVisible;
-	
-	UGameplayStatics::PlaySound2D( GetOwner().GetWorld(), bShowUI ? ShowUISound : HideUISound );
 
-	if ( bSetDefaultLayout)
-	{
-		bSetDefaultLayout = false;
-		SetDefaultWindowLayout();
-	}
-	else
+	for (AVREditorFloatingUI* Panel : EditorUIPanels)
 	{
-		for ( AVREditorFloatingUI* Panel : EditorUIPanels )
+		if (Panel != nullptr && Panel->IsUIVisible() != bShowUI)
 		{
-			if ( Panel != nullptr && Panel->IsUIVisible() != bShowUI )
-			{
-				if ( bShowUI )
-				{
-					bool SetNewVisibility = true;
-					const AVREditorFloatingUI::EDockedTo DockedTo = Panel->GetDockedTo();
-					if ( DockedTo == AVREditorFloatingUI::EDockedTo::LeftHand || DockedTo == AVREditorFloatingUI::EDockedTo::RightHand )
-					{
-						if ( !PanelOnHand )
-						{
-							UVREditorInteractor* VREditorInteractor = GetOwner().GetHandInteractor( Panel->GetDockedTo() == AVREditorFloatingUI::EDockedTo::LeftHand ? EControllerHand::Left : EControllerHand::Right );
-							VREditorInteractor->SetHasUIInFront( true );
-							PanelOnHand = Panel;
-						}
-						else
-						{
-							SetNewVisibility = false;
-						}
-					}
-					else if ( DockedTo == AVREditorFloatingUI::EDockedTo::Nothing )
-					{
-						SetNewVisibility = false;
-					}
-
-					if ( SetNewVisibility )
-					{
-						Panel->ShowUI( true );
-					}
-				}
-				else
-				{
-					Panel->ShowUI( false );
-				}
-			}
+			Panel->ShowUI(bShowUI);
 		}
 	}
 
-	if( !bShowUI && PanelOnHand )
+	if (UIInteractor)
 	{
-		UVREditorInteractor* VREditorInteractor = GetOwner().GetHandInteractor( PanelOnHand->GetDockedTo() == AVREditorFloatingUI::EDockedTo::LeftHand ? EControllerHand::Left : EControllerHand::Right );
-		VREditorInteractor->SetHasUIInFront( false );
-	}
-
+		// Play sound
+		const UVREditorAssetContainer& AssetContainer = VRMode->GetAssetContainer();
+		VRMode->PlaySound(bShowUI ? AssetContainer.DockableWindowOpenSound : AssetContainer.DockableWindowCloseSound, UIInteractor->GetTransform().GetLocation());
 	}
+}
 
 float UVREditorUISystem::GetMaxDockWindowSize() const
 {
@@ -1652,6 +1570,16 @@ float UVREditorUISystem::GetMinDockWindowSize() const
 	return VREd::MinDockWindowSize->GetFloat();
 }
 
+void UVREditorUISystem::IsProxyTabSupported( FTabId TabId, bool& bOutIsTabSupported )
+{
+	// Avoid spawning GPU Profiler tab into VR
+	static FName VisualizerSpawnPointTabId( "VisualizerSpawnPoint" );
+	if( TabId.TabType == VisualizerSpawnPointTabId )
+	{
+		bOutIsTabSupported = false;
+	}
+}
+
 void UVREditorUISystem::OnProxyTabLaunched(TSharedPtr<SDockTab> NewTab)
 {
 	ShowAssetEditor();
@@ -1674,21 +1602,13 @@ void UVREditorUISystem::ShowAssetEditor()
 	// A tab was opened, so make sure the "Asset" UI is visible.  That's where the user can interact
 	// with the newly-opened tab
 	if ( !IsShowingEditorUIPanel(EEditorUIPanel::AssetEditor) )
-	{
-		// Always spawn on a hand.  But which hand?  Well, we'll choose the hand that isn't actively clicking on something using a laser.
-		UVREditorInteractor* VREditorInteractor = GetOwner().GetHandInteractor( EControllerHand::Left ); // Hand that did not clicked with a laser
-		//@todo VREditor: This should include a check on VREditorInteractor->IsClickingOnUI(), however that is broken because of double click on UI.
-		if( !VREditorInteractor->HasUIInFront() )
-		{
-			VREditorInteractor = GetOwner().GetHandInteractor( EControllerHand::Right );
-		}
-		
+	{	
 		const bool bShouldShow = true;
-		const bool bShowOnHand = true;
-		ShowEditorUIPanel( EEditorUIPanel::AssetEditor, VREditorInteractor, bShouldShow, bShowOnHand );
+		const bool bSpawnInFront = true;
+		ShowEditorUIPanel( EEditorUIPanel::AssetEditor, UIInteractor, bShouldShow, bSpawnInFront );
 
 		// Play haptic effect so user knows to look at their hand that now has UI on it!
-		VREditorInteractor->PlayHapticEffect( VREd::UIAssetEditorSummonedOnHandHapticFeedbackStrength->GetFloat() );
+		UIInteractor->PlayHapticEffect( VREd::UIAssetEditorSummonedOnHandHapticFeedbackStrength->GetFloat() );
 	}
 }
 
@@ -1698,50 +1618,20 @@ void UVREditorUISystem::TogglePanelVisibility( const EEditorUIPanel EditorUIPane
 	if (Panel != nullptr)
 	{
 		const bool bIsShowing = Panel->IsUIVisible();
-		UVREditorInteractor* InteractorWithQuickMenu = GetOwner().GetHandInteractor( QuickMenuUI->GetDockedTo() == AVREditorFloatingUI::EDockedTo::LeftArm ? EControllerHand::Left : EControllerHand::Right );
-		if( Panel->GetDockedTo() == AVREditorFloatingUI::EDockedTo::Room )
+		if(bIsShowing)
 		{
-			Panel->ShowUI( false );
-			Panel->SetDockedTo( AVREditorFloatingUI::EDockedTo::Nothing );
+			Panel->ShowUI(false);
 		}
 		else
 		{
-			ShowEditorUIPanel( Panel, InteractorWithQuickMenu, !bIsShowing, true );
+			const bool bDragFromOpen = ShouldPreviewPanel();
+			const bool bSpawnInFront = true;
+			const bool bPlaySound = false;
+			ShowEditorUIPanel(Panel, UIInteractor, !bIsShowing, bSpawnInFront, bDragFromOpen, bPlaySound);
 		}
 	}
 }
 
-void UVREditorUISystem::SetDefaultWindowLayout()
-{
-	for ( int32 PanelIndex = 0; PanelIndex < DefaultWindowTransforms.Num(); ++PanelIndex )
-	{
-		AVREditorFloatingUI* Panel = EditorUIPanels[PanelIndex];
-		if ( Panel )
-		{
-			const AVREditorFloatingUI::EDockedTo DockedTo = Panel->GetDockedTo();
-			if (DockedTo == AVREditorFloatingUI::EDockedTo::LeftHand || DockedTo == AVREditorFloatingUI::EDockedTo::RightHand)
-			{
-				const EControllerHand ControllerHand = DockedTo == AVREditorFloatingUI::EDockedTo::LeftHand ? EControllerHand::Left : EControllerHand::Right;
-				UVREditorInteractor* Interactor = GetOwner().GetHandInteractor( ControllerHand );
-				Interactor->SetHasUIInFront( false );
-			}
-
-			Panel->SetDockedTo( AVREditorFloatingUI::EDockedTo::Room );
-			Panel->ShowUI( true );
-
-			// Make sure the UIs are centered around the direction your head is looking (yaw only!)
-			const FVector RoomSpaceHeadLocation = GetOwner().GetRoomSpaceHeadTransform().GetLocation() / GetOwner().GetWorldScaleFactor();
-			FRotator RoomSpaceHeadYawRotator = GetOwner().GetRoomSpaceHeadTransform().GetRotation().Rotator();
-			RoomSpaceHeadYawRotator.Pitch = 0.0f;
-			RoomSpaceHeadYawRotator.Roll = 0.0f;
-
-			FTransform NewTransform = DefaultWindowTransforms[PanelIndex];
-			NewTransform *= FTransform( RoomSpaceHeadYawRotator.Quaternion(), FVector::ZeroVector );
-			Panel->SetLocalRotation( NewTransform.GetRotation().Rotator() );
-			Panel->SetRelativeOffset( RoomSpaceHeadLocation + NewTransform.GetTranslation() );
-		}
-	}
-}
 void UVREditorUISystem::DockableAreaRestored()
 {
 }
@@ -1764,7 +1654,7 @@ void UVREditorUISystem::CreateVRColorPicker(const TSharedRef<SColorPicker>& Colo
 	if( ColorPickerUI == nullptr )
 	{
 		const bool bWithSceneComponent = false;
-		ColorPickerUI = UViewportWorldInteraction::SpawnTransientSceneActor< AVREditorDockableWindow >( GetOwner().GetWorld(), TEXT( "ColorPickerUI" ), bWithSceneComponent );
+		ColorPickerUI = GetOwner().SpawnTransientSceneActor<AVREditorDockableWindow>(TEXT("ColorPickerUI"), bWithSceneComponent);
 		FloatingUIs.Add( ColorPickerUI );
 		EditorUIPanels[ (int32)EEditorUIPanel::ColorPicker ] = ColorPickerUI;
 		bJustSpawned = true;
@@ -1780,7 +1670,7 @@ void UVREditorUISystem::CreateVRColorPicker(const TSharedRef<SColorPicker>& Colo
 	}
 
 	// Spawn the color picker close to the pressed UI it was opened using UI. Otherwise spawn it in front of the user.
-	UWidgetComponent* HoveredWidgetComponent = VREditorInteractor->GetHoveringOverWidgetComponent();
+	UWidgetComponent* HoveredWidgetComponent = VREditorInteractor->GetLastHoveredWidgetComponent();
 	if( HoveredWidgetComponent != nullptr )
 	{
 		FVector LaserPointerStart, LaserPointerEnd;
@@ -1828,12 +1718,12 @@ void UVREditorUISystem::CreateVRColorPicker(const TSharedRef<SColorPicker>& Colo
 
 void UVREditorUISystem::DestroyVRColorPicker()
 {
-
 	UVREditorInteractor* VREditorInteractor = nullptr;
 	const bool bShouldShow = false;
-	const bool bShowOnHand = false;
-	const bool bRefreshQuickMenu = true;
-	ShowEditorUIPanel(EEditorUIPanel::ColorPicker, VREditorInteractor, bShouldShow, bShowOnHand, bRefreshQuickMenu);
+	const bool bSpawnInFront = false;
+	const bool bDragFromOpen = false;
+	const bool bPlaySound = true;
+	ShowEditorUIPanel(EEditorUIPanel::ColorPicker, VREditorInteractor, bShouldShow, bSpawnInFront, bDragFromOpen, bPlaySound);
 }
 
 void UVREditorUISystem::MakeUniformGridMenu(const TSharedRef<FMultiBox>& MultiBox, const TSharedRef<SMultiBoxWidget>& MultiBoxWidget, int32 Columns)
@@ -1868,8 +1758,8 @@ void UVREditorUISystem::MakeUniformGridMenu(const TSharedRef<FMultiBox>& MultiBo
 					SAssignNew(TestOverlay, SOverlay)
 				];
 
-			AddHoverableButton(BlockWidget, TEXT("SMenuEntryButton"), TestOverlay);
-			SetButtonTextWrap(BlockWidget, 75.0);
+			//AddHoverableButton(BlockWidget, TEXT("SMenuEntryButton"), TestOverlay);
+			//SetButtonTextWrap(BlockWidget, 50.0);
 			Index++;
 		}
 	}
@@ -1888,56 +1778,68 @@ void UVREditorUISystem::MakeUniformGridMenu(const TSharedRef<FMultiBox>& MultiBo
 
 TSharedRef<SWidget> UVREditorUISystem::AddHoverableButton(TSharedRef<SWidget>& BlockWidget, FName ButtonType, TSharedRef<SOverlay>& TestOverlay)
 {
-	TSharedPtr<SBorder> ButtonBorder;
-	SOverlay::FOverlaySlot* BaseSlot;
-	SOverlay::FOverlaySlot* HoverSlot;
-	TestOverlay->AddSlot(0)
-	.Expose(BaseSlot)
-		[
-			SAssignNew(ButtonBorder, SBorder)
-			.BorderImage(FEditorStyle::GetBrush("ToolPanel.GroupBorder"))
-		[
-			BlockWidget
-		]
-		];
-	// Placeholder slot to put the button into when it is hovered
-	TestOverlay->AddSlot(10)
-	.Expose(HoverSlot)
-		[
-			SNew(SOverlay)
-			.Visibility(EVisibility::SelfHitTestInvisible)
-		];
 
 	TSharedRef<SWidget> TestWidget = FindWidgetOfType(BlockWidget, ButtonType);
 	if (TestWidget != SNullWidget::NullWidget)
 	{
 		TSharedRef<SButton> Button = StaticCastSharedRef<SButton>(TestWidget);
 		Button->SetRenderTransformPivot(FVector2D(0.5, 0.5));
-		Button->SetOnHovered(FSimpleDelegate::CreateUObject(this, &UVREditorUISystem::OnHoverBeginEffect, Button));
-		Button->SetOnUnhovered(FSimpleDelegate::CreateUObject(this, &UVREditorUISystem::OnHoverEndEffect, Button));
 		FSlateSound SlateButtonPressSound = FSlateSound();
-		SlateButtonPressSound.SetResourceObject(ButtonPressSound);
+
+		const UVREditorAssetContainer& AssetContainer = VRMode->GetAssetContainer();
+		SlateButtonPressSound.SetResourceObject(AssetContainer.ButtonPressSound);
 		Button->SetPressedSound(SlateButtonPressSound);
-		ButtonBorder->SetRenderTransformPivot(FVector2D(0.5, 0.5));
-		FVRButton ButtonToAdd = FVRButton(Button, BaseSlot, HoverSlot, ButtonBorder);
-		VRButtons.Add(ButtonToAdd);
 	}
 	return BlockWidget;
 }
 
-TSharedRef<SWidget> UVREditorUISystem::SetButtonTextWrap(TSharedRef<SWidget>& BlockWidget, float WrapSize)
+TSharedRef<SWidget> UVREditorUISystem::SetButtonFormatting(TSharedRef<SWidget>& BlockWidget, float WrapSize)
 {
-	TSharedRef<SWidget> TestWidget = FindWidgetOfType(BlockWidget, FName(TEXT("STextBlock")));
+	TSharedRef<SWidget> TestWidget = FindWidgetOfType(BlockWidget, FName(TEXT("SImage")));
+	if (TestWidget != SNullWidget::NullWidget)
+	{
+		TSharedRef<SImage> Image = StaticCastSharedRef<SImage>(TestWidget);
+		Image->SetRenderTransformPivot(FVector2D(0.5f, 0.5f));
+		Image->SetRenderTransform(FSlateRenderTransform::FTransform2D(4.0f));	
+	}
+
+	// Format the button text
+	TestWidget = FindWidgetOfType(BlockWidget, FName(TEXT("STextBlock")));
 	if (TestWidget != SNullWidget::NullWidget)
 	{
 		TSharedRef<STextBlock> TextBlock = StaticCastSharedRef<STextBlock>(TestWidget);
 		TextBlock->SetWrapTextAt(WrapSize);
 		TextBlock->SetJustification(ETextJustify::Center);
-	}			
+		
+		// Change the button font color based on whether the option is checked or not
+		TSharedRef<SWidget> TestCheckbox = FindWidgetOfType(BlockWidget, FName(TEXT("SCheckBox")));
+		if (TestCheckbox != SNullWidget::NullWidget)
+		{
+			TSharedRef<SCheckBox> CheckBox = StaticCastSharedRef<SCheckBox>(TestCheckbox);
+			TAttribute<FSlateColor> DynamicSelectedTextColor = TAttribute<FSlateColor>::Create(
+				TAttribute<FSlateColor>::FGetter::CreateLambda([this, CheckBox]() -> FSlateColor
+				{
+					return CheckBox->IsChecked() ? this->GetOwner().GetColor(UVREditorMode::EColors::UIColor) : FSlateColor(FLinearColor::White);
+				}
+				)
+				);
+				TAttribute<FSlateFontInfo> DynamicSelectedTextFont = TAttribute<FSlateFontInfo>::Create(
+					TAttribute<FSlateFontInfo>::FGetter::CreateLambda([CheckBox]() -> FSlateFontInfo
+				{
+					return CheckBox->IsChecked() ? FVREditorStyle::Get().GetFontStyle("VRRadialStyle.ActiveFont") : FVREditorStyle::Get().GetFontStyle("VRRadialStyle.InactiveFont");
+				}
+				)
+				);
+				TextBlock->SetColorAndOpacity(DynamicSelectedTextColor);
+				TextBlock->SetFont(DynamicSelectedTextFont);
+			}
+		
+
+	}	
 	return BlockWidget;
 }
 
-void UVREditorUISystem::MakeRadialBoxMenu(const TSharedRef<FMultiBox>& MultiBox, const TSharedRef<SMultiBoxWidget>& MultiBoxWidget, float RadiusRatioOverride)
+void UVREditorUISystem::MakeRadialBoxMenu(const TSharedRef<FMultiBox>& MultiBox, const TSharedRef<SMultiBoxWidget>& MultiBoxWidget, float RadiusRatioOverride, FName ButtonTypeOverride)
 {
 	// Grab the list of blocks, early out if there's nothing to fill the widget with
 	const TArray< TSharedRef< const FMultiBlock > >& Blocks = MultiBox->GetBlocks();
@@ -1946,42 +1848,68 @@ void UVREditorUISystem::MakeRadialBoxMenu(const TSharedRef<FMultiBox>& MultiBox,
 	{
 		return;
 	}
-	MultiBox.Get().SetStyle(&FVREditorStyle::Get(), FVREditorStyle::GetStyleSetName());
-	TSharedRef<SRadialBox> GridPanel = SNew(SRadialBox)
-		.RadiusRatio(RadiusRatioOverride);
-	GridPanel->NumEntries = Blocks.Num();
+
+	FName StyleName = NAME_None;
+	if (!QuickRadialMenu->GetCurrentMenuWidget().IsValid() || MultiBoxWidget != QuickRadialMenu->GetCurrentMenuWidget())
+	{
+		if (bRadialMenuIsNumpad)
+		{
+			const FName NumPadOverride = TEXT("SButton");
+			QuickRadialMenu->SetButtonTypeOverride(NumPadOverride);
+			StyleName = FVREditorStyle::GetNumpadStyleSetName();
+		}
+		else
+		{
+			const FName RadialMenuOverride = TEXT("SMenuEntryButton");
+			QuickRadialMenu->SetButtonTypeOverride(RadialMenuOverride);
+			StyleName = FVREditorStyle::GetSecondaryStyleSetName();
+		}
+		QuickRadialMenu->Reset();
+		QuickRadialMenu->SetCurrentMenuWidget(MultiBoxWidget);
+	}
+	MultiBox.Get().SetStyle(&FVREditorStyle::Get(), StyleName);
+
+	if (ButtonTypeOverride == FName(TEXT("SMenuEntryButton")))
+	{
+		QuickRadialMenu->SetNumberOfEntries(NumItems - 1);
+	}
+	else
+	{
+		QuickRadialMenu->SetNumberOfEntries(NumItems);
+	}
 	for (const TSharedRef<const FMultiBlock>& MultiBlock : Blocks)
 	{
 		const TSharedRef<const FMultiBlock>& Block = MultiBlock;
-		TSharedRef<SWidget> BlockWidget = Block->MakeWidget(MultiBoxWidget, EMultiBlockLocation::Middle, true)->AsWidget();
-		TSharedRef<SOverlay> TestOverlay = SNew(SOverlay);
-		
-		GridPanel->AddSlot()
-		.VAlign(VAlign_Center)
-		.HAlign(HAlign_Center)
-		[
-			SAssignNew(TestOverlay, SOverlay)
-		];
-		
-		BlockWidget = AddHoverableButton(BlockWidget, TEXT("SButton"), TestOverlay);
-		BlockWidget = SetButtonTextWrap(BlockWidget, 75.0);
+		if (Block->GetType() == EMultiBlockType::MenuEntry)
+		{
+			TSharedRef<SWidget> BlockWidget = Block->MakeWidget(MultiBoxWidget, EMultiBlockLocation::Middle, true)->AsWidget();
+			TSharedRef<SOverlay> TestOverlay = SNew(SOverlay);
 
-	}
-	RadialWidget = GridPanel;
-	MultiBoxWidget->SetContent(
-			SNew(SDPIScaler)
-			.DPIScale(3)
-			[
-				SNew(SVerticalBox)
-				+ SVerticalBox::Slot()
-				.Padding(10, 10, 10, 10)
-				.VAlign(VAlign_Fill)
-				.HAlign(HAlign_Fill)
+			BlockWidget = AddHoverableButton(BlockWidget, ButtonTypeOverride, TestOverlay);
+			BlockWidget = SetButtonFormatting(BlockWidget, 50.0f);
+			TSharedRef<SBox> RadialMenuElement = SNew(SBox);
+			RadialMenuElement->SetContent(
+				SNew(SDPIScaler)
+				.DPIScale(3)
 				[
-					GridPanel
+					SNew(SVerticalBox)
+					+ SVerticalBox::Slot()
+					.Padding(5, 5, 5, 5)
+					.VAlign(VAlign_Center)
+					.HAlign(HAlign_Center)
+					[
+						BlockWidget
+					]
 				]
-			]
-		);
+			);
+
+			AVREditorFloatingUI::EDockedTo RadialMenuInteractor = AVREditorFloatingUI::EDockedTo::LeftArm;
+			const FIntPoint DefaultResolution(VREd::DefaultRadialElementResolutionX->GetInt(), VREd::DefaultRadialElementResolutionY->GetInt());
+			QuickRadialMenu->SetSlateWidget(*this, RadialMenuElement, DefaultResolution, 40.0f, RadialMenuInteractor);
+			FVRButton ButtonToAdd = FVRButton(QuickRadialMenu->GetWidgetComponents().Last(), FVector(1.0f/25.0f));
+			VRButtons.Add(ButtonToAdd);
+		}
+	}
 }
 
 
@@ -1991,94 +1919,168 @@ TSharedRef<SWidget> UVREditorUISystem::BuildQuickMenuWidget()
 	FVREditorActionCallbacks::UpdateGizmoModeText(VRMode);
 
 	FMultiBox::FOnMakeMultiBoxBuilderOverride VREditorMenuBuilderOverride;
-	VREditorMenuBuilderOverride.BindUObject(this, &UVREditorUISystem::MakeUniformGridMenu, 3);
+
+	VREditorMenuBuilderOverride.BindUObject(this, &UVREditorUISystem::MakeUniformGridMenu, 4);
 	const TSharedRef<FExtender> MenuExtender = MakeShareable(new FExtender());
 	TSharedPtr<FUICommandList> CommandList = MakeShareable(new FUICommandList());
-	FMenuBuilder SnapMenuBuilder(false, CommandList, MenuExtender);
+	FMenuBuilder GizmoMenuBuilder (false, CommandList, MenuExtender);
 
-	// Snap menu
-	SnapMenuBuilder.AddMenuEntry(
-		LOCTEXT("ToggleTranslationSnap", "Translate"),
-		LOCTEXT("ToggleTranslationSnapTooltip", "Toggle Translation Snap"),
-		FSlateIcon(FEditorStyle::GetStyleSetName(), "EditorViewport.TranslateMode.Small"),
+
+	TAttribute<FText> DynamicGizmoModeLabel;
+	DynamicGizmoModeLabel.BindStatic(&FVREditorActionCallbacks::GetGizmoModeText);
+	GizmoMenuBuilder.AddMenuEntry(
+		DynamicGizmoModeLabel,
+		LOCTEXT("GizmoModeTooltip", "Toggle Gizmo Mode"),
+		FSlateIcon(FEditorStyle::GetStyleSetName(), "EditorViewport.TranslateRotateMode"),
 		FUIAction
 		(
-			FExecuteAction::CreateStatic(&FLevelEditorActionCallbacks::LocationGridSnap_Clicked),
-			FCanExecuteAction::CreateStatic(&FLevelEditorActionCallbacks::DefaultCanExecuteAction),
-			FGetActionCheckState::CreateStatic(&FVREditorActionCallbacks::GetTranslationSnapState)
-		),
+			FExecuteAction::CreateStatic(&FVREditorActionCallbacks::OnGizmoModeButtonClicked, VRMode),
+			FCanExecuteAction::CreateStatic(&FLevelEditorActionCallbacks::DefaultCanExecuteAction)
+			),
 		NAME_None,
-		EUserInterfaceActionType::Check
+		EUserInterfaceActionType::CollapsedButton
 		);
-	SnapMenuBuilder.AddMenuEntry(
-		LOCTEXT("ToggleRotationSnap", "Rotate"),
-		LOCTEXT("ToggleRotationSnapTooltip", "Toggle Rotation Snap"),
-		FSlateIcon(FEditorStyle::GetStyleSetName(), "EditorViewport.RotateMode.Small"),
+	TAttribute<FText> DynamicGizmoCoordinateSystemLabel;
+	DynamicGizmoCoordinateSystemLabel.BindStatic(&FVREditorActionCallbacks::GetGizmoCoordinateSystemText);
+	GizmoMenuBuilder.AddMenuEntry(
+		DynamicGizmoCoordinateSystemLabel,
+		LOCTEXT("GizmoCoordinateTooltip", "Toggle Gizmo Coordinates"),
+		FSlateIcon(FEditorStyle::GetStyleSetName(), "EditorViewport.RelativeCoordinateSystem_World"),
 		FUIAction
 		(
-			FExecuteAction::CreateStatic(&FLevelEditorActionCallbacks::RotationGridSnap_Clicked),
+			FExecuteAction::CreateStatic(&FVREditorActionCallbacks::OnGizmoCoordinateSystemButtonClicked, VRMode),
+			FCanExecuteAction::CreateStatic(&FLevelEditorActionCallbacks::DefaultCanExecuteAction)
+			),
+		NAME_None,
+		EUserInterfaceActionType::CollapsedButton
+		);
+
+	GizmoMenuBuilder.AddMenuEntry(
+		LOCTEXT("AlignToActors", "Align To Actors"),
+		LOCTEXT("AlignToActorsTooltip", "Align to Actors as you transform an object"),
+		FSlateIcon(FVREditorStyle::GetStyleSetName(), "VREditorStyle.AlignActors"),
+		FUIAction
+		(
+			FExecuteAction::CreateStatic(&FVREditorActionCallbacks::ToggleAligningToActors, VRMode),
 			FCanExecuteAction::CreateStatic(&FLevelEditorActionCallbacks::DefaultCanExecuteAction),
-			FGetActionCheckState::CreateStatic(&FVREditorActionCallbacks::GetRotationSnapState)
+			FGetActionCheckState::CreateStatic(&FVREditorActionCallbacks::AreAligningToActors, VRMode)
+			),
+		NAME_None,
+		EUserInterfaceActionType::ToggleButton
+		);
+
+	TAttribute<FText> DynamicAlignSelectionLabel;
+	DynamicAlignSelectionLabel.BindStatic(&FVREditorActionCallbacks::GetSelectingCandidateActorsText);
+	GizmoMenuBuilder.AddMenuEntry(
+		DynamicAlignSelectionLabel,
+		FText(),
+		FSlateIcon(FEditorStyle::GetStyleSetName(), "DeviceDetails.Share"),
+		FUIAction
+		(
+			FExecuteAction::CreateStatic(&FVREditorActionCallbacks::ToggleSelectingCandidateActors, VRMode),
+			FCanExecuteAction::CreateStatic(&FVREditorActionCallbacks::CanSelectCandidateActors, VRMode)
 			),
 		NAME_None,
-		EUserInterfaceActionType::Check
+		EUserInterfaceActionType::CollapsedButton
 		);
+
+	TSharedRef<SWidget> GizmoMenuWidget = GizmoMenuBuilder.MakeWidget(&VREditorMenuBuilderOverride);
+
+	VREditorMenuBuilderOverride.BindUObject(this, &UVREditorUISystem::MakeUniformGridMenu, 6);
+	CommandList = MakeShareable(new FUICommandList());
+	FMenuBuilder SnapMenuBuilder(false, CommandList, MenuExtender);
+	
+	// Snap menu 
+
 	SnapMenuBuilder.AddMenuEntry(
-		LOCTEXT("ToggleScaleSnap", "Scale"),
-		LOCTEXT("ToggleScaleSnapTooltip", "Toggle Scale Snap"),
-		FSlateIcon(FEditorStyle::GetStyleSetName(), "EditorViewport.ScaleMode.Small"),
+		FText(),
+		LOCTEXT("ToggleTranslationSnapTooltip", "Toggle Translation Snap"),
+		FSlateIcon(FEditorStyle::GetStyleSetName(), "EditorViewport.TranslateMode"),
 		FUIAction
 		(
-			FExecuteAction::CreateStatic(&FLevelEditorActionCallbacks::ScaleGridSnap_Clicked),
+			FExecuteAction::CreateStatic(&FLevelEditorActionCallbacks::LocationGridSnap_Clicked),
 			FCanExecuteAction::CreateStatic(&FLevelEditorActionCallbacks::DefaultCanExecuteAction),
-			FGetActionCheckState::CreateStatic(&FVREditorActionCallbacks::GetScaleSnapState)
-			),
+			FGetActionCheckState::CreateStatic(&FVREditorActionCallbacks::GetTranslationSnapState)
+		),
 		NAME_None,
-		EUserInterfaceActionType::Check
+		EUserInterfaceActionType::ToggleButton
 		);
 	TAttribute<FText> DynamicTranslationSizeLabel;
 	DynamicTranslationSizeLabel.BindStatic(&FVREditorActionCallbacks::GetTranslationSnapSizeText);
 	SnapMenuBuilder.AddMenuEntry(
 		DynamicTranslationSizeLabel,
 		LOCTEXT("ToggleTranslationSnapSizeTooltip", "Toggle Translation Snap Size"),
-		FSlateIcon(),
+		FSlateIcon(FEditorStyle::GetStyleSetName(), "EditorViewport.LocationGridSnap"),
 		FUIAction
 		(
 			FExecuteAction::CreateStatic(&FVREditorActionCallbacks::OnTranslationSnapSizeButtonClicked),
 			FCanExecuteAction::CreateStatic(&FLevelEditorActionCallbacks::DefaultCanExecuteAction)
-			)
+			),
+		NAME_None,
+		EUserInterfaceActionType::CollapsedButton
+		);
+	SnapMenuBuilder.AddMenuEntry(
+		FText(),
+		LOCTEXT("ToggleRotationSnapTooltip", "Toggle Rotation Snap"),
+		FSlateIcon(FEditorStyle::GetStyleSetName(), "EditorViewport.RotateMode"),
+		FUIAction
+		(
+			FExecuteAction::CreateStatic(&FLevelEditorActionCallbacks::RotationGridSnap_Clicked),
+			FCanExecuteAction::CreateStatic(&FLevelEditorActionCallbacks::DefaultCanExecuteAction),
+			FGetActionCheckState::CreateStatic(&FVREditorActionCallbacks::GetRotationSnapState)
+			),
+		NAME_None,
+		EUserInterfaceActionType::ToggleButton
 		);
 	TAttribute<FText> DynamicRotationSizeLabel;
 	DynamicRotationSizeLabel.BindStatic(&FVREditorActionCallbacks::GetRotationSnapSizeText);
 	SnapMenuBuilder.AddMenuEntry(
 		DynamicRotationSizeLabel,
 		LOCTEXT("ToggleRotationSnapSizeTooltip", "Toggle Rotation Snap Size"),
-		FSlateIcon(),
+		FSlateIcon(FEditorStyle::GetStyleSetName(), "EditorViewport.RotationGridSnap"),
 		FUIAction
 		(
 			FExecuteAction::CreateStatic(&FVREditorActionCallbacks::OnRotationSnapSizeButtonClicked),
 			FCanExecuteAction::CreateStatic(&FLevelEditorActionCallbacks::DefaultCanExecuteAction)
-			)
+			),
+		NAME_None,
+		EUserInterfaceActionType::CollapsedButton
+		);
+	SnapMenuBuilder.AddMenuEntry(
+		FText(),
+		LOCTEXT("ToggleScaleSnapTooltip", "Toggle Scale Snap"),
+		FSlateIcon(FEditorStyle::GetStyleSetName(), "EditorViewport.ScaleMode"),
+		FUIAction
+		(
+			FExecuteAction::CreateStatic(&FLevelEditorActionCallbacks::ScaleGridSnap_Clicked),
+			FCanExecuteAction::CreateStatic(&FLevelEditorActionCallbacks::DefaultCanExecuteAction),
+			FGetActionCheckState::CreateStatic(&FVREditorActionCallbacks::GetScaleSnapState)
+			),
+		NAME_None,
+		EUserInterfaceActionType::ToggleButton
 		);
 	TAttribute<FText> DynamicScaleSizeLabel;
 	DynamicScaleSizeLabel.BindStatic(&FVREditorActionCallbacks::GetScaleSnapSizeText);
 	SnapMenuBuilder.AddMenuEntry(
 		DynamicScaleSizeLabel,
 		LOCTEXT("ToggleScaleSnapSizeTooltip", "Toggle Scale Snap Size"),
-		FSlateIcon(),
+		FSlateIcon(FEditorStyle::GetStyleSetName(), "EditorViewport.ScaleGridSnap"),
 		FUIAction
 		(
 			FExecuteAction::CreateStatic(&FVREditorActionCallbacks::OnScaleSnapSizeButtonClicked),
 			FCanExecuteAction::CreateStatic(&FLevelEditorActionCallbacks::DefaultCanExecuteAction)
-			)
+			),
+		NAME_None,
+		EUserInterfaceActionType::CollapsedButton
 		);
 
+
 	TSharedRef<SWidget> SnapMenuWidget = SnapMenuBuilder.MakeWidget(&VREditorMenuBuilderOverride);
 
 	// Next Menu section: Interface
 	CommandList = MakeShareable(new FUICommandList());
 	FMenuBuilder InterfaceMenuBuilder(false, CommandList, MenuExtender);
-
+	VREditorMenuBuilderOverride.BindUObject(this, &UVREditorUISystem::MakeUniformGridMenu, 4);
 	InterfaceMenuBuilder.AddMenuEntry(
 		LOCTEXT("ActorDetails", "Details"),
 		LOCTEXT("ActorDetailsTooltip", "Details"),
@@ -2090,7 +2092,7 @@ TSharedRef<SWidget> UVREditorUISystem::BuildQuickMenuWidget()
 			FGetActionCheckState::CreateStatic(&FVREditorActionCallbacks::GetUIToggledState, VRMode, UVREditorUISystem::EEditorUIPanel::ActorDetails)
 			),
 		NAME_None,
-		EUserInterfaceActionType::Check
+		EUserInterfaceActionType::ToggleButton
 		);
 	InterfaceMenuBuilder.AddMenuEntry(
 		LOCTEXT("ContentBrowser", "Content Browser"),
@@ -2103,7 +2105,7 @@ TSharedRef<SWidget> UVREditorUISystem::BuildQuickMenuWidget()
 			FGetActionCheckState::CreateStatic(&FVREditorActionCallbacks::GetUIToggledState, VRMode, UVREditorUISystem::EEditorUIPanel::ContentBrowser)
 			),
 		NAME_None,
-		EUserInterfaceActionType::Check
+		EUserInterfaceActionType::ToggleButton
 		);
 	InterfaceMenuBuilder.AddMenuEntry(
 		LOCTEXT("Modes", "Modes"),
@@ -2116,7 +2118,7 @@ TSharedRef<SWidget> UVREditorUISystem::BuildQuickMenuWidget()
 			FGetActionCheckState::CreateStatic(&FVREditorActionCallbacks::GetUIToggledState, VRMode, UVREditorUISystem::EEditorUIPanel::Modes)
 			),
 		NAME_None,
-		EUserInterfaceActionType::Check
+		EUserInterfaceActionType::ToggleButton
 		);
 	InterfaceMenuBuilder.AddMenuEntry(
 		LOCTEXT("WorldOutliner", "World Outliner"),
@@ -2129,33 +2131,7 @@ TSharedRef<SWidget> UVREditorUISystem::BuildQuickMenuWidget()
 			FGetActionCheckState::CreateStatic(&FVREditorActionCallbacks::GetUIToggledState, VRMode, UVREditorUISystem::EEditorUIPanel::WorldOutliner)
 			),
 		NAME_None,
-		EUserInterfaceActionType::Check
-		);
-	InterfaceMenuBuilder.AddMenuEntry(
-		LOCTEXT("AssetEditor", "Asset Editor"),
-		LOCTEXT("AssetEditorTooltip", "Asset Editor"),
-		FSlateIcon(),
-		FUIAction
-		(
-			FExecuteAction::CreateStatic(&FVREditorActionCallbacks::OnUIToggleButtonClicked, VRMode, UVREditorUISystem::EEditorUIPanel::AssetEditor),
-			FCanExecuteAction::CreateStatic(&FLevelEditorActionCallbacks::DefaultCanExecuteAction),
-			FGetActionCheckState::CreateStatic(&FVREditorActionCallbacks::GetUIToggledState, VRMode, UVREditorUISystem::EEditorUIPanel::AssetEditor)
-			),
-		NAME_None,
-		EUserInterfaceActionType::Check
-		);
-	InterfaceMenuBuilder.AddMenuEntry(
-		LOCTEXT("Tutorial", "Tutorial"),
-		LOCTEXT("TutorialTooltip", "Tutorial"),
-		FSlateIcon(FEditorStyle::GetStyleSetName(), "Tutorials.Browser.DefaultTutorialIcon"),
-		FUIAction
-		(
-			FExecuteAction::CreateStatic(&FVREditorActionCallbacks::OnUIToggleButtonClicked, VRMode, UVREditorUISystem::EEditorUIPanel::Tutorial),
-			FCanExecuteAction::CreateStatic(&FLevelEditorActionCallbacks::DefaultCanExecuteAction),
-			FGetActionCheckState::CreateStatic(&FVREditorActionCallbacks::GetUIToggledState, VRMode, UVREditorUISystem::EEditorUIPanel::Tutorial)
-			),
-		NAME_None,
-		EUserInterfaceActionType::Check
+		EUserInterfaceActionType::ToggleButton
 		);
 	InterfaceMenuBuilder.AddMenuEntry(
 		LOCTEXT("WorldSettings", "World Settings"),
@@ -2168,116 +2144,129 @@ TSharedRef<SWidget> UVREditorUISystem::BuildQuickMenuWidget()
 			FGetActionCheckState::CreateStatic(&FVREditorActionCallbacks::GetUIToggledState, VRMode, UVREditorUISystem::EEditorUIPanel::WorldSettings)
 			),
 		NAME_None,
-		EUserInterfaceActionType::Check
+		EUserInterfaceActionType::ToggleButton
 		);
 	InterfaceMenuBuilder.AddMenuEntry(
-		LOCTEXT("CreateNewSequence", "Create Sequence"),
+		LOCTEXT("CreateNewSequence", "Create Sequence"), 
 		LOCTEXT("CreateNewSequenceTooltip", "Create a new Level Sequence and open it"),
-		FSlateIcon(),
+		FSlateIcon(FEditorStyle::GetStyleSetName(), "LevelEditor.EditMatinee"),
 		FUIAction
 		(
 			FExecuteAction::CreateStatic(&FVREditorActionCallbacks::CreateNewSequence, VRMode),
 			FCanExecuteAction::CreateStatic(&FLevelEditorActionCallbacks::DefaultCanExecuteAction)
-			)
+			),
+		NAME_None,
+		EUserInterfaceActionType::ToggleButton
 		);
-	InterfaceMenuBuilder.AddMenuEntry(
-		LOCTEXT("ToggleAllUI", "Toggle All UI"),
-		LOCTEXT("ToggleAllUITooltip", "ToggleAllUI"),
-		FSlateIcon(),
+	TSharedRef<SWidget> InterfaceMenuWidget = InterfaceMenuBuilder.MakeWidget(&VREditorMenuBuilderOverride);
+
+	// Next Menu section: Extras
+	CommandList = MakeShareable(new FUICommandList());
+	FMenuBuilder ExtrasMenuBuilder(false, CommandList, MenuExtender);
+	VREditorMenuBuilderOverride.BindUObject(this, &UVREditorUISystem::MakeUniformGridMenu, 4);
+
+	TAttribute<FText> DynamicSimulateLabel;
+	DynamicSimulateLabel.BindStatic( &FVREditorActionCallbacks::GetSimulateText );
+
+	ExtrasMenuBuilder.AddMenuEntry(
+		DynamicSimulateLabel,
+		LOCTEXT("SimulateTooltip", "Simulate"),
+		FSlateIcon(FEditorStyle::GetStyleSetName(), "PlayWorld.Simulate"),
 		FUIAction
 		(
-			FExecuteAction::CreateUObject(this, &UVREditorUISystem::TogglePanelsVisibility),
+			FExecuteAction::CreateStatic(&FVREditorActionCallbacks::OnSimulateButtonClicked, VRMode),
 			FCanExecuteAction::CreateStatic(&FLevelEditorActionCallbacks::DefaultCanExecuteAction)
-			)
+			),
+		NAME_None,
+		EUserInterfaceActionType::CollapsedButton
 		);
-	
 
-	TSharedRef<SWidget> InterfaceMenuWidget = InterfaceMenuBuilder.MakeWidget(&VREditorMenuBuilderOverride);
 
-	// Next Menu section: Gizmo
-	CommandList = MakeShareable(new FUICommandList());
-	FMenuBuilder GizmoMenuBuilder(false, CommandList, MenuExtender);
+	ExtrasMenuBuilder.AddMenuEntry(
+		LOCTEXT("KeepSimulationChanges", "Save Simulation"),
+		LOCTEXT("KeepSimulationChangesTooltip", "Keep Simulation Changes"),
+		FSlateIcon(FEditorStyle::GetStyleSetName(), "LevelEditor.Save"),
+		FUIAction
+		(
+			FExecuteAction::CreateStatic(&FLevelEditorActionCallbacks::OnKeepSimulationChanges),
+			FCanExecuteAction::CreateStatic(&FLevelEditorActionCallbacks::CanExecuteKeepSimulationChanges)
+		),
+		NAME_None,
+		EUserInterfaceActionType::CollapsedButton
+		);
 
-	TAttribute<FText> DynamicGizmoCoordinateSystemLabel;
-	DynamicGizmoCoordinateSystemLabel.BindStatic(&FVREditorActionCallbacks::GetGizmoCoordinateSystemText);
-	GizmoMenuBuilder.AddMenuEntry(
-		DynamicGizmoCoordinateSystemLabel,
-		LOCTEXT("GizmoCoordinateTooltip", "Toggle Gizmo Coordinates"),
-		FSlateIcon(),
+
+	ExtrasMenuBuilder.AddMenuEntry(
+		LOCTEXT("PauseSimulation", "Pause Simulation"),
+		LOCTEXT("PauseSimulationTooltip", "Pause Simulation"),
+		FSlateIcon(FEditorStyle::GetStyleSetName(), "PlayWorld.Simulate"),
 		FUIAction
 		(
-			FExecuteAction::CreateStatic(&FVREditorActionCallbacks::OnGizmoCoordinateSystemButtonClicked, VRMode),
-			FCanExecuteAction::CreateStatic(&FLevelEditorActionCallbacks::DefaultCanExecuteAction)
-			)
+			FExecuteAction::CreateStatic(&FPlayWorldCommandCallbacks::PausePlaySession_Clicked),
+			FCanExecuteAction::CreateStatic(&FPlayWorldCommandCallbacks::HasPlayWorldAndRunning)
+		),
+		NAME_None,
+		EUserInterfaceActionType::CollapsedButton
 		);
-	TAttribute<FText> DynamicGizmoModeLabel;
-	DynamicGizmoModeLabel.BindStatic(&FVREditorActionCallbacks::GetGizmoModeText);
-	GizmoMenuBuilder.AddMenuEntry(
-		DynamicGizmoModeLabel,
-		LOCTEXT("GizmoModeTooltip", "Toggle Gizmo Mode"),
-		FSlateIcon(),
+
+	ExtrasMenuBuilder.AddMenuEntry(
+		LOCTEXT("ResumeSimulation", "Resume Simulation"),
+		LOCTEXT("ResumeSimulationTooltip", "Resume Simulation"),
+		FSlateIcon(FEditorStyle::GetStyleSetName(), "PlayWorld.Simulate"),
 		FUIAction
 		(
-			FExecuteAction::CreateStatic(&FVREditorActionCallbacks::OnGizmoModeButtonClicked, VRMode),
-			FCanExecuteAction::CreateStatic(&FLevelEditorActionCallbacks::DefaultCanExecuteAction)
-			)
+			FExecuteAction::CreateStatic(&FPlayWorldCommandCallbacks::ResumePlaySession_Clicked),
+			FCanExecuteAction::CreateStatic(&FPlayWorldCommandCallbacks::HasPlayWorldAndPaused)
+		),
+		NAME_None,
+		EUserInterfaceActionType::CollapsedButton
 		);
-	GizmoMenuBuilder.AddMenuEntry(
-		LOCTEXT("CycleGizmo", "Cycle Gizmo"),
-		LOCTEXT("CycleGizmoTooltip", "Cycle Gizmo"),
-		FSlateIcon(),
+
+	ExtrasMenuBuilder.AddMenuEntry(
+		LOCTEXT("PlayInEditor", "Play"),
+		LOCTEXT("PlayTooltip", "Play"),
+		FSlateIcon(FEditorStyle::GetStyleSetName(), "PlayWorld.PlayInVR"),
 		FUIAction
 		(
-			FExecuteAction::CreateStatic(&FVREditorActionCallbacks::OnGizmoCycle, VRMode),
+			FExecuteAction::CreateStatic(&FVREditorActionCallbacks::OnPlayButtonClicked, VRMode),
 			FCanExecuteAction::CreateStatic(&FLevelEditorActionCallbacks::DefaultCanExecuteAction)
-			)
+		),
+		NAME_None,
+		EUserInterfaceActionType::CollapsedButton
 		);
 
-	TSharedRef<SWidget> GizmoMenuWidget = GizmoMenuBuilder.MakeWidget(&VREditorMenuBuilderOverride);
-
-	// Next Menu section: Extras
-	CommandList = MakeShareable(new FUICommandList());
-	FMenuBuilder ExtrasMenuBuilder(false, CommandList, MenuExtender);
-
-
 	ExtrasMenuBuilder.AddMenuEntry(
 		LOCTEXT("Screenshot", "Screenshot"),
 		LOCTEXT("ScreenshotTooltip", "Screenshot"),
-		FSlateIcon(),
+		FSlateIcon(FEditorStyle::GetStyleSetName(), "HighresScreenshot.Capture"),
 		FUIAction
 		(
 			FExecuteAction::CreateStatic(&FVREditorActionCallbacks::OnScreenshotButtonClicked, VRMode),
 			FCanExecuteAction::CreateStatic(&FLevelEditorActionCallbacks::DefaultCanExecuteAction)
-			)
+		),
+		NAME_None,
+		EUserInterfaceActionType::CollapsedButton
 		);
+
 	ExtrasMenuBuilder.AddMenuEntry(
 		LOCTEXT("Flashlight", "Flashlight"),
 		LOCTEXT("FlashlightTooltip", "Flashlight"),
-		FSlateIcon(),
+		FSlateIcon(FEditorStyle::GetStyleSetName(), "EditorViewport.DetailLightingMode"),
 		FUIAction
 		(
 			FExecuteAction::CreateStatic(&FVREditorActionCallbacks::OnLightButtonClicked, VRMode),
-			FCanExecuteAction::CreateStatic(&FLevelEditorActionCallbacks::DefaultCanExecuteAction),
-			FGetActionCheckState::CreateStatic(&FVREditorActionCallbacks::GetFlashlightState, VRMode)
-			),
-		NAME_None,
-		EUserInterfaceActionType::Check
-		);
-	ExtrasMenuBuilder.AddMenuEntry(
-		LOCTEXT("PlayInEditor", "Play"),
-		LOCTEXT("PlayTooltip", "Play"),
-		FSlateIcon(),
-		FUIAction
-		(
-			FExecuteAction::CreateStatic(&FVREditorActionCallbacks::OnPlayButtonClicked, VRMode),
 			FCanExecuteAction::CreateStatic(&FLevelEditorActionCallbacks::DefaultCanExecuteAction)
-			)
+		),
+		NAME_None,
+		EUserInterfaceActionType::CollapsedButton
 		);
+
 	TSharedRef<SWidget> ExtrasMenuWidget = ExtrasMenuBuilder.MakeWidget(&VREditorMenuBuilderOverride);
 	const FSlateBrush* NoBorderBrush = FStyleDefaults::GetNoBrush();
 	TSharedRef<SWidget> MenuWidget = SNew(SBox)
-		.VAlign(VAlign_Top)
+		.VAlign(VAlign_Center)
 		.HAlign(HAlign_Center)
+		.MinDesiredWidth(1000.0f)
 		[
 			SNew(SDPIScaler)
 			.DPIScale(3)
@@ -2286,22 +2275,39 @@ TSharedRef<SWidget> UVREditorUISystem::BuildQuickMenuWidget()
 				+SVerticalBox::Slot()
 				.HAlign(HAlign_Fill)
 				.VAlign(VAlign_Bottom)
+				.AutoHeight()
 				[
 					SNew(STextBlock)
 					.Justification(ETextJustify::Center)
-					.Text(LOCTEXT("GridSnap", "Grid Snap"))
+					.Text(LOCTEXT("TransformControls", "Transform Controls"))
 					.TextStyle(FVREditorStyle::Get(),"VREditorStyle.Heading")
 				]
 				+ SVerticalBox::Slot()
 				.HAlign(HAlign_Fill)
 				.VAlign(VAlign_Top)
 				.AutoHeight()
+				.Padding(10.0f)
 				[
-					SnapMenuWidget
+					SNew(SVerticalBox)
+					+ SVerticalBox::Slot()
+					.HAlign(HAlign_Fill)
+					.VAlign(VAlign_Bottom)
+					.AutoHeight()
+					[
+						GizmoMenuWidget
+					]
+					+ SVerticalBox::Slot()
+					.HAlign(HAlign_Fill)
+					.VAlign(VAlign_Top)
+					.AutoHeight()
+					[
+						SnapMenuWidget
+					]
 				]
 				+ SVerticalBox::Slot()
 				.HAlign(HAlign_Fill)
 				.VAlign(VAlign_Bottom)
+				.AutoHeight()
 				[
 					SNew(STextBlock)
 					.Justification(ETextJustify::Center)
@@ -2312,12 +2318,14 @@ TSharedRef<SWidget> UVREditorUISystem::BuildQuickMenuWidget()
 				.HAlign(HAlign_Fill)
 				.VAlign(VAlign_Top)
 				.AutoHeight()
+				.Padding(10.0f)
 				[
 					InterfaceMenuWidget
 				]
 				+ SVerticalBox::Slot()
 				.HAlign(HAlign_Fill)
 				.VAlign(VAlign_Bottom)
+				.AutoHeight()
 				[
 					SNew(STextBlock)
 					.Justification(ETextJustify::Center)
@@ -2328,25 +2336,10 @@ TSharedRef<SWidget> UVREditorUISystem::BuildQuickMenuWidget()
 				.HAlign(HAlign_Fill)
 				.VAlign(VAlign_Top)
 				.AutoHeight()
+				.Padding(10.0f)
 				[
 					ExtrasMenuWidget
 				]
-				+ SVerticalBox::Slot()
-				.HAlign(HAlign_Fill)
-				.VAlign(VAlign_Bottom)
-				[
-					SNew(STextBlock)
-					.Justification(ETextJustify::Center)
-					.Text(LOCTEXT("Gizmo", "Gizmo"))
-					.TextStyle(FVREditorStyle::Get(), "VREditorStyle.Heading")
-				]
-				+ SVerticalBox::Slot()
-				.HAlign(HAlign_Fill)
-				.VAlign(VAlign_Top)
-				.AutoHeight()
-				[
-					GizmoMenuWidget
-				]
 			]
 		]
 	;
@@ -2357,110 +2350,30 @@ TSharedRef<SWidget> UVREditorUISystem::BuildQuickMenuWidget()
 
 
 
-TSharedRef<SWidget> UVREditorUISystem::BuildRadialMenuWidget()
+void UVREditorUISystem::BuildRadialMenuWidget()
 {
 
 	FMultiBox::FOnMakeMultiBoxBuilderOverride VREditorMenuBuilderOverride;
-	VREditorMenuBuilderOverride.BindUObject(this, &UVREditorUISystem::MakeRadialBoxMenu, 1.0f);
+	
 	const TSharedRef<FExtender> MenuExtender = MakeShareable(new FExtender());
 	TSharedPtr<FUICommandList> CommandList = MakeShareable(new FUICommandList());
-	FMenuBarBuilder MenuBuilder(CommandList, MenuExtender);
-
-	FVREditorActionCallbacks::UpdateGizmoCoordinateSystemText(VRMode);
-	FVREditorActionCallbacks::UpdateGizmoModeText(VRMode);
-
-	// First menu entry is at 90 degrees 
-	MenuBuilder.AddMenuEntry(
-		LOCTEXT( "Delete", "Delete" ),
-		LOCTEXT( "DeleteTooltip", "Delete" ),
-		FSlateIcon(),
-		FUIAction
-		(
-			FExecuteAction::CreateStatic( &FLevelEditorActionCallbacks::ExecuteExecCommand, FString( TEXT( "DELETE" ) ) ),
-			FCanExecuteAction::CreateStatic( &FLevelEditorActionCallbacks::Delete_CanExecute )
-		)
-	);
-	MenuBuilder.AddMenuEntry(
-		LOCTEXT( "Cut", "Cut" ),
-		LOCTEXT( "CutTooltip", "Cut" ),
-		FSlateIcon(),
-		FUIAction
-		(
-			FExecuteAction::CreateStatic( &FLevelEditorActionCallbacks::ExecuteExecCommand, FString( TEXT( "EDIT CUT" ) ) ),
-			FCanExecuteAction::CreateStatic( &FLevelEditorActionCallbacks::Cut_CanExecute )
-		)
-	);
-	MenuBuilder.AddMenuEntry(
-		LOCTEXT( "Copy", "Copy" ),
-		LOCTEXT( "CopyTooltip", "Copy" ),
-		FSlateIcon(),
-		FUIAction
-		(
-			FExecuteAction::CreateStatic( &FLevelEditorActionCallbacks::ExecuteExecCommand, FString( TEXT( "EDIT COPY" ) ) ),
-			FCanExecuteAction::CreateStatic( &FLevelEditorActionCallbacks::Copy_CanExecute )
-		)
-	);
-	MenuBuilder.AddMenuEntry(
-		LOCTEXT( "Duplicate", "Duplicate" ),
-		LOCTEXT( "DuplicateTooltip", "Duplicate" ),
-		FSlateIcon(),
-		FUIAction
-		(
-			FExecuteAction::CreateStatic( &FLevelEditorActionCallbacks::ExecuteExecCommand, FString( TEXT( "DUPLICATE" ) ) ),
-			FCanExecuteAction::CreateStatic( &FLevelEditorActionCallbacks::Duplicate_CanExecute )
-		)
-	);
-	MenuBuilder.AddMenuEntry(
-		LOCTEXT( "Paste", "Paste" ),
-		LOCTEXT( "PasteTooltip", "Paste" ),
-		FSlateIcon(),
-		FUIAction
-		(
-			FExecuteAction::CreateStatic( &FLevelEditorActionCallbacks::ExecuteExecCommand, FString( TEXT( "EDIT PASTE" ) ) ),
-			FCanExecuteAction::CreateStatic( &FLevelEditorActionCallbacks::Paste_CanExecute )
-		)
-	);
-	MenuBuilder.AddMenuEntry(
-		LOCTEXT( "Undo", "Undo" ),
-		LOCTEXT( "UndoTooltip", "Undo" ),
-		FSlateIcon(),
-		FUIAction
-		(
-			FExecuteAction::CreateStatic( &FLevelEditorActionCallbacks::ExecuteExecCommand, FString( TEXT( "TRANSACTION UNDO" ) ) ),
-			FCanExecuteAction::CreateStatic( &FLevelEditorActionCallbacks::DefaultCanExecuteAction )
-		)
-	);
-	MenuBuilder.AddMenuEntry(
-		LOCTEXT( "SnapActorsToGround", "Snap Actors To Ground" ),
-		LOCTEXT( "SnapActorsToGroundTooltip", "Snap Actors To Ground" ),
-		FSlateIcon(),
-		FUIAction
-		(
-			FExecuteAction::CreateStatic( &FVREditorActionCallbacks::OnSnapActorsToGroundClicked, VRMode ),
-			FCanExecuteAction::CreateStatic( &FLevelEditorActionCallbacks::Copy_CanExecute )
-		)
-	);
-	MenuBuilder.AddMenuEntry(
-		LOCTEXT( "Redo", "Redo" ),
-		LOCTEXT( "RedoTooltip", "Redo" ),
-		FSlateIcon(),
-		FUIAction
-		(
-			FExecuteAction::CreateStatic( &FLevelEditorActionCallbacks::ExecuteExecCommand, FString( TEXT( "TRANSACTION REDO" ) ) ),
-			FCanExecuteAction::CreateStatic( &FLevelEditorActionCallbacks::DefaultCanExecuteAction )
-		)
-	);
+	FMenuBuilder MenuBuilder(false, CommandList, MenuExtender);
+	float RadiusOverride = 1.0f;
+	if (RadialMenuHandler != nullptr)
+	{
+		RadialMenuHandler->BuildRadialMenuCommands(MenuBuilder, CommandList, VRMode, RadiusOverride);
+	}
+	VREditorMenuBuilderOverride.BindUObject(this, &UVREditorUISystem::MakeRadialBoxMenu, RadiusOverride, FName(TEXT("SMenuEntryButton")));
 
 	// Create the menu widget
-	TSharedRef<SWidget> MenuWidget = MenuBuilder.MakeWidget(&VREditorMenuBuilderOverride);
-	return MenuWidget;
+	MenuBuilder.MakeWidget(&VREditorMenuBuilderOverride);
 }
 
-TSharedRef<SWidget> UVREditorUISystem::BuildNumPadWidget()
+void UVREditorUISystem::BuildNumPadWidget()
 {
 
 	FMultiBox::FOnMakeMultiBoxBuilderOverride VREditorMenuBuilderOverride;
-	VREditorMenuBuilderOverride.BindUObject(this, &UVREditorUISystem::MakeRadialBoxMenu, 0.5f);
+	VREditorMenuBuilderOverride.BindUObject(this, &UVREditorUISystem::MakeRadialBoxMenu, 0.5f, FName(TEXT("SButton")));
 	const TSharedRef<FExtender> MenuExtender = MakeShareable(new FExtender());
 	TSharedPtr<FUICommandList> CommandList = MakeShareable(new FUICommandList());
 	FMenuBarBuilder MenuBuilder(CommandList, MenuExtender);
@@ -2577,8 +2490,8 @@ TSharedRef<SWidget> UVREditorUISystem::BuildNumPadWidget()
 	);
 
 	// Create the menu widget
-	TSharedRef<SWidget> MenuWidget = MenuBuilder.MakeWidget(&VREditorMenuBuilderOverride);
-
+	MenuBuilder.MakeWidget(&VREditorMenuBuilderOverride);
+	/**
 	TSharedRef<SWidget> FinalWidget = SNew( SOverlay )
 		+ SOverlay::Slot()
 		[
@@ -2611,8 +2524,7 @@ TSharedRef<SWidget> UVREditorUISystem::BuildNumPadWidget()
 			]
 		]
 		;
-
-	return FinalWidget;
+		*/
 }
 
 const TSharedRef<SWidget>& UVREditorUISystem::FindWidgetOfType(const TSharedRef<SWidget>& Content, FName WidgetType)
@@ -2636,67 +2548,256 @@ const TSharedRef<SWidget>& UVREditorUISystem::FindWidgetOfType(const TSharedRef<
 	return SNullWidget::NullWidget;
 }
 
+const bool UVREditorUISystem::FindAllWidgetsOfType(TArray<TSharedRef<SWidget>>& FoundWidgets, const TSharedRef<SWidget>& Content, FName WidgetType)
+{
+	bool bFoundMatch = false;
+	if (Content->GetType() == WidgetType)
+	{
+		FoundWidgets.Add(Content);
+		bFoundMatch = true;
+	}
+
+	FChildren* Children = Content->GetChildren();
+	int32 NumChildren = Children->Num();
+
+	for (int32 Index = 0; Index < NumChildren; ++Index)
+	{
+		bFoundMatch = FindAllWidgetsOfType(FoundWidgets, (Children->GetChildAt(Index)), WidgetType);
+	}
+
+	return bFoundMatch;
+}
+
 
 
-void UVREditorUISystem::OnHoverBeginEffect(TSharedRef<SButton> Button)
+void UVREditorUISystem::OnHoverBeginEffect(UVREditorWidgetComponent* Button)
 {
 	FVRButton* ButtonToAnimate = VRButtons.FindByPredicate([&Button](const FVRButton& ButtonData)
 	{
-		return Button == ButtonData.Button;
+		return Button == ButtonData.ButtonWidget;
 	});
 
-	check(ButtonToAnimate);
-
-	//Set the newly hovered button's animation state to Forward
-	ButtonToAnimate->AnimationDirection = EVREditorAnimationState::Forward;
-
-	// Put the button in the higher z-order overlay slot
-	const TSharedRef<SWidget>& PlaceholderWidget = ButtonToAnimate->HoverSlot->DetachWidget().ToSharedRef();
-	const TSharedRef<SWidget>& WidgetToTransfer = ButtonToAnimate->BaseSlot->DetachWidget().ToSharedRef();
-	ButtonToAnimate->HoverSlot->AttachWidget(WidgetToTransfer);
-	ButtonToAnimate->BaseSlot->AttachWidget(PlaceholderWidget);
+	if (ButtonToAnimate)
+	{
+		//Set the newly hovered button's animation state to Forward
+		ButtonToAnimate->AnimationDirection = EVREditorAnimationState::Forward;
+	}
 }
 
-void UVREditorUISystem::OnHoverEndEffect(TSharedRef<SButton> Button)
+void UVREditorUISystem::OnHoverEndEffect(UVREditorWidgetComponent* Button)
 {
 	FVRButton* ButtonToAnimate = VRButtons.FindByPredicate([&Button](const FVRButton& ButtonData)
 	{
-		return Button == ButtonData.Button;
+		return Button == ButtonData.ButtonWidget;
 	});
 
-	check(ButtonToAnimate);
+	if (ButtonToAnimate)
+	{
+		// Set the unhovered button's animation state to Backward
+		ButtonToAnimate->AnimationDirection = EVREditorAnimationState::Backward;
+	}
 
-	// Set the unhovered button's animation state to Backward
-	ButtonToAnimate->AnimationDirection = EVREditorAnimationState::Backward;
-	// Put the button back in the lower z-order slot
-	TSharedRef<SWidget> PlaceholderWidget = (ButtonToAnimate->BaseSlot->DetachWidget()).ToSharedRef();
-	TSharedRef<SWidget> WidgetToTransfer = (ButtonToAnimate->HoverSlot->DetachWidget()).ToSharedRef();
-	ButtonToAnimate->BaseSlot->AttachWidget(WidgetToTransfer);
-	ButtonToAnimate->HoverSlot->AttachWidget(PlaceholderWidget);
 }
 
-const TSharedPtr<SWidget>& UVREditorUISystem::GetRadialWidget()
+void UVREditorUISystem::SequencerOpenendFromRadialMenu(const bool bInOpenedFromRadialMenu /*= true*/)
 {
-	return RadialWidget;
+	bSequencerOpenedFromRadialMenu = bInOpenedFromRadialMenu;
+}
+
+bool UVREditorUISystem::CanScalePanel() const
+{
+	return bPanelCanScale || (DraggingUI != nullptr && LaserInteractor != nullptr && InteractorDraggingUI != nullptr && LaserInteractor == InteractorDraggingUI);
+}
+
+UVREditorMotionControllerInteractor* UVREditorUISystem::GetUIInteractor()
+{
+	return UIInteractor;
+}
+
+void UVREditorUISystem::SequencerRadialMenuGenerator(FMenuBuilder MenuBuilder, TSharedPtr<FUICommandList> CommandList, class UVREditorMode* VRMode, float& RadiusOverride)
+{
+	RadiusOverride = 1.0f;
+	// First menu entry is at 90 degrees 
+	MenuBuilder.AddMenuEntry(
+		LOCTEXT("Play", "Play"),
+		FText(),
+		FSlateIcon(FVREditorStyle::GetStyleSetName(), "VREditorStyle.SequencerPlay"),
+		FUIAction
+		(
+			FExecuteAction::CreateStatic(&FVREditorActionCallbacks::PlaySequenceAtRate, VRMode, 1.0f),
+			FCanExecuteAction::CreateUObject(this->GetRadialMenuHandler(), &UVRRadialMenuHandler::IsActionMenuBound)
+		)
+		);
+	MenuBuilder.AddMenuEntry(
+		LOCTEXT("Reverse", "Reverse"),
+		FText(),
+		FSlateIcon(FVREditorStyle::GetStyleSetName(), "VREditorStyle.SequencerReverse"),
+		FUIAction
+		(
+			FExecuteAction::CreateStatic(&FVREditorActionCallbacks::PlaySequenceAtRate, VRMode, -1.0f),
+			FCanExecuteAction::CreateUObject(this->GetRadialMenuHandler(), &UVRRadialMenuHandler::IsActionMenuBound)
+		)
+		);
+	MenuBuilder.AddMenuEntry(
+		LOCTEXT("Stop", "Stop"),
+		FText(),
+		FSlateIcon(FVREditorStyle::GetStyleSetName(), "VREditorStyle.SequencerStop"),
+		FUIAction
+		(
+			FExecuteAction::CreateStatic(&FVREditorActionCallbacks::PauseSequencePlayback, VRMode),
+			FCanExecuteAction::CreateUObject(this->GetRadialMenuHandler(), &UVRRadialMenuHandler::IsActionMenuBound)
+		)
+		);
+	MenuBuilder.AddMenuEntry(
+		LOCTEXT("PlayFromStart", "Play From Start"),
+		FText(),
+		FSlateIcon(FVREditorStyle::GetStyleSetName(), "VREditorStyle.PlayFromStart"),
+		FUIAction
+		(
+			FExecuteAction::CreateStatic(&FVREditorActionCallbacks::PlayFromBeginning, VRMode),
+			FCanExecuteAction::CreateUObject(this->GetRadialMenuHandler(), &UVRRadialMenuHandler::IsActionMenuBound)
+			)
+		);
+	MenuBuilder.AddMenuEntry(
+		LOCTEXT("ScrubTime", "Scrub Time"),
+		FText(),
+		FSlateIcon(FVREditorStyle::GetStyleSetName(), "VREditorStyle.Scrub"),
+		FUIAction
+		(
+			FExecuteAction::CreateStatic(&FVREditorActionCallbacks::ToggleSequencerScrubbing, VRMode, UIInteractor),
+			FCanExecuteAction::CreateUObject(this->GetRadialMenuHandler(), &UVRRadialMenuHandler::IsActionMenuBound),
+			FGetActionCheckState::CreateStatic(&FVREditorActionCallbacks::GetSequencerScrubState, UIInteractor)
+			),
+		NAME_None,
+		EUserInterfaceActionType::ToggleButton
+		);
+	MenuBuilder.AddMenuEntry(
+		LOCTEXT("ToggleLooping", "Toggle Looping"),
+		FText(),
+		FSlateIcon(FVREditorStyle::GetStyleSetName(), "VREditorStyle.ToggleLooping"),
+		FUIAction
+		(
+			FExecuteAction::CreateStatic(&FVREditorActionCallbacks::ToggleLooping, VRMode),
+			FCanExecuteAction::CreateUObject(this->GetRadialMenuHandler(), &UVRRadialMenuHandler::IsActionMenuBound),
+			FGetActionCheckState::CreateStatic(&FVREditorActionCallbacks::IsLoopingChecked, VRMode)
+			),
+		NAME_None,
+		EUserInterfaceActionType::ToggleButton
+		);
+}
+
+void UVREditorUISystem::HandleEditorModeChanged(FEdMode* Mode, bool IsEnabled)
+{
+	if (IsEnabled == true)
+	{
+		if (Mode->GetID() == FBuiltinEditorModes::EM_Foliage ||
+			Mode->GetID() == FBuiltinEditorModes::EM_Landscape ||
+			Mode->GetID() == FBuiltinEditorModes::EM_MeshPaint)
+		{
+			AVREditorFloatingUI* Panel = EditorUIPanels[(int32)EEditorUIPanel::Modes];
+			if (Panel != nullptr && UIInteractor != nullptr && !IsShowingEditorUIPanel(EEditorUIPanel::Modes))
+			{
+				const bool bShouldShow = true;
+				const bool bSpawnInFront = true;
+				const bool bDragFromOpen = ShouldPreviewPanel();
+				ShowEditorUIPanel(Panel, UIInteractor, bShouldShow, bSpawnInFront, bDragFromOpen);
+			}
+		}
+	}
+}
+
+void UVREditorUISystem::ResetAll()
+{
+	RadialMenuHandler->Home();
+	HideRadialMenu(UIInteractor, false);
+
+	FVREditorActionCallbacks::DeselectAll();
+
+	// Close all the editor UI panels.
+	const bool bShouldShow = false;
+	const bool bSpawnInFront = false;
+	const bool bDragFromOpen = false;
+	const bool bPlaySound = false;
+	for (AVREditorFloatingUI* FloatingUIPtr : FloatingUIs)
+	{
+		if (FloatingUIPtr != nullptr)
+		{
+			ShowEditorUIPanel(FloatingUIPtr, UIInteractor, bShouldShow, bSpawnInFront, bDragFromOpen, bPlaySound);
+		}
+	}
+	FVREditorActionCallbacks::ChangeEditorModes(FBuiltinEditorModes::EM_Placement);
+
+	if( GEditor->bIsSimulatingInEditor )
+	{
+		VRMode->ToggleSIEAndVREditor();
+	}
+
+	// Turn off all snapping
+	if (FVREditorActionCallbacks::GetTranslationSnapState() == ECheckBoxState::Checked)
+	{
+		FLevelEditorActionCallbacks::LocationGridSnap_Clicked();
+	}
+	if (FVREditorActionCallbacks::GetRotationSnapState() == ECheckBoxState::Checked)
+	{
+		FLevelEditorActionCallbacks::RotationGridSnap_Clicked();
+	}
+	if (FVREditorActionCallbacks::GetScaleSnapState() == ECheckBoxState::Checked)
+	{
+		FLevelEditorActionCallbacks::ScaleGridSnap_Clicked();
+	}
+	if (FVREditorActionCallbacks::AreAligningToActors(VRMode) == ECheckBoxState::Checked)
+	{
+		FVREditorActionCallbacks::ToggleAligningToActors(VRMode);
+	}
+}
+
+bool UVREditorUISystem::CheckForVRDragDrop()
+{
+	if (LaserInteractor != nullptr)
+	{
+		return LaserInteractor->IsClickingOnUI();
+	}
+	else
+	{
+		return false;
+	}
+}
+
+bool UVREditorUISystem::ShouldPreviewPanel()
+{
+	bool bDragFromOpen = true;
+	if (LaserInteractor != nullptr)
+	{
+		// If we are clicking with the laser interactor, instantly spawn the panel
+		FHitResult HitResult = LaserInteractor->GetHitResultFromLaserPointer();
+		if (HitResult.Actor.IsValid())
+		{
+			if (HitResult.Actor == QuickRadialMenu)
+			{
+				bDragFromOpen = !LaserInteractor->IsClickingOnUI();
+			}
+		}
+	}
+	return bDragFromOpen;
+}
+
+AVREditorRadialFloatingUI* UVREditorUISystem::GetRadialMenuFloatingUI() const
+{
+	return QuickRadialMenu;
 }
 
 void UVREditorUISystem::SwapRadialMenu()
 {
-	const AVREditorFloatingUI::EDockedTo RadialMenuInteractor = QuickRadialMenu->GetDockedTo();
-	const FIntPoint DefaultResolution( VREd::DefaultEditorUIResolutionX->GetInt(), VREd::DefaultEditorUIResolutionY->GetInt() );
 	if (bRadialMenuIsNumpad)
 	{
-		QuickRadialMenu->SetSlateWidget( *this, BuildRadialMenuWidget(), DefaultResolution, 40.0f, RadialMenuInteractor );
 		bRadialMenuIsNumpad = false;
 	}
 	else
 	{
-		QuickRadialMenu->SetSlateWidget( *this, BuildNumPadWidget(), DefaultResolution, 40.0f, RadialMenuInteractor );
 		bRadialMenuIsNumpad = true;
 	}
-	QuickRadialMenu->GetWidgetComponent()->SetOpacityFromTexture( 1.0f );
-	QuickRadialMenu->GetWidgetComponent()->SetBackgroundColor( FLinearColor::Transparent );
-	QuickRadialMenu->GetWidgetComponent()->SetBlendMode( EWidgetBlendMode::Masked );
+	const bool bForceRefresh = true;
+	TryToSpawnRadialMenu(UIInteractor, bForceRefresh);
 }
 
 void UVREditorUISystem::UpdateSequencerUI()
@@ -2704,7 +2805,7 @@ void UVREditorUISystem::UpdateSequencerUI()
 	if (EditorUIPanels[(int32)EEditorUIPanel::SequencerUI] != nullptr)
 	{
 		AVREditorFloatingUI* SequencerPanel = EditorUIPanels[(int32)EEditorUIPanel::SequencerUI];
-		const FName TabIdentifier = NAME_None;	// No tab for us!
+
 		if(GetOwner().GetCurrentSequencer() != nullptr)
 		{
 			const TSharedRef<SWidget> SequencerWidget = GetOwner().GetCurrentSequencer()->GetSequencerWidget();
@@ -2719,14 +2820,90 @@ void UVREditorUISystem::UpdateSequencerUI()
 
 			const bool bWithSceneComponent = false;
 			SequencerPanel->SetSlateWidget(*this, WidgetToDraw, FIntPoint(VREd::SequencerUIResolutionX->GetFloat(), VREd::SequencerUIResolutionY->GetFloat()), VREd::EditorUISize->GetFloat(), AVREditorFloatingUI::EDockedTo::Nothing);
-			UVREditorInteractor* VREditorInteractor = VRMode->GetHandInteractor(EControllerHand::Left);
-			if (VREditorInteractor->IsHoveringOverUI())
+	
+			if (bSequencerOpenedFromRadialMenu == true)
+			{
+				const bool bDragFromOpen = ShouldPreviewPanel();
+				const bool bShouldShow = true;
+				const bool bSpawnInFront = true;
+				ShowEditorUIPanel(SequencerPanel, UIInteractor, bShouldShow, bSpawnInFront, bDragFromOpen);
+				bSequencerOpenedFromRadialMenu = false;
+			}
+			else
+			{
+				// Spawn the opened sequencer just in front of the content browser
+				AVREditorFloatingUI* ContentBrowser = EditorUIPanels[(int32)EEditorUIPanel::ContentBrowser];
+				if (ContentBrowser && ContentBrowser->IsUIVisible())	
+				{
+					ShowEditorUIPanel(ContentBrowser, UIInteractor, false);
+
+					const float WorldScaleFactor = VRMode->GetWorldScaleFactor();
+					const FTransform WindowToWorld = ContentBrowser->GetActorTransform();
+					const FTransform WorldToRoom = VRMode->GetRoomTransform().Inverse();
+					const FTransform WindowToRoom = WindowToWorld * WorldToRoom;
+					const FVector RoomSpaceWindowLocation = WindowToRoom.GetLocation() / WorldScaleFactor;
+					const FQuat RoomSpaceWindowRotation = WindowToRoom.GetRotation();
+
+					SequencerPanel->SetScale(ContentBrowser->GetScale());
+					SequencerPanel->SetDockedTo(AVREditorBaseActor::EDockedTo::Room);
+					SequencerPanel->SetRelativeOffset(RoomSpaceWindowLocation);
+					SequencerPanel->SetLocalRotation(RoomSpaceWindowRotation.Rotator());
+					SequencerPanel->ShowUI(true);
+
+					// Play haptic effect so user knows to look at their hand that now has UI on it!
+					UIInteractor->PlayHapticEffect(VREd::UIAssetEditorSummonedOnHandHapticFeedbackStrength->GetFloat());
+				}
+			}
+
+			FOnRadialMenuGenerated SequencerRadialMenu;
+			SequencerRadialMenu.BindUObject(this, &UVREditorUISystem::SequencerRadialMenuGenerator);
+			if (RadialMenuHandler != nullptr)
 			{
-				VREditorInteractor = VRMode->GetHandInteractor(EControllerHand::Right);
+				if (RadialMenuHandler->IsActionMenuBound())
+				{
+					ExistingActionsMenu = RadialMenuHandler->GetActionsMenuGenerator();
+					ExistingActionsMenuLabel = RadialMenuHandler->GetActionMenuLabel();
+				}
+				RadialMenuHandler->SetActionsMenuGenerator(SequencerRadialMenu, LOCTEXT("Sequencer", "Sequencer"));
+			}
+		}
+		else
+		{
+			if (RadialMenuHandler != nullptr)
+			{
+				RadialMenuHandler->Home();
+				RadialMenuHandler->ResetActionsMenuGenerator();
+				if (ExistingActionsMenu.IsBound())
+				{
+					RadialMenuHandler->SetActionsMenuGenerator(ExistingActionsMenu, ExistingActionsMenuLabel);
+				}
+			}
+		}
+	}
+}
 
+void UVREditorUISystem::TransitionWorld(UWorld* NewWorld)
+{
+	for (AVREditorFloatingUI* FloatingUI : FloatingUIs)
+	{
+		UUserWidget* UserWidget = FloatingUI->GetUserWidget();
+		if (UserWidget != nullptr)
+		{
+			// Only reparent the UserWidget if it was parented to a level to begin with.  It may have been parented to an actor or
+			// some other object that doesn't require us to rename anything
+			ULevel* ExistingWidgetOuterLevel = Cast<ULevel>(UserWidget->GetOuter());
+			if (ExistingWidgetOuterLevel != nullptr && ExistingWidgetOuterLevel != NewWorld->PersistentLevel)
+			{
+				UserWidget->Rename(nullptr, NewWorld->PersistentLevel);
 			}
-			ShowEditorUIPanel(UVREditorUISystem::EEditorUIPanel::SequencerUI, VREditorInteractor, true);
 		}
+	}	
+	
+	AVREditorFloatingUI* TabManagerUI = EditorUIPanels[(int32)EEditorUIPanel::AssetEditor];
+	if (TabManagerUI != nullptr)
+	{
+		TabManagerUI->GetWidgetComponent()->UpdateWidget();
+		ProxyTabManager->SetParentWindow(TabManagerUI->GetWidgetComponent()->GetSlateWindow().ToSharedRef());
 	}
 }
 
