Merging //UE4/Private-Geometry to Dev-VREditor (//UE4/Dev-VREditor) 3339870

diff --git a/Engine/Source/Editor/VREditor/UI/VREditorUISystem.h b/Engine/Source/Editor/VREditor/UI/VREditorUISystem.h
--- a/Engine/Source/Editor/VREditor/UI/VREditorUISystem.h
+++ b/Engine/Source/Editor/VREditor/UI/VREditorUISystem.h
@@ -7,6 +7,7 @@
 #include "UObject/ObjectMacros.h"
 #include "UObject/Object.h"
 #include "Widgets/SOverlay.h"
+#include "VRRadialMenuHandler.h"
 #include "VREditorUISystem.generated.h"
 
 class AVREditorDockableWindow;
@@ -20,6 +21,9 @@ class SMultiBoxWidget;
 class SWidget;
 class UViewportInteractor;
 class UVREditorInteractor;
+class FMenuBuilder;
+class FUICommandList;
+class UVREditorWidgetComponent;
 
 
 /** Stores the animation playback state of a VR UI element */
@@ -31,23 +35,21 @@ enum class EVREditorAnimationState : uint8
 };
 
 /** Structure to keep track of all relevant interaction and animation elements of a VR Button */
+USTRUCT()
 struct FVRButton
 {
-	/** Pointer to button */
-	TSharedPtr<SButton> Button;
-
-	/** Pointer to slot in overlay with lower z-order */
-	SOverlay::FOverlaySlot* BaseSlot;
+	GENERATED_USTRUCT_BODY()
 
-	/** Pointer to slot in overlay with higher z-order */
-	SOverlay::FOverlaySlot* HoverSlot;
-
-	/** Pointer to button border */
-	TSharedPtr<SBorder> ButtonBorder;
+	/** Pointer to button */
+	UPROPERTY()
+	UVREditorWidgetComponent* ButtonWidget;
 
 	/** Animation playback state of the button */
 	EVREditorAnimationState AnimationDirection;
 
+	/** Original relative scale of the button element */
+	FVector OriginalRelativeScale;
+
 	/** Current scale of the button element */
 	float CurrentScale;
 
@@ -68,12 +70,10 @@ struct FVRButton
 		ScaleRate(2.0f)
 		{}
 
-	FVRButton(TSharedPtr<SButton> InButton, SOverlay::FOverlaySlot* InBaseSlot, SOverlay::FOverlaySlot* InHoverSlot, TSharedPtr<SBorder> InButtonBorder,
+	FVRButton(class UVREditorWidgetComponent* InButtonWidget, FVector InOriginalScale,
 		EVREditorAnimationState InAnimationDirection = EVREditorAnimationState::None, float InCurrentScale = 1.0f, float InMinScale = 1.0f, float InMaxScale = 1.25f, float InScaleRate = 2.0f)
-		: Button(InButton),
-		BaseSlot(InBaseSlot),
-		HoverSlot(InHoverSlot),
-		ButtonBorder(InButtonBorder),
+		: ButtonWidget(InButtonWidget),
+		OriginalRelativeScale(InOriginalScale),
 		AnimationDirection(InAnimationDirection),
 		CurrentScale(InCurrentScale),
 		MinScale(InMinScale),
@@ -142,27 +142,27 @@ public:
 
 	/** Called by our owner right before a map is loaded or switching between Simulate and normal Editor, so we can clean
 	    up our spawned actors */
-	void CleanUpActorsBeforeMapChangeOrSimulate();
+	void CleanupActorsBeforeShutdown();
 
 	/** Returns true if the specified editor UI panel is currently visible */
 	bool IsShowingEditorUIPanel( const EEditorUIPanel EditorUIPanel ) const;
 
 	/** Sets whether the specified editor UI panel should be visible.  Any other UI floating off this hand will be dismissed when showing it. */
-	void ShowEditorUIPanel( const class UWidgetComponent* WidgetComponent, UVREditorInteractor* Interactor, const bool bShouldShow, const bool OnHand = false, const bool bRefreshQuickMenu = true, const bool bPlaySound = true );
-	void ShowEditorUIPanel( const EEditorUIPanel EditorUIPanel, UVREditorInteractor* Interactor, const bool bShouldShow, const bool OnHand = false, const bool bRefreshQuickMenu = true, const bool bPlaySound = true );
-	void ShowEditorUIPanel( class AVREditorFloatingUI* Panel, UVREditorInteractor* Interactor, const bool bShouldShow, const bool OnHand = false, const bool bRefreshQuickMenu = true, const bool bPlaySound = true );
+	void ShowEditorUIPanel(const class UWidgetComponent* WidgetComponent, UVREditorInteractor* Interactor, const bool bShouldShow, const bool bSpawnInFront = false, const bool bDragFromOpen = false, const bool bPlaySound = true);
+	void ShowEditorUIPanel(const EEditorUIPanel EditorUIPanel, UVREditorInteractor* Interactor, const bool bShouldShow, const bool bSpawnInFront = false, const bool bDragFromOpen = false, const bool bPlaySound = true);
+	void ShowEditorUIPanel(class AVREditorFloatingUI* Panel, UVREditorInteractor* Interactor, const bool bShouldShow, const bool bSpawnInFront = false, const bool bDragFromOpen = false, const bool bPlaySound = true);
 
 	/** Returns true if the radial menu is visible on this hand */
 	bool IsShowingRadialMenu(const UVREditorInteractor* Interactor ) const;
 
 	/** Tries to spawn the radial menu (if the specified hand isn't doing anything else) */
-	void TryToSpawnRadialMenu( UVREditorInteractor* Interactor, const bool bForceOverUI );
+	void TryToSpawnRadialMenu( UVREditorInteractor* Interactor, const bool bForceRefresh, const bool bPlaySound = true );
 
 	/** Hides the radial menu if the specified hand is showing it */
-	void HideRadialMenu( UVREditorInteractor* Interactor );
+	void HideRadialMenu( UVREditorInteractor* Interactor, const bool bPlaySound = true );
 
 	/** Start dragging a dock window on the hand */
-	void StartDraggingDockUI( class AVREditorDockableWindow* InitDraggingDockUI, UVREditorInteractor* Interactor, const float DockSelectDistance );
+	void StartDraggingDockUI( class AVREditorDockableWindow* InitDraggingDockUI, UVREditorInteractor* Interactor, const float DockSelectDistance, const bool bPlaySound = true );
 
 	/** Makes up a transform for a dockable UI when dragging it with a laser at the specified distance from the laser origin */
 	FTransform MakeDockableUITransformOnLaser( AVREditorDockableWindow* InitDraggingDockUI, UVREditorInteractor* Interactor, const float DockSelectDistance ) const;
@@ -174,13 +174,16 @@ public:
 	class AVREditorDockableWindow* GetDraggingDockUI() const;
 
 	/** Resets all values to stop dragging the current dock window */
-	void StopDraggingDockUI( UVREditorInteractor* VREditorInteractor );
+	void StopDraggingDockUI( UVREditorInteractor* VREditorInteractor, const bool bPlaySound = true );
 
 	/** Are we currently dragging a dock window */
-	bool IsDraggingDockUI();	
+	bool IsDraggingDockUI();
 
 	bool IsInteractorDraggingDockUI( const UVREditorInteractor* Interactor ) const;
 
+	/** If a panel was opened and dragged with the UI interactor */
+	bool IsDraggingPanelFromOpen() const;
+
 	/** Hides and unhides all the editor UI panels */
 	void TogglePanelsVisibility();
 
@@ -194,7 +197,7 @@ public:
 	void TogglePanelVisibility( const EEditorUIPanel EditorUIPanel );
 
 	/** Returns the radial widget so other classes, like the interactors, can access its functionality */
-	const TSharedPtr<SWidget>& GetRadialWidget();
+	class AVREditorRadialFloatingUI* GetRadialMenuFloatingUI() const;
 
 	/** 
 	 * Finds a widget with a given name inside the Content argument 
@@ -203,23 +206,57 @@ public:
 	 */
 	static const TSharedRef<SWidget>& FindWidgetOfType(const TSharedRef<SWidget>& Content, FName WidgetType);
 
+	/**
+	* Finds a widget with a given name inside the Content argument
+	* @param Content The widget to begin searching in
+	* @param The FName of the widget type to find.
+	*/
+	static const bool FindAllWidgetsOfType(TArray<TSharedRef<SWidget>>& FoundWidgets, const TSharedRef<SWidget>& Content, FName WidgetType);
+
 	/** Function to force an update of the Sequencer UI based on a change */
 	void UpdateSequencerUI();
 
+	/** Transition the user widgets to a new world */
+	void TransitionWorld(UWorld* NewWorld);
+
+	UVRRadialMenuHandler* GetRadialMenuHandler()
+	{
+		return RadialMenuHandler;
+	}
+
+	/** Called when a laser or simulated mouse hover enters a button */
+	void OnHoverBeginEffect(UVREditorWidgetComponent* Button);
+	/** Called when a laser or simulated mouse hover leaves a button */
+	void OnHoverEndEffect(UVREditorWidgetComponent* Button);
+
+	/** Set if sequencer was opened from the radial menu */
+	void SequencerOpenendFromRadialMenu(const bool bInOpenedFromRadialMenu = true);
+
+	/** If a dockable window can be scaled */
+	bool CanScalePanel() const;
+
+	/** Get the interactor that holds the radial menu */
+	class UVREditorMotionControllerInteractor* GetUIInteractor();
+
 protected:
 
+	/** Called to "preview" an input event to get a first chance at it. */
+	void OnPreviewInputAction( class FEditorViewportClient& ViewportClient, UViewportInteractor* Interactor,
+		const struct FViewportActionKeyInput& Action, bool& bOutIsInputCaptured, bool& bWasHandled );
+
 	/** Called when the user presses a button on their motion controller device */
 	void OnVRAction( class FEditorViewportClient& ViewportClient, UViewportInteractor* Interactor,
 		const struct FViewportActionKeyInput& Action, bool& bOutIsInputCaptured, bool& bWasHandled );
 
-	void SequencerPlayback(UViewportInteractor* Interactor, class UVREditorMotionControllerInteractor* VREditorInteractor);
-
 	/** Called every frame to update hover state */
 	void OnVRHoverUpdate( UViewportInteractor* Interactor, FVector& HoverImpactPoint, bool& bWasHandled );
 
 	/** Testing Slate UI */
 	void CreateUIs();
 
+	/** Called by proxy tab manager to ask us whether we support the specified tab ID.  bOutIsTabSupported defaults to true. */
+	void IsProxyTabSupported( FTabId TabId, bool& bOutIsTabSupported );
+
 	/**
 	 * Called when the injected proxy tab manager reports a new tab has been launched, 
 	 * signaling an asset editor has been launched, but really it could be any major global tab.
@@ -235,9 +272,6 @@ protected:
 	/** Called when an asset editor opens an asset while in VR Editor Mode. */
 	void OnAssetEditorOpened(UObject* Asset);
 
-	/** Sets the main windows to their default transform */
-	void SetDefaultWindowLayout();
-
 	/** Can be used when the tab manager is restoring a dockable tab area */
 	void DockableAreaRestored();
 
@@ -250,27 +284,40 @@ protected:
 	/**  Makes a uniform grid widget from the menu information contained in a MultiBox and MultiBoxWidget */
 	void MakeUniformGridMenu( const TSharedRef<FMultiBox>& MultiBox, const TSharedRef<SMultiBoxWidget>& MultiBoxWidget, int32 Columns);
 	/**  Makes a radial box widget from the menu information contained in a MultiBox and MultiBoxWidget */
-	void MakeRadialBoxMenu(const TSharedRef<FMultiBox>& MultiBox, const TSharedRef<SMultiBoxWidget>& MultiBoxWidget, float RadiusRatioOverride);
+	void MakeRadialBoxMenu(const TSharedRef<FMultiBox>& MultiBox, const TSharedRef<SMultiBoxWidget>& MultiBoxWidget, float RadiusRatioOverride, FName ButtonTypeOverride);
 
 	/** Adds a hoverable button of a given type to an overlay, using menu data from a BlockWidget */
 	TSharedRef<SWidget> AddHoverableButton(TSharedRef<SWidget>& BlockWidget, FName ButtonType, TSharedRef<SOverlay>& TestOverlay);
 	/** Sets the text wrap size of the text block element nested in a BlockWidget */
-	TSharedRef<SWidget> SetButtonTextWrap(TSharedRef<SWidget>& BlockWidget, float WrapSize);
+	TSharedRef<SWidget> SetButtonFormatting(TSharedRef<SWidget>& BlockWidget, float WrapSize);
 	
 	/** Builds the quick menu Slate widget */
 	TSharedRef<SWidget> BuildQuickMenuWidget();
 	/** Builds the radial menu Slate widget */
-	TSharedRef<SWidget> BuildRadialMenuWidget();
+	void BuildRadialMenuWidget();
 	/** Builds the numpad Slate widget */
-	TSharedRef<SWidget> BuildNumPadWidget();
+	void BuildNumPadWidget();
 	/** Swaps the content of the radial menu between the radial menu and the numpad */
 	void SwapRadialMenu();
 
-	/** Called when a laser or simulated mouse hover enters a button */
-	void OnHoverBeginEffect(TSharedRef<SButton> Button);
-	/** Called when a laser or simulated mouse hover leaves a button */
-	void OnHoverEndEffect(TSharedRef<SButton> Button);
+	/** Creates the sequencer radial menu to pass to the radial menu generator */
+	void SequencerRadialMenuGenerator(FMenuBuilder MenuBuilder, TSharedPtr<FUICommandList> CommandList, class UVREditorMode* VRMode, float& RadiusOverride);
 
+	/**
+	* Handles being notified when any editor mode changes to see if any VR Editor UI needs to change.
+	* @param Mode The mode that changed.
+	* @param IsEnabled true if the mode is enabled, otherwise false.
+	*/
+	void HandleEditorModeChanged(class FEdMode* Mode, bool IsEnabled);
+
+	/** Reset function that puts you back in placement mode, closes all UIs, etc. */
+	void ResetAll();
+
+	/**The VR Editor UI System's rules for when drag drop should be checked for */
+	bool CheckForVRDragDrop();
+
+	/** Preview the UI panel's location if spawning with the UI interactor, else spawn immediately */
+	bool ShouldPreviewPanel();
 
 protected:
 
@@ -292,17 +339,11 @@ protected:
 
 	/** The Radial Menu UI */
 	UPROPERTY()
-	AVREditorFloatingUI* QuickRadialMenu;
+	class AVREditorRadialFloatingUI* QuickRadialMenu;
 
 	/** The time since the radial menu was updated */
 	float RadialMenuHideDelayTime;
 
-	/** Tutorial menu widget class */
-	UClass* TutorialWidgetClass;
-
-	/** Pointer to the radial menu widget */
-	TSharedPtr<SWidget> RadialWidget;
-
 	/** True if the radial menu was visible when the content was swapped */
 	bool bRadialMenuVisibleAtSwap;
 
@@ -327,6 +368,16 @@ protected:
 	UPROPERTY()
 	class AVREditorDockableWindow* ColorPickerUI;
 
+	//
+	// Asymmetry
+	//
+	/** Interactor that has a laser and is generally interacting with the scene */
+	UPROPERTY()
+	class UVREditorMotionControllerInteractor* LaserInteractor;
+	/** Interactor that usually accesses UI and other helper functionality */
+	UPROPERTY()
+	class UVREditorMotionControllerInteractor* UIInteractor;
+
 	//
 	// Tab Manager UI
 	//
@@ -337,44 +388,37 @@ protected:
 	/** Set to true when we need to refocus the viewport. */
 	bool bRefocusViewport;
 
+	/** The last dragged hover location by the laser */
+	FVector LastDraggingHoverLocation;
 
-	//
-	// Sounds
-	//
-
-	/** Start dragging UI sound */
-	UPROPERTY()
-	class USoundCue* StartDragUISound;
-
-	/** Stop dragging UI sound */
+	/** All buttons created for the radial and quick menus */
 	UPROPERTY()
-	class USoundCue* StopDragUISound;
+	TArray<FVRButton> VRButtons;
 
-	/** Hide UI sound */
+	/** The add-on that handles radial menu switching */
 	UPROPERTY()
-	class USoundCue* HideUISound;
+	UVRRadialMenuHandler* RadialMenuHandler;
 
-	/** Show UI sound */
-	UPROPERTY()
-	class USoundCue* ShowUISound;
+	/** When replacing the actions menu, store off any existing actions */
+	FOnRadialMenuGenerated ExistingActionsMenu;
 
-	/** Button press sound */
-	UPROPERTY()
-	class USoundCue* ButtonPressSound;
+	/** When replacing the actions menu, store off the name of the existing actions menu */
+	FText ExistingActionsMenuLabel;
 
-	/** If the current dragged dock passed a certain distance if dragged from a hand */
-	bool bDraggedDockFromHandPassedThreshold;
+	/** The time the modifier was pressed at to spawn the menu */
+	FTimespan RadialMenuModifierSpawnTime;
 
-	/** The last dragged hover location by the laser */
-	FVector LastDraggingHoverLocation;
+	/** If sequenced was opened from the radial menu or somewhere else such as the content browser. */
+	bool bSequencerOpenedFromRadialMenu;
 
-	/** Default transforms */
-	TArray<FTransform> DefaultWindowTransforms;
+	/** If started dragging from opening a UI panel */
+	bool bDragPanelFromOpen;
 
-	/** All buttons created for the radial and quick menus */
-	TArray<FVRButton> VRButtons;
+	/** The time dragging a panel that was opened resulting in an instant drag */
+	float DragPanelFromOpenTime;
 
-	/** If this is the first time using TogglePanelsVisibility */
-	bool bSetDefaultLayout;
+	/** When started dragging from the radial menu we want the analog stick to be reset before the user is allowed to scale the panel. 
+	    Otherwise the panel will immediately start scaling because the user is using the analog stick to aim at the radial menu items. */
+	bool bPanelCanScale;
 };
 
diff --git a/Engine/Source/Editor/VREditor/UI/VRRadialMenuHandler.cpp b/Engine/Source/Editor/VREditor/UI/VRRadialMenuHandler.cpp
new file mode 100644
--- /dev/null
+++ b/Engine/Source/Editor/VREditor/UI/VRRadialMenuHandler.cpp
@@ -0,0 +1,824 @@
+// Copyright 1998-2017 Epic Games, Inc. All Rights Reserved.
+
+#include "VRRadialMenuHandler.h"
+
+#include "Widgets/SNullWidget.h"
+#include "Widgets/DeclarativeSyntaxSupport.h"
+#include "Widgets/SWidget.h"
+#include "Widgets/SWindow.h"
+#include "Widgets/Input/SButton.h"
+#include "Widgets/Layout/SUniformGridPanel.h"
+#include "Widgets/Layout/SBorder.h"
+#include "Widgets/Layout/SBox.h"
+#include "Widgets/Images/SImage.h"
+
+#include "Framework/MultiBox/MultiBox.h"
+#include "Framework/MultiBox/MultiBoxBuilder.h"
+#include "Framework/Application/SlateApplication.h"
+#include "Framework/Docking/TabManager.h"
+#include "Layout/WidgetPath.h"
+
+#include "Engine/EngineTypes.h"
+#include "Editor/UnrealEdEngine.h"
+#include "Kismet/GameplayStatics.h"
+#include "GameFramework/WorldSettings.h"
+#include "Sound/SoundCue.h"
+#include "EditorModeManager.h"
+#include "EditorModes.h"
+#include "UnrealEdGlobals.h"
+
+#include "VREditorMode.h"
+#include "VREditorUISystem.h"
+#include "VREditorBaseActor.h"
+#include "VREditorBaseUserWidget.h"
+#include "VREditorFloatingUI.h"
+#include "VREditorDockableWindow.h"
+#include "ViewportInteractionTypes.h"
+#include "IHeadMountedDisplay.h"
+#include "ViewportWorldInteraction.h"
+#include "VREditorInteractor.h"
+#include "VREditorMotionControllerInteractor.h"
+#include "IVREditorModule.h"
+#include "Settings/EditorExperimentalSettings.h"
+
+// UI
+#include "Components/WidgetComponent.h"
+#include "VREditorWidgetComponent.h"
+#include "VREditorStyle.h"
+#include "SUniformGridPanel.h"
+#include "SScrollBox.h"
+
+// Content Browser
+#include "CollectionManagerTypes.h"
+#include "IContentBrowserSingleton.h"
+#include "ContentBrowserModule.h"
+
+// World Outliner
+#include "Editor/SceneOutliner/Public/ISceneOutliner.h"
+#include "Editor/SceneOutliner/Public/SceneOutlinerPublicTypes.h"
+#include "Editor/SceneOutliner/Public/SceneOutlinerModule.h"
+
+// Actor Details, Modes
+#include "SEditorViewport.h"
+#include "Toolkits/AssetEditorManager.h"
+#include "ILevelEditor.h"
+#include "LevelEditor.h"
+
+// World Settings
+#include "PropertyEditorModule.h"
+#include "IDetailsView.h"
+
+#include "SLevelViewport.h"
+#include "SScaleBox.h"
+#include "SDPIScaler.h"
+#include "SWidget.h"
+#include "SDockTab.h"
+#include "Widgets/Layout/SDPIScaler.h"
+#include "Widgets/Docking/SDockTab.h"
+#include "Widgets/Colors/SColorPicker.h"
+#include "MultiBox.h"
+#include "LevelEditorActions.h"
+#include "VREditorActions.h"
+#include "GenericCommands.h"
+#include "ISequencer.h"
+#include "Private/SLevelEditor.h"
+#include "DebuggerCommands.h"
+#include "UMGStyle.h"
+
+#include "UI/VREditorRadialFloatingUI.h"
+
+#define LOCTEXT_NAMESPACE "VREditor"
+
+FText UVRRadialMenuHandler::ActionMenuLabel = LOCTEXT("DefaultActions", "Actions");
+
+UVRRadialMenuHandler::UVRRadialMenuHandler(const FObjectInitializer& Initializer) :
+	Super(Initializer)
+{
+	
+}
+
+void UVRRadialMenuHandler::BackOutMenu()
+{
+	if (MenuStack.Num() > 0)
+	{
+		FOnRadialMenuGenerated PreviousMenu = MenuStack.Pop();
+		const bool bShouldAddToStack = false;
+		RegisterMenuGenerator(PreviousMenu, bShouldAddToStack);
+	}
+}
+
+void UVRRadialMenuHandler::Home()
+{
+	MenuStack.Empty();
+	const bool bShouldAddToStack = false;
+	RegisterMenuGenerator(HomeMenu, bShouldAddToStack);
+}
+
+void UVRRadialMenuHandler::Init(UVREditorUISystem* InUISystem)
+{
+	UIOwner = InUISystem;
+	HomeMenu.BindUObject(this, &UVRRadialMenuHandler::HomeMenuGenerator);
+	SnapMenu.BindUObject(this, &UVRRadialMenuHandler::SnapMenuGenerator);
+	GizmoMenu.BindUObject(this, &UVRRadialMenuHandler::GizmoMenuGenerator);
+	UIMenu.BindUObject(this, &UVRRadialMenuHandler::UIMenuGenerator);
+	EditMenu.BindUObject(this, &UVRRadialMenuHandler::EditMenuGenerator);
+	ToolsMenu.BindUObject(this, &UVRRadialMenuHandler::ToolsMenuGenerator);
+	ModesMenu.BindUObject(this, &UVRRadialMenuHandler::ModesMenuGenerator);
+
+	Home();
+}
+
+void UVRRadialMenuHandler::BuildRadialMenuCommands(FMenuBuilder MenuBuilder, TSharedPtr<FUICommandList> CommandList, UVREditorMode* VRMode, float& RadiusOverride)
+{
+	UVRRadialMenuHandler::OnRadialMenuGenerated.ExecuteIfBound(MenuBuilder, CommandList, VRMode, RadiusOverride);
+
+	
+
+	if (UIOwner != nullptr && UIOwner->GetRadialMenuFloatingUI() != nullptr)
+	{
+		TSharedPtr<SWidget> HomeWidget = SNullWidget::NullWidget;
+		switch (MenuStack.Num())
+		{
+			case 0:
+			{
+				HomeWidget = SNew(SImage)
+					.Image(FVREditorStyle::GetBrush("VREditorStyle.Home"));
+				break;
+			}
+			case 1:
+			{
+				HomeWidget = SNew(SImage)
+					.Image(FVREditorStyle::GetBrush("VREditorStyle.OneLevel"));
+				break;
+			}
+			case 2:
+			{
+				HomeWidget = SNew(SImage)
+					.Image(FVREditorStyle::GetBrush("VREditorStyle.TwoLevel"));
+				break;
+			}
+			default:
+			{
+				break;
+			}
+		}
+		if (HomeWidget != SNullWidget::NullWidget)
+		{
+			UIOwner->GetRadialMenuFloatingUI()->UpdateCentralWidgetComponent(HomeWidget);
+		}
+	}
+}
+
+void UVRRadialMenuHandler::HomeMenuGenerator(FMenuBuilder MenuBuilder, TSharedPtr<FUICommandList> CommandList, UVREditorMode* VRMode, float& RadiusOverride)
+{
+	// First menu entry is at 90 degrees 
+	MenuBuilder.AddMenuEntry(
+		LOCTEXT("SnapSettings", "Snapping"),
+		FText(),
+		FSlateIcon(FVREditorStyle::GetStyleSetName(), "VREditorStyle.SnapMenu"),
+		FUIAction
+		(
+			FExecuteAction::CreateUObject(this, &UVRRadialMenuHandler::RegisterMenuGenerator, SnapMenu, true),
+			FCanExecuteAction::CreateStatic(&FLevelEditorActionCallbacks::DefaultCanExecuteAction)
+			),
+		NAME_None,
+		EUserInterfaceActionType::CollapsedButton
+		);
+	MenuBuilder.AddMenuEntry(
+		LOCTEXT("GizmoModes", "Gizmo"),
+		FText(),
+		FSlateIcon(FVREditorStyle::GetStyleSetName(), "VREditorStyle.GizmoMenu"),
+		FUIAction
+		(
+			FExecuteAction::CreateUObject(this, &UVRRadialMenuHandler::RegisterMenuGenerator, GizmoMenu, true),
+			FCanExecuteAction::CreateStatic(&FLevelEditorActionCallbacks::DefaultCanExecuteAction)
+			),
+		NAME_None,
+		EUserInterfaceActionType::CollapsedButton
+		);
+	MenuBuilder.AddMenuEntry(
+		LOCTEXT("Windows", "Windows"),
+		FText(),
+		FSlateIcon(FVREditorStyle::GetStyleSetName(), "VREditorStyle.WindowsMenu"),
+		FUIAction
+		(
+			FExecuteAction::CreateUObject(this, &UVRRadialMenuHandler::RegisterMenuGenerator, UIMenu, true),
+			FCanExecuteAction::CreateStatic(&FLevelEditorActionCallbacks::DefaultCanExecuteAction)
+			),
+		NAME_None,
+		EUserInterfaceActionType::CollapsedButton
+		);
+	MenuBuilder.AddMenuEntry(
+		LOCTEXT("Edit", "Edit"),
+		FText(),
+		FSlateIcon(FVREditorStyle::GetStyleSetName(), "VREditorStyle.EditMenu"),
+		FUIAction
+		(
+			FExecuteAction::CreateUObject(this, &UVRRadialMenuHandler::RegisterMenuGenerator, EditMenu, true),
+			FCanExecuteAction::CreateStatic(&FLevelEditorActionCallbacks::DefaultCanExecuteAction)
+			),
+		NAME_None,
+		EUserInterfaceActionType::CollapsedButton
+		);
+	MenuBuilder.AddMenuEntry(
+		LOCTEXT("Tools", "Tools"),
+		FText(),
+		FSlateIcon(FVREditorStyle::GetStyleSetName(), "VREditorStyle.ToolsMenu"),
+		FUIAction
+		(
+			FExecuteAction::CreateUObject(this, &UVRRadialMenuHandler::RegisterMenuGenerator, ToolsMenu, true),
+			FCanExecuteAction::CreateStatic(&FLevelEditorActionCallbacks::DefaultCanExecuteAction)
+			),
+		NAME_None,
+		EUserInterfaceActionType::CollapsedButton
+		);
+	MenuBuilder.AddMenuEntry(
+		LOCTEXT("Modes", "Modes"),
+		FText(),
+		FSlateIcon(FVREditorStyle::GetStyleSetName(), "VREditorStyle.ModesMenu"),
+		FUIAction
+		(
+			FExecuteAction::CreateUObject(this, &UVRRadialMenuHandler::RegisterMenuGenerator, ModesMenu, true),
+			FCanExecuteAction::CreateStatic(&FLevelEditorActionCallbacks::DefaultCanExecuteAction)
+			),
+		NAME_None,
+		EUserInterfaceActionType::CollapsedButton
+		);
+	TAttribute<FText> DynamicActionsLabel;
+	DynamicActionsLabel.BindStatic(&UVRRadialMenuHandler::GetActionMenuLabel);
+	MenuBuilder.AddMenuEntry(
+		DynamicActionsLabel,
+		FText(),
+		FSlateIcon(FVREditorStyle::GetStyleSetName(), "VREditorStyle.ActionsMenu"),
+		FUIAction
+		(
+			FExecuteAction::CreateUObject(this, &UVRRadialMenuHandler::RegisterMenuGenerator, ActionsMenu, true),
+			FCanExecuteAction::CreateUObject(this, &UVRRadialMenuHandler::IsActionMenuBound)
+		),
+		NAME_None,
+		EUserInterfaceActionType::CollapsedButton
+		);
+}
+
+void UVRRadialMenuHandler::SnapMenuGenerator(FMenuBuilder MenuBuilder, TSharedPtr<FUICommandList> CommandList, UVREditorMode* VRMode, float& RadiusOverride)
+{
+	FVREditorActionCallbacks::UpdateSelectingCandidateActorsText(VRMode);
+	MenuBuilder.AddMenuEntry(
+		LOCTEXT("ToggleTranslationSnap", "Translate Snap"),
+		LOCTEXT("ToggleTranslationSnapTooltip", "Toggle Translation Snap"),
+		FSlateIcon(FVREditorStyle::GetStyleSetName(), "VREditorStyle.TranslateSnap"),
+		FUIAction
+		(
+			FExecuteAction::CreateStatic(&FLevelEditorActionCallbacks::LocationGridSnap_Clicked),
+			FCanExecuteAction::CreateStatic(&FLevelEditorActionCallbacks::DefaultCanExecuteAction),
+			FGetActionCheckState::CreateStatic(&FVREditorActionCallbacks::GetTranslationSnapState)
+			),
+		NAME_None,
+		EUserInterfaceActionType::ToggleButton
+		);
+	TAttribute<FText> DynamicTranslationSizeLabel;
+	DynamicTranslationSizeLabel.BindStatic(&FVREditorActionCallbacks::GetTranslationSnapSizeText);
+	MenuBuilder.AddMenuEntry(
+		DynamicTranslationSizeLabel,
+		LOCTEXT("ToggleTranslationSnapSizeTooltip", "Toggle Translation Snap Size"),
+		FSlateIcon(FVREditorStyle::GetStyleSetName(), "VREditorStyle.GridNum"),
+		FUIAction
+		(
+			FExecuteAction::CreateStatic(&FVREditorActionCallbacks::OnTranslationSnapSizeButtonClicked),
+			FCanExecuteAction::CreateStatic(&FLevelEditorActionCallbacks::DefaultCanExecuteAction)
+			),
+		NAME_None,
+		EUserInterfaceActionType::ToggleButton
+		);
+	MenuBuilder.AddMenuEntry(
+		LOCTEXT("ToggleRotationSnap", "Rotate Snap"),
+		LOCTEXT("ToggleRotationSnapTooltip", "Toggle Rotation Snap"),
+		FSlateIcon(FVREditorStyle::GetStyleSetName(), "VREditorStyle.RotateSnap"),
+		FUIAction
+		(
+			FExecuteAction::CreateStatic(&FLevelEditorActionCallbacks::RotationGridSnap_Clicked),
+			FCanExecuteAction::CreateStatic(&FLevelEditorActionCallbacks::DefaultCanExecuteAction),
+			FGetActionCheckState::CreateStatic(&FVREditorActionCallbacks::GetRotationSnapState)
+			),
+		NAME_None,
+		EUserInterfaceActionType::ToggleButton
+		);
+	TAttribute<FText> DynamicRotationSizeLabel;
+	DynamicRotationSizeLabel.BindStatic(&FVREditorActionCallbacks::GetRotationSnapSizeText);
+	MenuBuilder.AddMenuEntry(
+		DynamicRotationSizeLabel,
+		LOCTEXT("ToggleRotationSnapSizeTooltip", "Toggle Rotation Snap Size"),
+		FSlateIcon(FVREditorStyle::GetStyleSetName(), "VREditorStyle.AngleNum"),
+		FUIAction
+		(
+			FExecuteAction::CreateStatic(&FVREditorActionCallbacks::OnRotationSnapSizeButtonClicked),
+			FCanExecuteAction::CreateStatic(&FLevelEditorActionCallbacks::DefaultCanExecuteAction)
+			),
+		NAME_None,
+		EUserInterfaceActionType::CollapsedButton
+		);
+	MenuBuilder.AddMenuEntry(
+		LOCTEXT("ToggleScaleSnap", "Scale Snap"),
+		LOCTEXT("ToggleScaleSnapTooltip", "Toggle Scale Snap"),
+		FSlateIcon(FVREditorStyle::GetStyleSetName(), "VREditorStyle.ScaleSnap"),
+		FUIAction
+		(
+			FExecuteAction::CreateStatic(&FLevelEditorActionCallbacks::ScaleGridSnap_Clicked),
+			FCanExecuteAction::CreateStatic(&FLevelEditorActionCallbacks::DefaultCanExecuteAction),
+			FGetActionCheckState::CreateStatic(&FVREditorActionCallbacks::GetScaleSnapState)
+			),
+		NAME_None,
+		EUserInterfaceActionType::ToggleButton
+		);
+	TAttribute<FText> DynamicScaleSizeLabel;
+	DynamicScaleSizeLabel.BindStatic(&FVREditorActionCallbacks::GetScaleSnapSizeText);
+	MenuBuilder.AddMenuEntry(
+		DynamicScaleSizeLabel,
+		LOCTEXT("ToggleScaleSnapSizeTooltip", "Toggle Scale Snap Size"),
+		FSlateIcon(FVREditorStyle::GetStyleSetName(), "VREditorStyle.ScaleNum"),
+		FUIAction
+		(
+			FExecuteAction::CreateStatic(&FVREditorActionCallbacks::OnScaleSnapSizeButtonClicked),
+			FCanExecuteAction::CreateStatic(&FLevelEditorActionCallbacks::DefaultCanExecuteAction)
+			),
+		NAME_None,
+		EUserInterfaceActionType::CollapsedButton
+		);
+	MenuBuilder.AddMenuEntry(
+		LOCTEXT("SmartSnapping", "Smart Snapping"),
+		LOCTEXT("AlignToActorsTooltip", "Align to Actors as you transform an object"),
+		FSlateIcon(FVREditorStyle::GetStyleSetName(), "VREditorStyle.AlignActors"),
+		FUIAction
+		(
+			FExecuteAction::CreateStatic(&FVREditorActionCallbacks::ToggleAligningToActors, VRMode),
+			FCanExecuteAction::CreateStatic(&FLevelEditorActionCallbacks::DefaultCanExecuteAction),
+			FGetActionCheckState::CreateStatic(&FVREditorActionCallbacks::AreAligningToActors, VRMode)
+			),
+		NAME_None,
+		EUserInterfaceActionType::ToggleButton
+		);
+
+	TAttribute<FText> DynamicAlignSelectionLabel;
+	DynamicAlignSelectionLabel.BindStatic(&FVREditorActionCallbacks::GetSelectingCandidateActorsText);
+	MenuBuilder.AddMenuEntry(
+		DynamicAlignSelectionLabel,
+		FText(),
+		FSlateIcon(FVREditorStyle::GetStyleSetName(), "VREditorStyle.SetTargets"),
+		FUIAction
+		(
+			FExecuteAction::CreateStatic(&FVREditorActionCallbacks::ToggleSelectingCandidateActors, VRMode),
+			FCanExecuteAction::CreateStatic(&FVREditorActionCallbacks::CanSelectCandidateActors, VRMode)
+			),
+		NAME_None,
+		EUserInterfaceActionType::CollapsedButton
+		);
+}
+
+void UVRRadialMenuHandler::GizmoMenuGenerator(FMenuBuilder MenuBuilder, TSharedPtr<FUICommandList> CommandList, UVREditorMode* VRMode, float& RadiusOverride)
+{
+	MenuBuilder.AddMenuEntry(
+		LOCTEXT("LocalSpace", "Local Space"),
+		FText(),
+		FSlateIcon(FVREditorStyle::GetStyleSetName(), "VREditorStyle.LocalSpace"),
+		FUIAction
+		(
+			FExecuteAction::CreateStatic(&FVREditorActionCallbacks::SetCoordinateSystem, VRMode, ECoordSystem::COORD_Local),
+			FCanExecuteAction::CreateStatic(&FLevelEditorActionCallbacks::DefaultCanExecuteAction),
+			FGetActionCheckState::CreateStatic(&FVREditorActionCallbacks::IsActiveCoordinateSystem, VRMode, ECoordSystem::COORD_Local)
+			),
+		NAME_None,
+		EUserInterfaceActionType::ToggleButton
+		);
+	MenuBuilder.AddMenuEntry(
+		LOCTEXT("WorldSpace", "World Space"),
+		FText(),
+		FSlateIcon(FVREditorStyle::GetStyleSetName(), "VREditorStyle.WorldSpace"),
+		FUIAction
+		(
+			FExecuteAction::CreateStatic(&FVREditorActionCallbacks::SetCoordinateSystem, VRMode, ECoordSystem::COORD_World),
+			FCanExecuteAction::CreateStatic(&FLevelEditorActionCallbacks::DefaultCanExecuteAction),
+			FGetActionCheckState::CreateStatic(&FVREditorActionCallbacks::IsActiveCoordinateSystem, VRMode, ECoordSystem::COORD_World)
+		),
+		NAME_None,
+		EUserInterfaceActionType::ToggleButton
+		);
+	MenuBuilder.AddMenuEntry(
+		LOCTEXT("Universal", "Universal"),
+		FText(),
+		FSlateIcon(FVREditorStyle::GetStyleSetName(), "VREditorStyle.Universal"),
+		FUIAction
+		(
+			FExecuteAction::CreateStatic(&FVREditorActionCallbacks::SetGizmoMode, VRMode, EGizmoHandleTypes::All),
+			FCanExecuteAction::CreateStatic(&FLevelEditorActionCallbacks::DefaultCanExecuteAction),
+			FGetActionCheckState::CreateStatic(&FVREditorActionCallbacks::IsActiveGizmoMode, VRMode, EGizmoHandleTypes::All)
+			),
+		NAME_None,
+		EUserInterfaceActionType::ToggleButton
+		);
+	MenuBuilder.AddMenuEntry(
+		LOCTEXT("Translate", "Translate"),
+		FText(),
+		FSlateIcon(FVREditorStyle::GetStyleSetName(), "VREditorStyle.Translate"),
+		FUIAction
+		(
+			FExecuteAction::CreateStatic(&FVREditorActionCallbacks::SetGizmoMode, VRMode, EGizmoHandleTypes::Translate),
+			FCanExecuteAction::CreateStatic(&FLevelEditorActionCallbacks::DefaultCanExecuteAction),
+			FGetActionCheckState::CreateStatic(&FVREditorActionCallbacks::IsActiveGizmoMode, VRMode, EGizmoHandleTypes::Translate)
+			),
+		NAME_None,
+		EUserInterfaceActionType::ToggleButton
+		);
+	MenuBuilder.AddMenuEntry(
+		LOCTEXT("Rotate", "Rotate"),
+		FText(),
+		FSlateIcon(FVREditorStyle::GetStyleSetName(), "VREditorStyle.Rotate"),
+		FUIAction
+		(
+			FExecuteAction::CreateStatic(&FVREditorActionCallbacks::SetGizmoMode, VRMode, EGizmoHandleTypes::Rotate),
+			FCanExecuteAction::CreateStatic(&FLevelEditorActionCallbacks::DefaultCanExecuteAction),
+			FGetActionCheckState::CreateStatic(&FVREditorActionCallbacks::IsActiveGizmoMode, VRMode, EGizmoHandleTypes::Rotate)
+			),
+		NAME_None,
+		EUserInterfaceActionType::ToggleButton
+		);
+	MenuBuilder.AddMenuEntry(
+		LOCTEXT("Scale", "Scale"),
+		FText(),
+		FSlateIcon(FVREditorStyle::GetStyleSetName(), "VREditorStyle.Scale"),
+		FUIAction
+		(
+			FExecuteAction::CreateStatic(&FVREditorActionCallbacks::SetGizmoMode, VRMode, EGizmoHandleTypes::Scale),
+			FCanExecuteAction::CreateStatic(&FLevelEditorActionCallbacks::DefaultCanExecuteAction),
+			FGetActionCheckState::CreateStatic(&FVREditorActionCallbacks::IsActiveGizmoMode, VRMode, EGizmoHandleTypes::Scale)
+			),
+		NAME_None,
+		EUserInterfaceActionType::ToggleButton
+		);
+}
+
+void UVRRadialMenuHandler::UIMenuGenerator(FMenuBuilder MenuBuilder, TSharedPtr<FUICommandList> CommandList, UVREditorMode* VRMode, float& RadiusOverride)
+{
+	MenuBuilder.AddMenuEntry(
+		LOCTEXT("ActorDetails", "Details"),
+		LOCTEXT("ActorDetailsTooltip", "Details"),
+		FSlateIcon(FVREditorStyle::GetStyleSetName(), "VREditorStyle.Details"),
+		FUIAction
+		(
+			FExecuteAction::CreateStatic(&FVREditorActionCallbacks::OnUIToggleButtonClicked, VRMode, UVREditorUISystem::EEditorUIPanel::ActorDetails),
+			FCanExecuteAction::CreateStatic(&FLevelEditorActionCallbacks::DefaultCanExecuteAction),
+			FGetActionCheckState::CreateStatic(&FVREditorActionCallbacks::GetUIToggledState, VRMode, UVREditorUISystem::EEditorUIPanel::ActorDetails)
+			),
+		NAME_None,
+		EUserInterfaceActionType::ToggleButton
+		);
+	MenuBuilder.AddMenuEntry(
+		LOCTEXT("ContentBrowser", "Content Browser"),
+		FText(),
+		FSlateIcon(FVREditorStyle::GetStyleSetName(), "VREditorStyle.ContentBrowser"),
+		FUIAction
+		(
+			FExecuteAction::CreateStatic(&FVREditorActionCallbacks::OnUIToggleButtonClicked, VRMode, UVREditorUISystem::EEditorUIPanel::ContentBrowser),
+			FCanExecuteAction::CreateStatic(&FLevelEditorActionCallbacks::DefaultCanExecuteAction),
+			FGetActionCheckState::CreateStatic(&FVREditorActionCallbacks::GetUIToggledState, VRMode, UVREditorUISystem::EEditorUIPanel::ContentBrowser)
+			),
+		NAME_None,
+		EUserInterfaceActionType::ToggleButton
+		);
+	MenuBuilder.AddMenuEntry(
+		LOCTEXT("Modes", "Modes Panel"),
+		FText(),
+		FSlateIcon(FVREditorStyle::GetStyleSetName(), "VREditorStyle.ModesPanel"),
+		FUIAction
+		(
+			FExecuteAction::CreateStatic(&FVREditorActionCallbacks::OnUIToggleButtonClicked, VRMode, UVREditorUISystem::EEditorUIPanel::Modes),
+			FCanExecuteAction::CreateStatic(&FLevelEditorActionCallbacks::DefaultCanExecuteAction),
+			FGetActionCheckState::CreateStatic(&FVREditorActionCallbacks::GetUIToggledState, VRMode, UVREditorUISystem::EEditorUIPanel::Modes)
+			),
+		NAME_None,
+		EUserInterfaceActionType::ToggleButton
+		);
+	MenuBuilder.AddMenuEntry(
+		LOCTEXT("WorldOutliner", "World Outliner"),
+		FText(),
+		FSlateIcon(FVREditorStyle::GetStyleSetName(), "VREditorStyle.WorldOutliner"),
+		FUIAction
+		(
+			FExecuteAction::CreateStatic(&FVREditorActionCallbacks::OnUIToggleButtonClicked, VRMode, UVREditorUISystem::EEditorUIPanel::WorldOutliner),
+			FCanExecuteAction::CreateStatic(&FLevelEditorActionCallbacks::DefaultCanExecuteAction),
+			FGetActionCheckState::CreateStatic(&FVREditorActionCallbacks::GetUIToggledState, VRMode, UVREditorUISystem::EEditorUIPanel::WorldOutliner)
+			),
+		NAME_None,
+		EUserInterfaceActionType::ToggleButton
+		);
+	MenuBuilder.AddMenuEntry(
+		LOCTEXT("WorldSettings", "World Settings"),
+		FText(),
+		FSlateIcon(FVREditorStyle::GetStyleSetName(), "VREditorStyle.WorldSettings"),
+		FUIAction
+		(
+			FExecuteAction::CreateStatic(&FVREditorActionCallbacks::OnUIToggleButtonClicked, VRMode, UVREditorUISystem::EEditorUIPanel::WorldSettings),
+			FCanExecuteAction::CreateStatic(&FLevelEditorActionCallbacks::DefaultCanExecuteAction),
+			FGetActionCheckState::CreateStatic(&FVREditorActionCallbacks::GetUIToggledState, VRMode, UVREditorUISystem::EEditorUIPanel::WorldSettings)
+			),
+		NAME_None,
+		EUserInterfaceActionType::ToggleButton
+		);
+	MenuBuilder.AddMenuEntry(
+		LOCTEXT("CreateNewSequence", "Create Sequence"),
+		FText(),
+		FSlateIcon(FVREditorStyle::GetStyleSetName(), "VREditorStyle.Sequencer"),
+		FUIAction
+		(
+			FExecuteAction::CreateStatic(&FVREditorActionCallbacks::CreateNewSequence, VRMode),
+			FCanExecuteAction::CreateStatic(&FLevelEditorActionCallbacks::DefaultCanExecuteAction)
+			),
+		NAME_None,
+		EUserInterfaceActionType::ToggleButton
+		);
+}
+
+void UVRRadialMenuHandler::EditMenuGenerator(FMenuBuilder MenuBuilder, TSharedPtr<FUICommandList> CommandList, UVREditorMode* VRMode, float& RadiusOverride)
+{
+	// First menu entry is at 90 degrees 
+	MenuBuilder.AddMenuEntry(
+		LOCTEXT("DeselectAll", "Deselect All"),
+		FText(),
+		FSlateIcon(FVREditorStyle::GetStyleSetName(), "VREditorStyle.DeselectAll"),
+		FUIAction
+		(
+			FExecuteAction::CreateStatic(&FVREditorActionCallbacks::DeselectAll),
+			FCanExecuteAction::CreateStatic(&FLevelEditorActionCallbacks::DefaultCanExecuteAction)
+			)
+		);
+	MenuBuilder.AddMenuEntry(
+		LOCTEXT("Delete", "Delete"),
+		FText(),
+		FSlateIcon(FVREditorStyle::GetStyleSetName(), "VREditorStyle.Delete"),
+		FUIAction
+		(
+			FExecuteAction::CreateStatic(&FLevelEditorActionCallbacks::ExecuteExecCommand, FString(TEXT("DELETE"))),
+			FCanExecuteAction::CreateStatic(&FLevelEditorActionCallbacks::Delete_CanExecute)
+			)
+		);
+	MenuBuilder.AddMenuEntry(
+		LOCTEXT("Cut", "Cut"),
+		FText(),
+		FSlateIcon(FVREditorStyle::GetStyleSetName(), "VREditorStyle.Cut"),
+		FUIAction
+		(
+			FExecuteAction::CreateStatic(&FLevelEditorActionCallbacks::ExecuteExecCommand, FString(TEXT("EDIT CUT"))),
+			FCanExecuteAction::CreateStatic(&FLevelEditorActionCallbacks::Cut_CanExecute)
+			)
+		);
+	MenuBuilder.AddMenuEntry(
+		LOCTEXT("Copy", "Copy"),
+		FText(),
+		FSlateIcon(FVREditorStyle::GetStyleSetName(), "VREditorStyle.Copy"),
+		FUIAction
+		(
+			FExecuteAction::CreateStatic(&FLevelEditorActionCallbacks::ExecuteExecCommand, FString(TEXT("EDIT COPY"))),
+			FCanExecuteAction::CreateStatic(&FLevelEditorActionCallbacks::Copy_CanExecute)
+			)
+		);
+	MenuBuilder.AddMenuEntry(
+		LOCTEXT("Duplicate", "Duplicate Selected"),
+		FText(),
+		FSlateIcon(FVREditorStyle::GetStyleSetName(), "VREditorStyle.Duplicate"),
+		FUIAction
+		(
+			FExecuteAction::CreateStatic(&FLevelEditorActionCallbacks::ExecuteExecCommand, FString(TEXT("DUPLICATE"))),
+			FCanExecuteAction::CreateStatic(&FLevelEditorActionCallbacks::Duplicate_CanExecute)
+			)
+		);
+	MenuBuilder.AddMenuEntry(
+		LOCTEXT("Paste", "Paste"),
+		FText(),
+		FSlateIcon(FVREditorStyle::GetStyleSetName(), "VREditorStyle.Paste"),
+		FUIAction
+		(
+			FExecuteAction::CreateStatic(&FLevelEditorActionCallbacks::ExecuteExecCommand, FString(TEXT("EDIT PASTE"))),
+			FCanExecuteAction::CreateStatic(&FLevelEditorActionCallbacks::Paste_CanExecute)
+			)
+		);
+
+	MenuBuilder.AddMenuEntry(
+		LOCTEXT("SnapToFloor", "Snap To Floor"),
+		FText(),
+		FSlateIcon(FVREditorStyle::GetStyleSetName(), "VREditorStyle.SnapToFloor"),
+		FUIAction
+		(
+			FExecuteAction::CreateStatic(&FVREditorActionCallbacks::OnSnapActorsToGroundClicked, VRMode),
+			FCanExecuteAction::CreateStatic(&FLevelEditorActionCallbacks::Copy_CanExecute)
+			)
+		);
+}
+
+void UVRRadialMenuHandler::ToolsMenuGenerator(FMenuBuilder MenuBuilder, TSharedPtr<FUICommandList> CommandList, UVREditorMode* VRMode, float& RadiusOverride)
+{
+	TAttribute<FText> DynamicSimulateLabel;
+	DynamicSimulateLabel.BindStatic(&FVREditorActionCallbacks::GetSimulateText);
+
+	MenuBuilder.AddMenuEntry(
+		DynamicSimulateLabel,
+		FText(),
+		FSlateIcon(FVREditorStyle::GetStyleSetName(), "VREditorStyle.Simulate"),
+		FUIAction
+		(
+			FExecuteAction::CreateStatic(&FVREditorActionCallbacks::OnSimulateButtonClicked, VRMode),
+			FCanExecuteAction::CreateStatic(&FLevelEditorActionCallbacks::DefaultCanExecuteAction)
+			),
+		NAME_None,
+		EUserInterfaceActionType::CollapsedButton
+		);
+
+
+	MenuBuilder.AddMenuEntry(
+		LOCTEXT("KeepSimulationChanges", "Save Actors"),
+		FText(),
+		FSlateIcon(FVREditorStyle::GetStyleSetName(), "VREditorStyle.SaveSimulation"),
+		FUIAction
+		(
+			FExecuteAction::CreateStatic(&FLevelEditorActionCallbacks::OnKeepSimulationChanges),
+			FCanExecuteAction::CreateStatic(&FLevelEditorActionCallbacks::CanExecuteKeepSimulationChanges)
+			),
+		NAME_None,
+		EUserInterfaceActionType::CollapsedButton
+		);
+
+
+	MenuBuilder.AddMenuEntry(
+		LOCTEXT("PauseSimulation", "Pause"),
+		FText(),
+		FSlateIcon(FVREditorStyle::GetStyleSetName(), "VREditorStyle.Pause"),
+		FUIAction
+		(
+			FExecuteAction::CreateStatic(&FPlayWorldCommandCallbacks::PausePlaySession_Clicked),
+			FCanExecuteAction::CreateStatic(&FPlayWorldCommandCallbacks::HasPlayWorldAndRunning)
+			),
+		NAME_None,
+		EUserInterfaceActionType::CollapsedButton
+		);
+
+	MenuBuilder.AddMenuEntry(
+		LOCTEXT("ResumeSimulation", "Resume"),
+		FText(),
+		FSlateIcon(FVREditorStyle::GetStyleSetName(), "VREditorStyle.Resume"),
+		FUIAction
+		(
+			FExecuteAction::CreateStatic(&FPlayWorldCommandCallbacks::ResumePlaySession_Clicked),
+			FCanExecuteAction::CreateStatic(&FPlayWorldCommandCallbacks::HasPlayWorldAndPaused)
+			),
+		NAME_None,
+		EUserInterfaceActionType::CollapsedButton
+		);
+
+	MenuBuilder.AddMenuEntry(
+		LOCTEXT("PlayInEditor", "Play"),
+		FText(),
+		FSlateIcon(FVREditorStyle::GetStyleSetName(), "VREditorStyle.Play"),
+		FUIAction
+		(
+			FExecuteAction::CreateStatic(&FVREditorActionCallbacks::OnPlayButtonClicked, VRMode),
+			FCanExecuteAction::CreateStatic(&FVREditorActionCallbacks::CanPlay)
+			),
+		NAME_None,
+		EUserInterfaceActionType::CollapsedButton
+		);
+
+	MenuBuilder.AddMenuEntry(
+		LOCTEXT("Screenshot", "Screenshot"),
+		FText(),
+		FSlateIcon(FVREditorStyle::GetStyleSetName(), "VREditorStyle.Screenshot"),
+		FUIAction
+		(
+			FExecuteAction::CreateStatic(&FVREditorActionCallbacks::OnScreenshotButtonClicked, VRMode),
+			FCanExecuteAction::CreateStatic(&FLevelEditorActionCallbacks::DefaultCanExecuteAction)
+			),
+		NAME_None,
+		EUserInterfaceActionType::CollapsedButton
+		);
+
+	MenuBuilder.AddMenuEntry(
+		LOCTEXT("Flashlight", "Flashlight"),
+		FText(),
+		FSlateIcon(FVREditorStyle::GetStyleSetName(), "VREditorStyle.Flashlight"),
+		FUIAction
+		(
+			FExecuteAction::CreateStatic(&FVREditorActionCallbacks::OnLightButtonClicked, VRMode),
+			FCanExecuteAction::CreateStatic(&FLevelEditorActionCallbacks::DefaultCanExecuteAction)
+			),
+		NAME_None,
+		EUserInterfaceActionType::CollapsedButton
+		);
+}
+
+void UVRRadialMenuHandler::ModesMenuGenerator(FMenuBuilder MenuBuilder, TSharedPtr<FUICommandList> CommandList, UVREditorMode* VRMode, float& RadiusOverride)
+{
+	MenuBuilder.AddMenuEntry(
+		LOCTEXT("Actors", "Actors"),
+		FText(),
+		FSlateIcon(FVREditorStyle::GetStyleSetName(), "VREditorStyle.ActorsMode"),
+		FUIAction
+		(
+			FExecuteAction::CreateStatic(&FVREditorActionCallbacks::ChangeEditorModes, FBuiltinEditorModes::EM_Placement),
+			FCanExecuteAction::CreateStatic(&FLevelEditorActionCallbacks::DefaultCanExecuteAction),
+			FGetActionCheckState::CreateStatic(&FVREditorActionCallbacks::EditorModeActive, FBuiltinEditorModes::EM_Placement)
+			),
+		NAME_None,
+		EUserInterfaceActionType::ToggleButton
+		);
+	MenuBuilder.AddMenuEntry(
+		LOCTEXT("Foliage", "Foliage"),
+		FText(),
+		FSlateIcon(FVREditorStyle::GetStyleSetName(), "VREditorStyle.FoliageMode"),
+		FUIAction
+		(
+			FExecuteAction::CreateStatic(&FVREditorActionCallbacks::ChangeEditorModes, FBuiltinEditorModes::EM_Foliage),
+			FCanExecuteAction::CreateStatic(&FLevelEditorActionCallbacks::DefaultCanExecuteAction),
+			FGetActionCheckState::CreateStatic(&FVREditorActionCallbacks::EditorModeActive, FBuiltinEditorModes::EM_Foliage)
+			),
+		NAME_None,
+		EUserInterfaceActionType::ToggleButton
+		);
+	MenuBuilder.AddMenuEntry(
+		LOCTEXT("Landscape", "Landscape"),
+		FText(),
+		FSlateIcon(FVREditorStyle::GetStyleSetName(), "VREditorStyle.LandscapeMode"),
+		FUIAction
+		(
+			FExecuteAction::CreateStatic(&FVREditorActionCallbacks::ChangeEditorModes, FBuiltinEditorModes::EM_Landscape),
+			FCanExecuteAction::CreateStatic(&FLevelEditorActionCallbacks::DefaultCanExecuteAction),
+			FGetActionCheckState::CreateStatic(&FVREditorActionCallbacks::EditorModeActive, FBuiltinEditorModes::EM_Landscape)
+			),
+		NAME_None,
+		EUserInterfaceActionType::ToggleButton
+		);
+	MenuBuilder.AddMenuEntry(
+		LOCTEXT("MeshPaint", "Paint"),
+		FText(),
+		FSlateIcon(FVREditorStyle::GetStyleSetName(), "VREditorStyle.MeshPaintMode"),
+		FUIAction
+		(
+			FExecuteAction::CreateStatic(&FVREditorActionCallbacks::ChangeEditorModes, FBuiltinEditorModes::EM_MeshPaint),
+			FCanExecuteAction::CreateStatic(&FLevelEditorActionCallbacks::DefaultCanExecuteAction),
+			FGetActionCheckState::CreateStatic(&FVREditorActionCallbacks::EditorModeActive, FBuiltinEditorModes::EM_MeshPaint)
+			),
+		NAME_None,
+		EUserInterfaceActionType::ToggleButton
+		);
+}
+
+void UVRRadialMenuHandler::RegisterMenuGenerator(const FOnRadialMenuGenerated NewMenuGenerator, const bool bShouldAddToStack)
+{
+	if (bShouldAddToStack)
+	{
+		FOnRadialMenuGenerated LastMenuGenerated = GetCurrentMenuGenerator();
+		if (LastMenuGenerated.IsBound())
+		{
+			MenuStack.Push(LastMenuGenerated);
+		}
+	}
+
+	OnRadialMenuGenerated = NewMenuGenerator;
+	UVREditorInteractor* RadialMenuInteractor = UIOwner->GetOwner().GetHandInteractor(EControllerHand::Right);
+	if (UIOwner->IsShowingRadialMenu(RadialMenuInteractor) && NewMenuGenerator.IsBound())
+	{
+		const bool bForceRefresh = true;
+		UIOwner->TryToSpawnRadialMenu(RadialMenuInteractor, bForceRefresh);
+	}
+	else
+	{
+		RadialMenuInteractor = UIOwner->GetOwner().GetHandInteractor(EControllerHand::Left);
+		if (UIOwner->IsShowingRadialMenu(RadialMenuInteractor) && NewMenuGenerator.IsBound())
+		{
+			const bool bForceRefresh = true;
+			UIOwner->TryToSpawnRadialMenu(RadialMenuInteractor, bForceRefresh);
+		}
+	}
+}
+
+void UVRRadialMenuHandler::SetActionsMenuGenerator(const FOnRadialMenuGenerated NewMenuGenerator, const FText NewLabel)
+{
+	UVRRadialMenuHandler::ActionMenuLabel = NewLabel;
+	ActionsMenu = NewMenuGenerator;
+}
+
+
+void UVRRadialMenuHandler::ResetActionsMenuGenerator()
+{
+	UVRRadialMenuHandler::ActionMenuLabel = LOCTEXT("DefaultActions", "Actions");
+	FOnRadialMenuGenerated EmptyRadialMenu;
+	ActionsMenu = EmptyRadialMenu;
+}
+
+FText UVRRadialMenuHandler::GetActionMenuLabel()
+{
+	return  UVRRadialMenuHandler::ActionMenuLabel;
+}
+
+bool UVRRadialMenuHandler::IsActionMenuBound()
+{
+	return ActionsMenu.IsBound();
+}
+
+#undef LOCTEXT_NAMESPACE
\ No newline at end of file
diff --git a/Engine/Source/Editor/VREditor/UI/VRRadialMenuHandler.h b/Engine/Source/Editor/VREditor/UI/VRRadialMenuHandler.h
new file mode 100644
--- /dev/null
+++ b/Engine/Source/Editor/VREditor/UI/VRRadialMenuHandler.h
@@ -0,0 +1,109 @@
+// Copyright 1998-2017 Epic Games, Inc. All Rights Reserved.
+#pragma once
+
+#include "CoreMinimal.h"
+#include "SlateFwd.h"
+#include "UObject/ObjectMacros.h"
+#include "UObject/Object.h"
+#include "Framework/MultiBox/MultiBoxBuilder.h"
+#include "VRRadialMenuHandler.generated.h"
+
+class FUICommandList;
+class UVREditorMode;
+enum class EEditableMeshElementType;
+
+DECLARE_DELEGATE_FourParams(FOnRadialMenuGenerated, FMenuBuilder, TSharedPtr<FUICommandList>, UVREditorMode*, float&);
+
+/**
+* VR Editor user interface manager
+*/
+UCLASS()
+class UVRRadialMenuHandler : public UObject
+{
+	GENERATED_UCLASS_BODY()
+
+public:
+	/** Builds the current radial menu */
+	void BuildRadialMenuCommands(FMenuBuilder MenuBuilder, TSharedPtr<FUICommandList> CommandList, UVREditorMode* VRMode, float& RadiusOverride);
+
+	/** Registers a new menu generator and replaces the currently displayed menu if the radial menu is open */
+	void RegisterMenuGenerator(const FOnRadialMenuGenerated NewMenuGenerator, const bool bShouldAddToStack = true);
+
+	/** Sets a delegate for the context-specific actions menu */
+	void SetActionsMenuGenerator(const FOnRadialMenuGenerated NewMenuGenerator, const FText NewLabel);
+
+	/** Resets the delegate and button for the context-specific actions menu */
+	void ResetActionsMenuGenerator();
+
+	FOnRadialMenuGenerated GetCurrentMenuGenerator()
+	{
+		return OnRadialMenuGenerated;
+	};
+
+	FOnRadialMenuGenerated GetHomeMenuGenerator()
+	{
+		return HomeMenu;
+	};
+
+	FOnRadialMenuGenerated GetActionsMenuGenerator()
+	{
+		return ActionsMenu;
+	};
+
+	/** Returns to the previous radial menu */
+	void BackOutMenu();
+
+	/** Returns to the home menu */
+	void Home();
+
+	/** Allows disabling buttons in the action menu if it's not currently bound */
+	bool IsActionMenuBound();
+
+	/** Allows other systems to read and save the title of existing action menus */
+	static FText GetActionMenuLabel();
+
+	void Init(class UVREditorUISystem* InUISystem);
+
+protected:
+	/** Functions to bind to each menu delegate */
+	void HomeMenuGenerator(FMenuBuilder MenuBuilder, TSharedPtr<FUICommandList> CommandList, UVREditorMode* VRMode, float& RadiusOverride);
+
+	void SnapMenuGenerator(FMenuBuilder MenuBuilder, TSharedPtr<FUICommandList> CommandList, UVREditorMode* VRMode, float& RadiusOverride);
+
+	void GizmoMenuGenerator(FMenuBuilder MenuBuilder, TSharedPtr<FUICommandList> CommandList, UVREditorMode* VRMode, float& RadiusOverride);
+
+	void UIMenuGenerator(FMenuBuilder MenuBuilder, TSharedPtr<FUICommandList> CommandList, UVREditorMode* VRMode, float& RadiusOverride);
+
+	void EditMenuGenerator(FMenuBuilder MenuBuilder, TSharedPtr<FUICommandList> CommandList, UVREditorMode* VRMode, float& RadiusOverride);
+
+	void ToolsMenuGenerator(FMenuBuilder MenuBuilder, TSharedPtr<FUICommandList> CommandList, UVREditorMode* VRMode, float& RadiusOverride);
+
+	void ModesMenuGenerator(FMenuBuilder MenuBuilder, TSharedPtr<FUICommandList> CommandList, UVREditorMode* VRMode, float& RadiusOverride);
+
+protected:
+
+	FOnRadialMenuGenerated OnRadialMenuGenerated;
+
+	FOnRadialMenuGenerated HomeMenu;
+
+	FOnRadialMenuGenerated SnapMenu;
+
+	FOnRadialMenuGenerated GizmoMenu;
+
+	FOnRadialMenuGenerated UIMenu;
+
+	FOnRadialMenuGenerated EditMenu;
+
+	FOnRadialMenuGenerated ToolsMenu;
+
+	FOnRadialMenuGenerated ModesMenu;
+
+	FOnRadialMenuGenerated ActionsMenu;
+
+	static FText ActionMenuLabel;
+
+	class UVREditorUISystem* UIOwner;
+
+	TArray<FOnRadialMenuGenerated> MenuStack;
+
+};
diff --git a/Engine/Source/Editor/VREditor/VREditor.Build.cs b/Engine/Source/Editor/VREditor/VREditor.Build.cs
--- a/Engine/Source/Editor/VREditor/VREditor.Build.cs
+++ b/Engine/Source/Editor/VREditor/VREditor.Build.cs
@@ -6,8 +6,7 @@ namespace UnrealBuildTool.Rules
 	{
         public VREditor(TargetInfo Target)
 		{
-
-			PrivateDependencyModuleNames.AddRange(
+            PrivateDependencyModuleNames.AddRange(
                 new string[] {
                     "AppFramework",
 				    "Core",
@@ -29,19 +28,20 @@ namespace UnrealBuildTool.Rules
 
 			PublicDependencyModuleNames.AddRange(
 				new string[] {
-					"ViewportInteraction",
+					"ViewportInteraction"
 				}
 			);
 
 			PrivateIncludePathModuleNames.AddRange(
 				new string[] {
-					"PlacementMode",
+					"PlacementMode"
 				}
 			);
 
 			DynamicallyLoadedModuleNames.AddRange(
 				new string[] {
-				}
+                    "PlacementMode"
+                }
 			);
 
             PrivateIncludePaths.AddRange(
diff --git a/Engine/Source/Editor/VREditor/VREditorActions.cpp b/Engine/Source/Editor/VREditor/VREditorActions.cpp
--- a/Engine/Source/Editor/VREditor/VREditorActions.cpp
+++ b/Engine/Source/Editor/VREditor/VREditorActions.cpp
@@ -12,18 +12,34 @@
 #include "ImageUtils.h"
 #include "FileHelper.h"
 #include "Framework/Application/SlateApplication.h"
-
+#include "ILevelEditor.h"
+#include "LevelEditor.h"
 #include "AssetEditorManager.h"
 #include "Developer/AssetTools/Public/IAssetTools.h"
 #include "Developer/AssetTools/Public/AssetToolsModule.h"
 #include "ModuleManager.h"
 #include "LevelSequence.h"
 #include "LevelSequenceActor.h"
+#include "ISequencer.h"
+#include "SequencerSettings.h"
+#include "UnrealEdGlobals.h"
+#include "Editor/UnrealEdEngine.h"
+#include "EditorModes.h"
+#include "EdMode.h"
+#include "Engine/Selection.h"
+#include "VREditorMotionControllerInteractor.h"
+#include "LevelEditorActions.h"
 
 #define LOCTEXT_NAMESPACE "VREditorActions"
 
+namespace VREd
+{
+	static FAutoConsoleVariable AllowPlay(TEXT("VREd.AllowPlay"), 1, TEXT("Allow to start play."));
+}
+
 FText FVREditorActionCallbacks::GizmoCoordinateSystemText;
 FText FVREditorActionCallbacks::GizmoModeText;
+FText FVREditorActionCallbacks::SelectingCandidateActorsText;
 
 ECheckBoxState FVREditorActionCallbacks::GetTranslationSnapState()
 {
@@ -107,7 +123,7 @@ FText FVREditorActionCallbacks::GetScaleSnapSizeText()
 	return FText::AsNumber(GEditor->GetScaleGridSize());
 }
 
-void FVREditorActionCallbacks::OnGizmoCoordinateSystemButtonClicked(class UVREditorMode* InVRMode)
+void FVREditorActionCallbacks::OnGizmoCoordinateSystemButtonClicked(UVREditorMode* InVRMode)
 {
 	InVRMode->GetWorldInteraction().CycleTransformGizmoCoordinateSpace();
 	UpdateGizmoModeText(InVRMode);
@@ -115,18 +131,47 @@ void FVREditorActionCallbacks::OnGizmoCoordinateSystemButtonClicked(class UVREdi
 
 }
 
+void FVREditorActionCallbacks::SetCoordinateSystem(UVREditorMode* InVRMode, ECoordSystem InCoordSystem)
+{
+	if (!(InVRMode->GetWorldInteraction().GetTransformGizmoCoordinateSpace() == InCoordSystem))
+	{
+		InVRMode->GetWorldInteraction().SetTransformGizmoCoordinateSpace(InCoordSystem);
+		UpdateGizmoCoordinateSystemText(InVRMode);
+	}
+}
+
+ECheckBoxState FVREditorActionCallbacks::IsActiveCoordinateSystem(UVREditorMode* InVRMode, ECoordSystem InCoordSystem)
+{
+	return (InVRMode->GetWorldInteraction().GetTransformGizmoCoordinateSpace() == InCoordSystem) ? ECheckBoxState::Checked : ECheckBoxState::Unchecked;
+}
+
+
+void FVREditorActionCallbacks::SetGizmoMode(UVREditorMode* InVRMode, EGizmoHandleTypes InGizmoMode)
+{
+	if (!(InVRMode->GetWorldInteraction().GetCurrentGizmoType() == InGizmoMode))
+	{
+		InVRMode->GetWorldInteraction().SetGizmoHandleType(InGizmoMode);
+		UpdateGizmoCoordinateSystemText(InVRMode);
+	}
+}
+
+ECheckBoxState FVREditorActionCallbacks::IsActiveGizmoMode(UVREditorMode* InVRMode, EGizmoHandleTypes InGizmoMode)
+{
+	return (InVRMode->GetWorldInteraction().GetCurrentGizmoType() == InGizmoMode) ? ECheckBoxState::Checked : ECheckBoxState::Unchecked;
+}
+
 FText FVREditorActionCallbacks::GetGizmoCoordinateSystemText()
 {
 	return FVREditorActionCallbacks::GizmoCoordinateSystemText;
 }
 
-void FVREditorActionCallbacks::UpdateGizmoCoordinateSystemText(class UVREditorMode* InVRMode) 
+void FVREditorActionCallbacks::UpdateGizmoCoordinateSystemText(UVREditorMode* InVRMode) 
 {
 	const ECoordSystem CurrentCoordSystem = InVRMode->GetWorldInteraction().GetTransformGizmoCoordinateSpace(); //@todo VREditor
-	FVREditorActionCallbacks::GizmoCoordinateSystemText = (CurrentCoordSystem == COORD_World ? LOCTEXT("WorldCoordinateSystem", "World Space") : LOCTEXT("LocalCoordinateSystem", "Local Space"));
+	FVREditorActionCallbacks::GizmoCoordinateSystemText = (CurrentCoordSystem == COORD_World ? LOCTEXT("WorldCoordinateSystem", "World") : LOCTEXT("LocalCoordinateSystem", "Local"));
 }
 
-void FVREditorActionCallbacks::OnGizmoModeButtonClicked(class UVREditorMode* InVRMode)
+void FVREditorActionCallbacks::OnGizmoModeButtonClicked(UVREditorMode* InVRMode)
 {
 	const EGizmoHandleTypes CurrentGizmoMode = InVRMode->GetWorldInteraction().GetCurrentGizmoType();
 	InVRMode->CycleTransformGizmoHandleType();
@@ -140,7 +185,7 @@ FText FVREditorActionCallbacks::GetGizmoModeText()
 	return FVREditorActionCallbacks::GizmoModeText;
 }
 
-void FVREditorActionCallbacks::UpdateGizmoModeText(class UVREditorMode* InVRMode) 
+void FVREditorActionCallbacks::UpdateGizmoModeText(UVREditorMode* InVRMode) 
 {
 	const EGizmoHandleTypes CurrentGizmoType = InVRMode->GetWorldInteraction().GetCurrentGizmoType();
 	FText GizmoTypeText;
@@ -155,7 +200,7 @@ void FVREditorActionCallbacks::UpdateGizmoModeText(class UVREditorMode* InVRMode
 		break;
 
 	case EGizmoHandleTypes::Rotate:
-		GizmoTypeText = LOCTEXT("RotationGizmoType", "Rotation");
+		GizmoTypeText = LOCTEXT("RotationGizmoType", "Rotate");
 		break;
 
 	case EGizmoHandleTypes::Scale:
@@ -170,25 +215,18 @@ void FVREditorActionCallbacks::UpdateGizmoModeText(class UVREditorMode* InVRMode
 	FVREditorActionCallbacks::GizmoModeText = GizmoTypeText;
 }
 
-void FVREditorActionCallbacks::OnGizmoCycle(class UVREditorMode* InVRMode)
-{
-	InVRMode->CycleTransformGizmoHandleType();
-	UpdateGizmoModeText(InVRMode);
-	UpdateGizmoCoordinateSystemText(InVRMode);
-}
-
-void FVREditorActionCallbacks::OnUIToggleButtonClicked(class UVREditorMode* InVRMode, const UVREditorUISystem::EEditorUIPanel PanelToToggle)
+void FVREditorActionCallbacks::OnUIToggleButtonClicked(UVREditorMode* InVRMode, const UVREditorUISystem::EEditorUIPanel PanelToToggle)
 {
 	InVRMode->GetUISystem().TogglePanelVisibility(PanelToToggle);
 }
 
-ECheckBoxState FVREditorActionCallbacks::GetUIToggledState(class UVREditorMode* InVRMode, const UVREditorUISystem::EEditorUIPanel PanelToCheck)
+ECheckBoxState FVREditorActionCallbacks::GetUIToggledState(UVREditorMode* InVRMode, const UVREditorUISystem::EEditorUIPanel PanelToCheck)
 {
 	return InVRMode->GetUISystem().IsShowingEditorUIPanel(PanelToCheck) ? ECheckBoxState::Checked : ECheckBoxState::Unchecked;
 }
 
 
-void FVREditorActionCallbacks::OnLightButtonClicked(class UVREditorMode* InVRMode)
+void FVREditorActionCallbacks::OnLightButtonClicked(UVREditorMode* InVRMode)
 {
 	// Always spawn the flashlight on the hand clicking on the UI
 	UVREditorInteractor* VREditorInteractor = InVRMode->GetHandInteractor(EControllerHand::Left);
@@ -201,13 +239,13 @@ void FVREditorActionCallbacks::OnLightButtonClicked(class UVREditorMode* InVRMod
 
 }
 
-ECheckBoxState FVREditorActionCallbacks::GetFlashlightState(class UVREditorMode* InVRMode)
+ECheckBoxState FVREditorActionCallbacks::GetFlashlightState(UVREditorMode* InVRMode)
 {
 	return InVRMode->IsFlashlightOn() ? ECheckBoxState::Checked : ECheckBoxState::Unchecked;
 }
 
 
-void FVREditorActionCallbacks::OnScreenshotButtonClicked(class UVREditorMode* InVRMode)
+void FVREditorActionCallbacks::OnScreenshotButtonClicked(UVREditorMode* InVRMode)
 {
 	// @todo vreditor: update after direct buffer grab changes
 
@@ -251,21 +289,33 @@ void FVREditorActionCallbacks::OnScreenshotButtonClicked(class UVREditorMode* In
 		FFileHelper::SaveArrayToFile(CompressedBitmap, *FinalFileName);
 
 	}
+}
 
-
+void FVREditorActionCallbacks::OnPlayButtonClicked(UVREditorMode* InVRMode)
+{
+	InVRMode->TogglePIEAndVREditor();
 }
 
-void FVREditorActionCallbacks::OnPlayButtonClicked(class UVREditorMode* InVRMode)
+bool FVREditorActionCallbacks::CanPlay()
 {
-	InVRMode->StartExitingVRMode(EVREditorExitType::PIE_VR);
+	return FLevelEditorActionCallbacks::DefaultCanExecuteAction() && VREd::AllowPlay->GetInt() == 1;
 }
 
-void FVREditorActionCallbacks::OnSnapActorsToGroundClicked(class UVREditorMode* InVRMode)
+void FVREditorActionCallbacks::OnSimulateButtonClicked(UVREditorMode* InVRMode)
 {
-	InVRMode->SnapSelectedActorsToGround();
+	InVRMode->ToggleSIEAndVREditor();
 }
 
+FText FVREditorActionCallbacks::GetSimulateText()
+{
+	return GEditor->bIsSimulatingInEditor ? LOCTEXT( "SimulateStopButton", "Stop" ) : LOCTEXT( "SimulateStartButton", "Simulate" );
+}
 
+void FVREditorActionCallbacks::OnSnapActorsToGroundClicked(UVREditorMode* InVRMode)
+{
+	InVRMode->SnapSelectedActorsToGround();
+}
+ 
 void FVREditorActionCallbacks::SimulateCharacterEntry(const FString InChar)
 {
 
@@ -279,6 +329,14 @@ void FVREditorActionCallbacks::SimulateCharacterEntry(const FString InChar)
 
 }
 
+void FVREditorActionCallbacks::SimulateBackspace()
+{
+	// Slate editable text fields handle backspace as a character entry
+	FString BackspaceString = FString(TEXT("\b"));
+	bool bRepeat = false;
+	SimulateCharacterEntry(BackspaceString);
+}
+
 void FVREditorActionCallbacks::SimulateKeyDown(const FKey Key, const bool bRepeat)
 {
 	const uint32* KeyCodePtr;
@@ -311,7 +369,7 @@ void FVREditorActionCallbacks::SimulateKeyUp(const FKey Key)
 	FSlateApplication::Get().ProcessKeyUpEvent( KeyEvent );
 }
 
-void FVREditorActionCallbacks::CreateNewSequence(class UVREditorMode* InVRMode)
+void FVREditorActionCallbacks::CreateNewSequence(UVREditorMode* InVRMode)
 {
 	// Create a new level sequence
 	IAssetTools& AssetTools = FModuleManager::GetModuleChecked<FAssetToolsModule>("AssetTools").Get();
@@ -351,11 +409,13 @@ void FVREditorActionCallbacks::CreateNewSequence(class UVREditorMode* InVRMode)
 
 	ALevelSequenceActor* NewActor = CastChecked<ALevelSequenceActor>(GEditor->UseActorFactory(ActorFactory, FAssetData(NewAsset), &FTransform::Identity));
 
+	InVRMode->GetUISystem().SequencerOpenendFromRadialMenu(true);
+
 	// Open the Sequencer window
 	FAssetEditorManager::Get().OpenEditorForAsset(NewAsset);
 }
 
-void FVREditorActionCallbacks::CloseSequencer(class UMovieSceneSequence* OpenSequence)
+void FVREditorActionCallbacks::CloseSequencer(UMovieSceneSequence* OpenSequence)
 {
 	IAssetEditorInstance* SequencerEditor = FAssetEditorManager::Get().FindEditorForAsset(OpenSequence, false);
 	if (SequencerEditor != nullptr)
@@ -364,4 +424,154 @@ void FVREditorActionCallbacks::CloseSequencer(class UMovieSceneSequence* OpenSeq
 	}
 }
 
+void FVREditorActionCallbacks::PlaySequenceAtRate(UVREditorMode* InVRMode, float Rate)
+{
+	ISequencer* CurrentSequencer = InVRMode->GetCurrentSequencer();
+	if (CurrentSequencer != nullptr)
+	{
+		CurrentSequencer->OnPlay(false, Rate);
+	}
+}
+
+void FVREditorActionCallbacks::PauseSequencePlayback(UVREditorMode* InVRMode)
+{
+	ISequencer* CurrentSequencer = InVRMode->GetCurrentSequencer();
+	if (CurrentSequencer != nullptr)
+	{
+		CurrentSequencer->Pause();
+	}
+}
+
+void FVREditorActionCallbacks::PlayFromBeginning(UVREditorMode* InVRMode)
+{
+	ISequencer* CurrentSequencer = InVRMode->GetCurrentSequencer();
+	if (CurrentSequencer != nullptr)
+	{
+		CurrentSequencer->SetLocalTime(0.0f);
+		const float Rate = 1.0f;
+		CurrentSequencer->OnPlay(false, Rate);
+	}
+}
+
+void FVREditorActionCallbacks::ToggleLooping(UVREditorMode* InVRMode)
+{
+	ISequencer* CurrentSequencer = InVRMode->GetCurrentSequencer();
+	if (CurrentSequencer != nullptr)
+	{
+		CurrentSequencer->GetSequencerSettings()->SetLooping(!CurrentSequencer->GetSequencerSettings()->IsLooping());
+	}
+}
+
+ECheckBoxState FVREditorActionCallbacks::IsLoopingChecked(UVREditorMode* InVRMode)
+{
+	bool bShouldReturnChecked = false;
+	ISequencer* CurrentSequencer = InVRMode->GetCurrentSequencer();
+	if (CurrentSequencer != nullptr)
+	{
+		bShouldReturnChecked = InVRMode->GetCurrentSequencer()->GetSequencerSettings()->IsLooping();
+	}
+	return bShouldReturnChecked ? ECheckBoxState::Checked : ECheckBoxState::Unchecked;
+
+}
+
+void FVREditorActionCallbacks::ToggleSequencerScrubbing(UVREditorMode* InVRMode, UVREditorMotionControllerInteractor* InController)
+{
+	InController->ToggleSequencerScrubbingMode();
+	if (!InController->IsScrubbingSequencer())
+	{
+		PauseSequencePlayback(InVRMode);
+	}
+}
+
+ECheckBoxState FVREditorActionCallbacks::GetSequencerScrubState(UVREditorMotionControllerInteractor* InController)
+{
+	return InController->IsScrubbingSequencer() ? ECheckBoxState::Checked : ECheckBoxState::Unchecked;
+}
+
+void FVREditorActionCallbacks::ToggleAligningToActors(UVREditorMode* InVRMode)
+{
+	if (InVRMode->GetWorldInteraction().AreAligningToActors())
+	{
+		// Deselect any candidates if we have them
+		if (InVRMode->GetWorldInteraction().HasCandidatesSelected())
+		{
+			ToggleSelectingCandidateActors(InVRMode);
+		}
+		GUnrealEd->Exec(InVRMode->GetWorld(), TEXT("VI.EnableGuides 0"));
+	}
+	else
+	{
+		GUnrealEd->Exec(InVRMode->GetWorld(), TEXT("VI.EnableGuides 1"));
+	}
+}
+
+ECheckBoxState FVREditorActionCallbacks::AreAligningToActors(UVREditorMode* InVRMode)
+{
+	return InVRMode->GetWorldInteraction().AreAligningToActors() ? ECheckBoxState::Checked : ECheckBoxState::Unchecked;
+}
+
+void FVREditorActionCallbacks::ToggleSelectingCandidateActors(UVREditorMode* InVRMode)
+{
+	InVRMode->GetWorldInteraction().SetSelectionAsCandidates();
+	UpdateSelectingCandidateActorsText(InVRMode);
+}
+
+bool FVREditorActionCallbacks::CanSelectCandidateActors(UVREditorMode* InVRMode)
+{
+	return InVRMode->GetWorldInteraction().AreAligningToActors();
+}
+
+FText FVREditorActionCallbacks::GetSelectingCandidateActorsText()
+{
+	return FVREditorActionCallbacks::SelectingCandidateActorsText;
+}
+
+void FVREditorActionCallbacks::UpdateSelectingCandidateActorsText(UVREditorMode* InVRMode)
+{
+	if (InVRMode->GetWorldInteraction().HasCandidatesSelected())
+	{
+		FVREditorActionCallbacks::SelectingCandidateActorsText = LOCTEXT("ResetCandidates", "Reset Targets");
+	}
+	else
+	{
+		FVREditorActionCallbacks::SelectingCandidateActorsText = LOCTEXT("SetCandidates", "Set Targets");
+	}
+}
+
+void FVREditorActionCallbacks::ChangeEditorModes(FEditorModeID InMode)
+{
+	// *Important* - activate the mode first since FEditorModeTools::DeactivateMode will
+	// activate the default mode when the stack becomes empty, resulting in multiple active visible modes.
+	GLevelEditorModeTools().ActivateMode(InMode);
+
+	// Find and disable any other 'visible' modes since we only ever allow one of those active at a time.
+	TArray<FEdMode*> ActiveModes;
+	GLevelEditorModeTools().GetActiveModes(ActiveModes);
+	for (FEdMode* Mode : ActiveModes)
+	{
+		if (Mode->GetID() != InMode && Mode->GetModeInfo().bVisible)
+		{
+			GLevelEditorModeTools().DeactivateMode(Mode->GetID());
+		}
+	}
+}
+
+ECheckBoxState FVREditorActionCallbacks::EditorModeActive(FEditorModeID InMode)
+{
+	return GLevelEditorModeTools().IsModeActive(InMode) ? ECheckBoxState::Checked : ECheckBoxState::Unchecked;
+}
+
+bool FVREditorActionCallbacks::ModeActionsMenuActive(UVREditorMode* VRMode)
+{
+	bool bShouldModeActionsBeActive = (VRMode->GetCurrentSequencer() != nullptr);
+	return bShouldModeActionsBeActive;
+}
+
+void FVREditorActionCallbacks::DeselectAll()
+{
+	GEditor->SelectNone(true, true, false);
+	GEditor->GetSelectedActors()->DeselectAll();
+	GEditor->GetSelectedObjects()->DeselectAll();
+}
+
 #undef LOCTEXT_NAMESPACE
diff --git a/Engine/Source/Editor/VREditor/VREditorActions.h b/Engine/Source/Editor/VREditor/VREditorActions.h
--- a/Engine/Source/Editor/VREditor/VREditorActions.h
+++ b/Engine/Source/Editor/VREditor/VREditorActions.h
@@ -4,6 +4,12 @@
 #include "VREditorUISystem.h"
 #include "SlateTypes.h"
 #include "InputCoreTypes.h"
+#include "Editor.h"
+
+// Forward declarations 
+class UVREditorMode;
+enum ECoordSystem;
+enum class EGizmoHandleTypes : uint8;
 
 /**
 * Implementation of various VR editor action callback functions
@@ -43,7 +49,11 @@ public:
 	 * Toggles the gizmo coordinate system between local and world space
 	 * @param InVRMode Currently active VRMode
 	 */
-	static void OnGizmoCoordinateSystemButtonClicked(class UVREditorMode* InVRMode);
+	static void OnGizmoCoordinateSystemButtonClicked(UVREditorMode* InVRMode);
+
+	static void SetCoordinateSystem(UVREditorMode* InVRMode, ECoordSystem InCoordSystem);
+
+	static ECheckBoxState IsActiveCoordinateSystem(UVREditorMode* InVRMode, ECoordSystem InCoordSystem);
 
 	/** Returns the gizmo coordinate system as text to use as the button display text */
 	static FText GetGizmoCoordinateSystemText();
@@ -52,13 +62,17 @@ public:
 	 * Updates the gizmo coordinate system text if the coordinate system or gizmo type is changed
 	 * @param InVRMode Currently active VRMode
 	 */
-	static void UpdateGizmoCoordinateSystemText(class UVREditorMode* InVRMode);
+	static void UpdateGizmoCoordinateSystemText(UVREditorMode* InVRMode);
 
 	/**
 	 * Rotates the gizmo type through universal, translate, rotate, and scale
 	 * @param InVRMode Currently active VRMode
 	 */
-	static void OnGizmoModeButtonClicked(class UVREditorMode* InVRMode);
+	static void OnGizmoModeButtonClicked(UVREditorMode* InVRMode);
+
+	static void SetGizmoMode(UVREditorMode* InVRMode, EGizmoHandleTypes InGizmoMode);
+
+	static ECheckBoxState IsActiveGizmoMode(UVREditorMode* InVRMode, EGizmoHandleTypes InGizmoMode);
 
 	/** Returns the gizmo type as text to use as the button display text */
 	static FText GetGizmoModeText();
@@ -67,57 +81,65 @@ public:
 	 * Updates the gizmo coordinate type text if the coordinate system or gizmo type is changed
 	 * @param InVRMode Currently active VRMode
 	 */
-	static void UpdateGizmoModeText(class UVREditorMode* InVRMode);
-
-	/**
-	* Cycles the gizmo through all coordinate spaces and Actor manipulation types
-	* @param InVRMode Currently active VRMode
-	*/
-	static void OnGizmoCycle(class UVREditorMode* InVRMode);
+	static void UpdateGizmoModeText(UVREditorMode* InVRMode);
 
 	/**
 	 * Toggles a VR UI panel's state between visible and invisible
 	 * @param InVRMode Currently active VRMode
 	 * @param PanelToToggle UI panel to change the state of
 	 */
-	static void OnUIToggleButtonClicked(class UVREditorMode* InVRMode, const UVREditorUISystem::EEditorUIPanel PanelToToggle);
+	static void OnUIToggleButtonClicked(UVREditorMode* InVRMode, const UVREditorUISystem::EEditorUIPanel PanelToToggle);
 
 	/**
 	 * Returns a VR UI panel's visibility - used for check boxes on the menu button
 	 * @param InVRMode Currently active VRMode
 	 * @param PanelToToggle UI panel to read the state of
 	 */
-	static ECheckBoxState GetUIToggledState(class UVREditorMode* InVRMode, const UVREditorUISystem::EEditorUIPanel PanelToCheck);
+	static ECheckBoxState GetUIToggledState(UVREditorMode* InVRMode, const UVREditorUISystem::EEditorUIPanel PanelToCheck);
 
 	/**
 	 * Toggles a flashlight on and off on the interactor that clicked on the UI button
 	 * @param InVRMode Currently active VRMode
 	 */
-	static void OnLightButtonClicked(class UVREditorMode* InVRMode);
+	static void OnLightButtonClicked(UVREditorMode* InVRMode);
 
 	/**
 	 * Returns whether or not the flashlight is enabled - used for check box on the flashlight button
 	 * @param InVRMode Currently active VRMode
 	 */
-	static ECheckBoxState GetFlashlightState(class UVREditorMode* InVRMode);
+	static ECheckBoxState GetFlashlightState(UVREditorMode* InVRMode);
 
 	/**
 	 * Takes a screenshot of the mirror viewport
 	 * @param InVRMode Currently active VRMode
 	 */
-	static void OnScreenshotButtonClicked(class UVREditorMode* InVRMode);
+	static void OnScreenshotButtonClicked(UVREditorMode* InVRMode);
 
 	/**
 	 * Enters Play In Editor mode for testing of gameplay
 	 * @param InVRMode Currently active VRMode
 	 */
-	static void OnPlayButtonClicked(class UVREditorMode* InVRMode);
+	static void OnPlayButtonClicked(UVREditorMode* InVRMode);
+
+	/**
+	 * If we are allowed to enter play.
+	 */
+	static bool CanPlay();
+
+	/**
+	 * Enters Simulate mode for physics and animation playback
+	 * @param InVRMode Currently active VRMode
+	 */
+	static void OnSimulateButtonClicked(UVREditorMode* InVRMode);
+
+	/** Returns the simulate button text to display */
+	static FText GetSimulateText();
 
 	/**
 	 * Snaps currently selected Actors to the ground
 	 * @param InVRMode Currently active VRMode
 	 */
-	static void OnSnapActorsToGroundClicked(class UVREditorMode* InVRMode);
+	static void OnSnapActorsToGroundClicked(UVREditorMode* InVRMode);
 
 	/**
 	 * Simulates the user entering characters with a keyboard for data entry
@@ -125,6 +147,9 @@ public:
 	 */
 	static void SimulateCharacterEntry(const FString InChar);
 
+	/** Send a backspace event. Slate editable text fields handle backspace as a TCHAR(8) character entry */
+	static void SimulateBackspace();
+
 	/**
 	 * Simulates the user pressing a key down
 	 * @param Key Key to press
@@ -139,16 +164,68 @@ public:
 	static void SimulateKeyUp( const FKey Key );
 
 	/** Create a new level sequence with an auto-generated name */
-	static void CreateNewSequence(class UVREditorMode* InVRMode);
+	static void CreateNewSequence(UVREditorMode* InVRMode);
 
 	/** Close the sequencer instance */
 	static void CloseSequencer(class UMovieSceneSequence* OpenSequence);
 
+	/** Plays the current sequence at a specified rate */
+	static void PlaySequenceAtRate(UVREditorMode* InVRMode, float Rate);
+
+	/** Pauses sequence playback */
+	static void PauseSequencePlayback(UVREditorMode* InVRMode);
+
+	/** Plays at a rate of 1.0 from the local start of the sequence */
+	static void PlayFromBeginning(UVREditorMode* InVRMode);
+
+	/** Toggles looping the sequence */
+	static void ToggleLooping(UVREditorMode* InVRMode);
+
+	/** Whether or not the current sequence is looping */
+	static ECheckBoxState IsLoopingChecked(UVREditorMode* InVRMode);
+
+	static void ToggleSequencerScrubbing(UVREditorMode* InVRMode, UVREditorMotionControllerInteractor* InController);
+
+	static ECheckBoxState GetSequencerScrubState(UVREditorMotionControllerInteractor* InController);
+
+	/** Toggles whether or not the world interaction method should align transformables to actors in the scene*/
+	static void ToggleAligningToActors(UVREditorMode* InVRMode);
+
+	/** True if the world interaction method is currently aligning transformables to actors in the scene*/
+	static ECheckBoxState AreAligningToActors(UVREditorMode* InVRMode);
+
+	/** Either selects or deselects alignment candidates */
+	static void ToggleSelectingCandidateActors(UVREditorMode* InVRMode);
+
+	/** Can only select candidate actors if actor alignment is active */
+	static bool CanSelectCandidateActors(UVREditorMode* InVRMode);
+
+	/** Dynamic label for the alignment candidate button */
+	static FText GetSelectingCandidateActorsText();
+
+	/** Updates the alignment candidate label based on the current aligning state */
+	static void UpdateSelectingCandidateActorsText(UVREditorMode* InVRMode);
+
+	/** Changes the editor mode to the given ID */
+	static void ChangeEditorModes(FEditorModeID InMode);
+
+	/** Checks whether the editor mode for the given ID is active*/
+	static ECheckBoxState EditorModeActive(FEditorModeID InMode);
+
+	/** True if a specialized menu is available */
+	static bool ModeActionsMenuActive(class UVREditorMode* VRMode);
+
+	/** Deselects everything currently selected */
+	static void DeselectAll();
+
 public:
 	
 	static FText GizmoCoordinateSystemText;
 
 	static FText GizmoModeText;
+
+	static FText SelectingCandidateActorsText;
+
 }
 ;
 
diff --git a/Engine/Source/Editor/VREditor/VREditorAvatarActor.cpp b/Engine/Source/Editor/VREditor/VREditorAvatarActor.cpp
--- a/Engine/Source/Editor/VREditor/VREditorAvatarActor.cpp
+++ b/Engine/Source/Editor/VREditor/VREditorAvatarActor.cpp
@@ -19,6 +19,8 @@
 #include "UObject/Package.h"
 #include "EngineGlobals.h"
 #include "IHeadMountedDisplay.h"
+#include "VREditorAssetContainer.h"
+#include "VRModeSettings.h"
 
 namespace VREd
 {
@@ -36,9 +38,6 @@ namespace VREd
 
 	static FAutoConsoleVariable ScaleProgressBarLength( TEXT( "VREd.ScaleProgressBarLength" ), 50.0f, TEXT( "Length of the progressbar that appears when scaling" ) );
 	static FAutoConsoleVariable ScaleProgressBarRadius( TEXT( "VREd.ScaleProgressBarRadius" ), 1.0f, TEXT( "Radius of the progressbar that appears when scaling" ) );
-
-	static FAutoConsoleVariable ShowMovementGrid( TEXT( "VREd.ShowMovementGrid" ), 1, TEXT( "Showing the ground movement grid" ) );
-	static FAutoConsoleVariable ShowWorldMovementPostProcess( TEXT( "VREd.ShowWorldMovementPostProcess" ), 1, TEXT( "Showing the movement post processing" ) );
 }
 
 AVREditorAvatarActor::AVREditorAvatarActor( const FObjectInitializer& ObjectInitializer ) :
@@ -58,6 +57,10 @@ AVREditorAvatarActor::AVREditorAvatarActor( const FObjectInitializer& ObjectInit
 		SetRootComponent( SceneRootComponent );
 	}
 
+	// Setup the asset container.
+	UVREditorAssetContainer* AssetContainer = LoadObject<UVREditorAssetContainer>(nullptr, *UVREditorMode::AssetContainerPath);
+	check(AssetContainer != nullptr);
+
 	// Give us a head mesh
 	{
 		HeadMeshComponent = CreateDefaultSubobject<UStaticMeshComponent>( TEXT( "HeadMeshComponent" ) );
@@ -65,12 +68,13 @@ AVREditorAvatarActor::AVREditorAvatarActor( const FObjectInitializer& ObjectInit
 		HeadMeshComponent->SetupAttachment( RootComponent );
 
 		// @todo vreditor: This needs to adapt based on the device you're using
-		UStaticMesh* HeadMesh = LoadObject<UStaticMesh>( nullptr, TEXT( "/Engine/VREditor/Devices/Generic/GenericHMD" ) );
+		UStaticMesh* HeadMesh = AssetContainer->GenericHMDMesh;
 		check( HeadMesh != nullptr );
 
 		HeadMeshComponent->SetStaticMesh( HeadMesh );
 		HeadMeshComponent->SetMobility( EComponentMobility::Movable );
 		HeadMeshComponent->SetCollisionEnabled( ECollisionEnabled::NoCollision );
+		HeadMeshComponent->bSelectable = false;
 	}
 
 	// World movement grid mesh
@@ -79,13 +83,14 @@ AVREditorAvatarActor::AVREditorAvatarActor( const FObjectInitializer& ObjectInit
 		AddOwnedComponent( WorldMovementGridMeshComponent );
 		WorldMovementGridMeshComponent->SetupAttachment( RootComponent );
 	
-		UStaticMesh* GridMesh = LoadObject<UStaticMesh>( nullptr, TEXT( "/Engine/VREditor/WorldMovementGrid/PlaneMesh" ) );
+		UStaticMesh* GridMesh = AssetContainer->PlaneMesh;
 		check( GridMesh != nullptr );
 		WorldMovementGridMeshComponent->SetStaticMesh( GridMesh );
 		WorldMovementGridMeshComponent->SetMobility( EComponentMobility::Movable );
 		WorldMovementGridMeshComponent->SetCollisionEnabled( ECollisionEnabled::NoCollision );
+		WorldMovementGridMeshComponent->bSelectable = false;
 
-		UMaterialInterface* GridMaterial = LoadObject<UMaterialInterface>( nullptr, TEXT( "/Engine/VREditor/WorldMovementGrid/GridMaterial" ) );
+		UMaterialInterface* GridMaterial = AssetContainer->GridMaterial;
 		check( GridMaterial != nullptr );
 
 		WorldMovementGridMID = UMaterialInstanceDynamic::Create( GridMaterial, GetTransientPackage( ) );
@@ -98,10 +103,10 @@ AVREditorAvatarActor::AVREditorAvatarActor( const FObjectInitializer& ObjectInit
 
 	{
 		{
-			UMaterialInterface* UserScaleIndicatorMaterial = LoadObject<UMaterialInterface>( nullptr, TEXT( "/Engine/VREditor/LaserPointer/LaserPointerMaterialInst" ) );
+			UMaterialInterface* UserScaleIndicatorMaterial = AssetContainer->LaserPointerMaterial;
 			check( UserScaleIndicatorMaterial != nullptr );
 
-			UMaterialInstance* TranslucentUserScaleIndicatorMaterial = LoadObject<UMaterialInstance>( nullptr, TEXT( "/Engine/VREditor/LaserPointer/TranslucentLaserPointerMaterialInst" ) );
+			UMaterialInterface* TranslucentUserScaleIndicatorMaterial = AssetContainer->LaserPointerTranslucentMaterial;
 			check( TranslucentUserScaleIndicatorMaterial != nullptr );
 
 			FixedUserScaleMID = UMaterialInstanceDynamic::Create( UserScaleIndicatorMaterial, GetTransientPackage() );
@@ -116,7 +121,7 @@ AVREditorAvatarActor::AVREditorAvatarActor( const FObjectInitializer& ObjectInit
 			TranslucentCurrentUserScaleMID = UMaterialInstanceDynamic::Create( TranslucentUserScaleIndicatorMaterial, GetTransientPackage() );
 			check( TranslucentCurrentUserScaleMID != nullptr );
 
-			UStaticMesh* ScaleLineMesh = LoadObject<UStaticMesh>( nullptr, TEXT( "/Engine/VREditor/LaserPointer/LaserPointerMesh" ) );
+			UStaticMesh* ScaleLineMesh = AssetContainer->LaserPointerMesh; //@todo VREditor: The laser pointer mesh is not a closed cylinder anymore.
 			check( ScaleLineMesh != nullptr );
 
 			// Creating the background bar progress of the scale 
@@ -130,6 +135,7 @@ AVREditorAvatarActor::AVREditorAvatarActor( const FObjectInitializer& ObjectInit
 				ScaleProgressMeshComponent->SetCollisionEnabled( ECollisionEnabled::NoCollision );
 				ScaleProgressMeshComponent->SetMaterial( 0, FixedUserScaleMID );
 				ScaleProgressMeshComponent->SetMaterial( 1, TranslucentFixedUserScaleMID );
+				ScaleProgressMeshComponent->bSelectable = false;
 
 				// The user scale indicator starts invisible
 				ScaleProgressMeshComponent->SetVisibility( false );
@@ -146,6 +152,7 @@ AVREditorAvatarActor::AVREditorAvatarActor( const FObjectInitializer& ObjectInit
 				CurrentScaleProgressMeshComponent->SetCollisionEnabled( ECollisionEnabled::NoCollision );
 				CurrentScaleProgressMeshComponent->SetMaterial( 0, CurrentUserScaleMID );
 				CurrentScaleProgressMeshComponent->SetMaterial( 1, TranslucentCurrentUserScaleMID );
+				CurrentScaleProgressMeshComponent->bSelectable = false;
 
 				// The user scale indicator starts invisible
 				CurrentScaleProgressMeshComponent->SetVisibility( false );
@@ -154,10 +161,10 @@ AVREditorAvatarActor::AVREditorAvatarActor( const FObjectInitializer& ObjectInit
 
 		// Creating the text for scaling
 		{
-			UFont* TextFont = LoadObject<UFont>( nullptr, TEXT( "/Engine/VREditor/Fonts/VRText_RobotoLarge" ) );
+			UFont* TextFont = AssetContainer->TextFont;
 			check( TextFont != nullptr );
 
-			UMaterialInterface* UserScaleIndicatorMaterial = LoadObject<UMaterialInterface>( nullptr, TEXT( "/Engine/VREditor/Fonts/VRTextMaterial" ) );
+			UMaterialInterface* UserScaleIndicatorMaterial = AssetContainer->TextMaterial;
 			check( UserScaleIndicatorMaterial != nullptr );
 
 			UserScaleIndicatorText = CreateDefaultSubobject<UTextRenderComponent>( TEXT( "UserScaleIndicatorText" ) );
@@ -167,6 +174,7 @@ AVREditorAvatarActor::AVREditorAvatarActor( const FObjectInitializer& ObjectInit
 			UserScaleIndicatorText->SetMobility( EComponentMobility::Movable );
 			UserScaleIndicatorText->SetCollisionEnabled( ECollisionEnabled::NoCollision );
 			UserScaleIndicatorText->SetCollisionProfileName( UCollisionProfile::NoCollision_ProfileName );
+			UserScaleIndicatorText->bSelectable = false;
 
 			UserScaleIndicatorText->bGenerateOverlapEvents = false;
 			UserScaleIndicatorText->SetCanEverAffectNavigation( false );
@@ -188,7 +196,7 @@ AVREditorAvatarActor::AVREditorAvatarActor( const FObjectInitializer& ObjectInit
 
 	// Load our post process material for the world movement grid
 	{
-		UMaterial* Material = LoadObject<UMaterial>( nullptr, TEXT( "/Engine/VREditor/WorldMovementGrid/GridPostProcess" ) );
+		UMaterial* Material = AssetContainer->WorldMovementPostProcessMaterial;
 		check( Material != nullptr );
 		WorldMovementPostProcessMaterial = UMaterialInstanceDynamic::Create( Material, GetTransientPackage() );
 		check( WorldMovementPostProcessMaterial != nullptr );
@@ -201,7 +209,7 @@ AVREditorAvatarActor::AVREditorAvatarActor( const FObjectInitializer& ObjectInit
 		PostProcessComponent->SetupAttachment( this->GetRootComponent( ) );
 
 		// Unlimited size
-		PostProcessComponent->bEnabled = VREd::ShowWorldMovementPostProcess->GetInt( ) == 1 ? true : false;
+		PostProcessComponent->bEnabled = GetDefault<UVRModeSettings>()->bShowWorldMovementPostProcess;
 		PostProcessComponent->bUnbound = true;
 	}
 
@@ -226,16 +234,16 @@ void AVREditorAvatarActor::Init( UVREditorMode* InVRMode  )
 	// Set the default color for the progress bar
 	{
 		static FName StaticLaserColorName( "LaserColor" );
-		const FLinearColor FixedProgressbarColor = VRMode->GetColor( UVREditorMode::EColors::WorldDraggingColor_TwoHands );
+		const FLinearColor FixedProgressbarColor = VRMode->GetColor( UVREditorMode::EColors::WorldDraggingColor );
 		FixedUserScaleMID->SetVectorParameterValue( StaticLaserColorName, FixedProgressbarColor );
 		TranslucentFixedUserScaleMID->SetVectorParameterValue( StaticLaserColorName, FixedProgressbarColor );
 
-		const FLinearColor CurrentProgressbarColor = VRMode->GetColor( UVREditorMode::EColors::GreenGizmoColor );
+		const FLinearColor CurrentProgressbarColor = VRMode->GetColor( UVREditorMode::EColors::DefaultColor );
 		CurrentUserScaleMID->SetVectorParameterValue( StaticLaserColorName, CurrentProgressbarColor );
 		TranslucentCurrentUserScaleMID->SetVectorParameterValue( StaticLaserColorName, CurrentProgressbarColor );
 	}
 
-	UserScaleIndicatorText->SetTextRenderColor( VRMode->GetColor( UVREditorMode::EColors::WorldDraggingColor_TwoHands ).ToFColor( false ) );
+	UserScaleIndicatorText->SetTextRenderColor( VRMode->GetColor( UVREditorMode::EColors::WorldDraggingColor ).ToFColor( false ) );
 
 	// Tell the grid to stay relative to the rootcomponent
 	const FAttachmentTransformRules AttachmentTransformRules = FAttachmentTransformRules( EAttachmentRule::KeepRelative, true );
@@ -281,7 +289,8 @@ void AVREditorAvatarActor::TickManually( const float DeltaTime )
 		UVREditorInteractor* LeftHandInteractor = VRMode->GetHandInteractor( EControllerHand::Left );
 		UVREditorInteractor* RightHandInteractor = VRMode->GetHandInteractor( EControllerHand::Right );
 
-		if ( LeftHandInteractor != nullptr && RightHandInteractor != nullptr &&
+		if (GetDefault<UVRModeSettings>()->bShowWorldScaleProgressBar &&
+			LeftHandInteractor != nullptr && RightHandInteractor != nullptr &&
 			 ( ( LeftHandInteractor->GetDraggingMode() == EViewportInteractionDraggingMode::World && RightHandInteractor->GetDraggingMode() == EViewportInteractionDraggingMode::AssistingDrag ) ||
 			   ( LeftHandInteractor->GetDraggingMode() == EViewportInteractionDraggingMode::AssistingDrag && RightHandInteractor->GetDraggingMode() == EViewportInteractionDraggingMode::World ) ) )
 		{
@@ -352,7 +361,7 @@ void AVREditorAvatarActor::TickManually( const float DeltaTime )
 
 	// Updating the opacity and visibility of the grid according to the controllers //@todo
 	{
-		if ( VREd::ShowMovementGrid->GetInt( ) == 1 )
+		if (GetDefault<UVRModeSettings>()->bShowWorldMovementGrid)
 		{
 			UVREditorInteractor* LeftHandInteractor = VRMode->GetHandInteractor( EControllerHand::Left );
 			UVREditorInteractor* RightHandInteractor = VRMode->GetHandInteractor( EControllerHand::Right );
@@ -416,7 +425,7 @@ void AVREditorAvatarActor::TickManually( const float DeltaTime )
 	// that are not nearby, and completely hides the skybox and other very distant objects.  This is to help
 	// prevent simulation sickness when moving/rotating/scaling the entire world around them.
 	{
-		PostProcessComponent->bEnabled = VREd::ShowWorldMovementPostProcess->GetInt( ) == 1 ? true : false;
+		PostProcessComponent->bEnabled = GetDefault<UVRModeSettings>()->bShowWorldMovementPostProcess;
 
 		// Make sure our world size is reflected in the post process material
 		static FName WorldScaleFactorParameterName( "WorldScaleFactor" );
diff --git a/Engine/Source/Editor/VREditor/VREditorAvatarActor.h b/Engine/Source/Editor/VREditor/VREditorAvatarActor.h
--- a/Engine/Source/Editor/VREditor/VREditorAvatarActor.h
+++ b/Engine/Source/Editor/VREditor/VREditorAvatarActor.h
@@ -27,6 +27,12 @@ public:
 	/** Called by VREditorMode to update us every frame */
 	void TickManually( const float DeltaTime );
 
+	// UObject overrides
+	virtual bool IsEditorOnly() const final
+	{
+		return true;
+	}
+
 private:
 
 	/** Our avatar's head mesh */
diff --git a/Engine/Source/Editor/VREditor/VREditorBaseActor.cpp b/Engine/Source/Editor/VREditor/VREditorBaseActor.cpp
--- a/Engine/Source/Editor/VREditor/VREditorBaseActor.cpp
+++ b/Engine/Source/Editor/VREditor/VREditorBaseActor.cpp
@@ -7,7 +7,8 @@
 namespace VREd
 {
 	static FAutoConsoleVariable UIOnHandRotationOffset( TEXT( "VREd.UIOnHandRotationOffset" ), 45.0f, TEXT( "Rotation offset for UI that's docked to your hand, to make it more comfortable to hold" ) );
-	static FAutoConsoleVariable UIOnArmRotationOffset( TEXT( "VREd.UIOnArmRotationOffset" ), 15.0f, TEXT( "Rotation offset for UI that's docked to your arm, so it aligns with the controllers" ) );
+	static FAutoConsoleVariable UIOnArmRotationOffset( TEXT( "VREd.UIOnArmRotationOffset" ), 0.0f, TEXT( "Rotation offset for UI that's docked to your arm, so it aligns with the controllers" ) );
+	static FAutoConsoleVariable DockUISmoothingAmount( TEXT( "VREd.DockUISmoothingAmount" ), 0.75f, TEXT( "How much to smooth out UI transforms (frame rate sensitive)" ) );
 }
 
 AVREditorBaseActor::AVREditorBaseActor() 
@@ -22,7 +23,8 @@ AVREditorBaseActor::AVREditorBaseActor()
 	StartMoveToTransform( FTransform() ),
 	MoveToAlpha( 0.0f ),
 	MoveToTime( 0.0f ),
-	MoveToResultDock( EDockedTo::Nothing )
+	MoveToResultDock( EDockedTo::Nothing ),
+	LastDockedUIToWorld()
 {
 
 }
@@ -147,8 +149,29 @@ void AVREditorBaseActor::UpdateTransformIfDocked()
 			check( 0 );
 		}
 
+		// Smooth out the transform
+		if( DockedTo != EDockedTo::Custom && DockedTo != EDockedTo::Room )
+		{
+			if( LastDockedUIToWorld.IsSet() )
+			{
+				FTransform SmoothedDockedUIToWorld;
+				SmoothedDockedUIToWorld.Blend( NewTransform, LastDockedUIToWorld.GetValue(), VREd::DockUISmoothingAmount->GetFloat() );
+				NewTransform = SmoothedDockedUIToWorld;
+			}
+
+			LastDockedUIToWorld = NewTransform;
+		}
+		else
+		{
+			LastDockedUIToWorld.Reset();
+		}
+
 		SetTransform( NewTransform );
 	}
+	else
+	{
+		LastDockedUIToWorld.Reset();
+	}
 }
 
 FTransform AVREditorBaseActor::MakeUITransformLockedToHand( UViewportInteractor* Interactor, const bool bOnArm )
diff --git a/Engine/Source/Editor/VREditor/VREditorBaseActor.h b/Engine/Source/Editor/VREditor/VREditorBaseActor.h
--- a/Engine/Source/Editor/VREditor/VREditorBaseActor.h
+++ b/Engine/Source/Editor/VREditor/VREditorBaseActor.h
@@ -97,12 +97,13 @@ protected:
 	/** How big the actor should be */
 	float Scale;
 
-private:
+protected:
 
 	/** The VR mode that owns this actor */
 	UPROPERTY()
 	class UVREditorMode* VRMode;
 
+private:
 	/** Local rotation of the UI */
 	FRotator LocalRotation;
 
@@ -132,4 +133,7 @@ private:
 
 	/** Dock state to switch to when finished moving */
 	EDockedTo MoveToResultDock;
+
+	/** When docked, our transform last frame.  This is used for smoothing */
+	TOptional<FTransform> LastDockedUIToWorld;
 };
diff --git a/Engine/Source/Editor/VREditor/VREditorFloatingText.h b/Engine/Source/Editor/VREditor/VREditorFloatingText.h
--- a/Engine/Source/Editor/VREditor/VREditorFloatingText.h
+++ b/Engine/Source/Editor/VREditor/VREditorFloatingText.h
@@ -22,6 +22,10 @@ public:
 
 	// AActor overrides
 	virtual void PostActorCreated() override;
+	virtual bool IsEditorOnly() const final
+	{
+		return true;
+	}
 
 	/** Sets the text to display */
 	void SetText( const FText& NewText );
diff --git a/Engine/Source/Editor/VREditor/VREditorInteractor.cpp b/Engine/Source/Editor/VREditor/VREditorInteractor.cpp
--- a/Engine/Source/Editor/VREditor/VREditorInteractor.cpp
+++ b/Engine/Source/Editor/VREditor/VREditorInteractor.cpp
@@ -3,13 +3,15 @@
 #include "VREditorInteractor.h"
 #include "EngineUtils.h"
 #include "ViewportWorldInteraction.h"
-
+#include "Editor.h"
+#include "Engine/Selection.h"
 #include "VREditorMode.h"
 #include "VREditorFloatingText.h"
-
+#include "VREditorWorldInteraction.h"
 #include "VREditorButton.h"
 #include "VREditorFloatingUI.h"
 #include "VREditorDockableWindow.h"
+#include "ActorTransformer.h"
 
 #define LOCTEXT_NAMESPACE "VREditor"
 
@@ -18,13 +20,12 @@ UVREditorInteractor::UVREditorInteractor( const FObjectInitializer& Initializer
 	VRMode( nullptr ),
 	bIsModifierPressed( false ),
 	SelectAndMoveTriggerValue( 0.0f ),
-	bHasUIInFront( false ),
 	bHasUIOnForearm( false ),
 	bIsClickingOnUI( false ),
 	bIsRightClickingOnUI( false ),
 	bIsHoveringOverUI( false ),
 	UIScrollVelocity( 0.0f ),
-	LastClickReleaseTime( 0.0f ),
+	LastUIPressTime( 0.0f ),
 	bWantHelpLabels( false ),
 	HelpLabelShowOrHideStartTime( FTimespan::MinValue() )
 {
@@ -47,7 +48,7 @@ void UVREditorInteractor::Shutdown()
 	for ( auto& KeyAndValue : HelpLabels )
 	{
 		AFloatingText* FloatingText = KeyAndValue.Value;
-		UViewportWorldInteraction::DestroyTransientActor( GetVRMode().GetWorld(), FloatingText );
+		GetVRMode().DestroyTransientActor(FloatingText);
 	}
 
 	HelpLabels.Empty();	
@@ -57,11 +58,38 @@ void UVREditorInteractor::Shutdown()
 	Super::Shutdown();
 }
 
+void UVREditorInteractor::SetControllerType( const EControllerType& InControllerType )
+{
+	this->ControllerType = InControllerType;
+}
+
 void UVREditorInteractor::Tick( const float DeltaTime )
 {
 	Super::Tick( DeltaTime );
+
+	// If the other controller is dragging freely, the UI controller can assist
+	if( ControllerType == EControllerType::UI )
+	{
+		if( GetOtherInteractor() &&
+			( GetOtherInteractor()->GetDraggingMode() == EViewportInteractionDraggingMode::TransformablesFreely ) )
+		{
+			ControllerType = EControllerType::AssistingLaser;
+		}
+	}
+	// Otherwise the UI controller resets to a UI controller
+	// Allow for "trading off" during an assisted drag
+	else if( ControllerType == EControllerType::AssistingLaser )
+	{
+		if( GetOtherInteractor() &&
+			!( GetOtherInteractor()->GetDraggingMode() == EViewportInteractionDraggingMode::TransformablesFreely ||
+				GetOtherInteractor()->GetInteractorData().bWasAssistingDrag ) )
+		{
+			ControllerType = EControllerType::UI;
+		}
+	}
 }
 
+
 FHitResult UVREditorInteractor::GetHitResultFromLaserPointer( TArray<AActor*>* OptionalListOfIgnoredActors /*= nullptr*/, const bool bIgnoreGizmos /*= false*/,
 	TArray<UClass*>* ObjectsInFrontOfGizmo /*= nullptr */, const bool bEvenIfBlocked /*= false */, const float LaserLengthOverride /*= 0.0f */ )
 {
@@ -72,8 +100,6 @@ FHitResult UVREditorInteractor::GetHitResultFromLaserPointer( TArray<AActor*>* O
 		OptionalListOfIgnoredActors = &IgnoredActors;
 	}
 
-	OptionalListOfIgnoredActors->Add( GetVRMode().GetAvatarMeshActor() );
-	
 	// Ignore UI widgets too
 	if ( GetDraggingMode() == EViewportInteractionDraggingMode::TransformablesAtLaserImpact )
 	{
@@ -97,21 +123,33 @@ FHitResult UVREditorInteractor::GetHitResultFromLaserPointer( TArray<AActor*>* O
 	return UViewportInteractor::GetHitResultFromLaserPointer( OptionalListOfIgnoredActors, bIgnoreGizmos, ObjectsInFrontOfGizmo, bEvenIfBlocked, LaserLengthOverride );
 }
 
-void UVREditorInteractor::ResetHoverState( const float DeltaTime )
+void UVREditorInteractor::PreviewInputKey( class FEditorViewportClient& ViewportClient, FViewportActionKeyInput& Action, const FKey Key, const EInputEvent Event, bool& bOutWasHandled )
 {
-	bIsHoveringOverUI = false;
-}	
+	// Update modifier state
+	if( Action.ActionType == VRActionTypes::Modifier )
+	{
+		if( Event == IE_Pressed )
+		{
+			bIsModifierPressed = true;
+		}
+		else if( Event == IE_Released )
+		{
+			bIsModifierPressed = false;
+		}
+	}
 
-void UVREditorInteractor::OnStartDragging( const FVector& HitLocation, const bool bIsPlacingNewObjects )
-{
-	// If the user is holding down the modifier key, go ahead and duplicate the selection first.  Don't do this if we're
-	// placing objects right now though.
-	if ( bIsModifierPressed && !bIsPlacingNewObjects )
+	if( !bOutWasHandled )
 	{
-		WorldInteraction->Duplicate();
+		Super::PreviewInputKey( ViewportClient, Action, Key, Event, bOutWasHandled );
 	}
 }
 
+void UVREditorInteractor::ResetHoverState()
+{
+	Super::ResetHoverState();
+	bIsHoveringOverUI = false;
+}	
+
 float UVREditorInteractor::GetSlideDelta()
 {
 	return 0.0f;
@@ -122,16 +160,6 @@ bool UVREditorInteractor::IsHoveringOverUI() const
 	return bIsHoveringOverUI;
 }
 
-void UVREditorInteractor::SetHasUIInFront( const bool bInHasUIInFront )
-{
-	bHasUIInFront = bInHasUIInFront;
-}
-
-bool UVREditorInteractor::HasUIInFront() const
-{
-	return bHasUIInFront;
-}
-
 void UVREditorInteractor::SetHasUIOnForearm( const bool bInHasUIOnForearm )
 {
 	bHasUIOnForearm = bInHasUIOnForearm;
@@ -142,14 +170,14 @@ bool UVREditorInteractor::HasUIOnForearm() const
 	return bHasUIOnForearm;
 }
 
-UWidgetComponent* UVREditorInteractor::GetHoveringOverWidgetComponent() const
+UWidgetComponent* UVREditorInteractor::GetLastHoveredWidgetComponent() const
 {
-	return InteractorData.HoveringOverWidgetComponent;
+	return InteractorData.LastHoveredWidgetComponent;
 }
 
-void UVREditorInteractor::SetHoveringOverWidgetComponent( UWidgetComponent* NewHoveringOverWidgetComponent )
+void UVREditorInteractor::SetLastHoveredWidgetComponent( UWidgetComponent* NewHoveringOverWidgetComponent )
 {
-	InteractorData.HoveringOverWidgetComponent = NewHoveringOverWidgetComponent;
+	InteractorData.LastHoveredWidgetComponent = NewHoveringOverWidgetComponent;
 }
 
 bool UVREditorInteractor::IsModifierPressed() const
@@ -182,14 +210,14 @@ bool UVREditorInteractor::IsRightClickingOnUI() const
 	return bIsRightClickingOnUI;
 }
 
-void UVREditorInteractor::SetLastClickReleaseTime( const double InLastClickReleaseTime )
+void UVREditorInteractor::SetLastUIPressTime( const double InLastUIPressTime )
 {
-	LastClickReleaseTime = InLastClickReleaseTime;
+	LastUIPressTime = InLastUIPressTime;
 }
 
-double UVREditorInteractor::GetLastClickReleaseTime() const
+double UVREditorInteractor::GetLastUIPressTime() const
 {
-	return LastClickReleaseTime;
+	return LastUIPressTime;
 }
 
 void UVREditorInteractor::SetUIScrollVelocity( const float InUIScrollVelocity )
@@ -209,7 +237,7 @@ float UVREditorInteractor::GetSelectAndMoveTriggerValue() const
 
 bool UVREditorInteractor::GetIsLaserBlocked() const
 {
-	return bHasUIInFront;
+	return Super::GetIsLaserBlocked() || bHasUIInFront;
 }
 
 #undef LOCTEXT_NAMESPACE
