Merging //UE4/Private-Geometry to Dev-VREditor (//UE4/Dev-VREditor) 3339870

diff --git a/Engine/Source/Editor/VREditor/VREditorMotionControllerInteractor.cpp b/Engine/Source/Editor/VREditor/VREditorMotionControllerInteractor.cpp
--- a/Engine/Source/Editor/VREditor/VREditorMotionControllerInteractor.cpp
+++ b/Engine/Source/Editor/VREditor/VREditorMotionControllerInteractor.cpp
@@ -12,10 +12,12 @@
 #include "Materials/MaterialInstanceDynamic.h"
 #include "Engine/StaticMeshSocket.h"
 #include "ViewportWorldInteraction.h"
+#include "VREditorWorldInteraction.h"
 #include "VREditorMode.h"
 #include "VREditorUISystem.h"
 #include "VREditorFloatingText.h"
 #include "VREditorDockableWindow.h"
+#include "VREditorRadialFloatingUI.h"
 #include "UObject/Package.h"
 #include "Engine/Engine.h"
 #include "EngineGlobals.h"
@@ -23,9 +25,15 @@
 #include "IHeadMountedDisplay.h"
 #include "MotionControllerComponent.h"
 #include "Features/IModularFeatures.h"
-
-#include "SVRRadialPanel.h"
+#include "Engine/Selection.h"
 #include "VREditorWidgetComponent.h"
+#include "LevelEditorActions.h"
+#include "Components/SplineComponent.h"
+#include "Components/SplineMeshComponent.h"
+#include "DrawDebugHelpers.h"
+#include "VREditorActions.h"
+#include "VREditorAssetContainer.h"
+#include "VRModeSettings.h"
 
 namespace VREd
 {
@@ -36,50 +44,50 @@ namespace VREd
 	static FAutoConsoleVariable ViveLaserPointerStartOffset( TEXT( "VI.ViveLaserPointerStartOffset" ), 1.25f /* 8.5f */, TEXT( "How far to offset the start of the laser pointer to avoid overlapping the hand mesh geometry (Vive)" ) );
 
 	// Laser visuals
-	static FAutoConsoleVariable LaserPointerRadius( TEXT( "VREd.LaserPointerRadius" ), 0.5f, TEXT( "Radius of the laser pointer line" ) );
-	static FAutoConsoleVariable LaserPointerHoverBallRadius( TEXT( "VREd.LaserPointerHoverBallRadius" ), 1.5f, TEXT( "Radius of the visual cue for a hovered object along the laser pointer ray" ) );
+	static FAutoConsoleVariable LaserPointerRadius( TEXT( "VREd.LaserPointerRadius" ), .5f, TEXT( "Radius of the laser pointer line" ) );
+	static FAutoConsoleVariable LaserPointerHoverBallRadius( TEXT( "VREd.LaserPointerHoverBallRadius" ), 1.0f, TEXT( "Radius of the visual cue for a hovered object along the laser pointer ray" ) );
 	static FAutoConsoleVariable LaserPointerLightPullBackDistance( TEXT( "VREd.LaserPointerLightPullBackDistance" ), 2.5f, TEXT( "How far to pull back our little hover light from the impact surface" ) );
-	static FAutoConsoleVariable LaserPointerLightRadius( TEXT( "VREd.LaserPointLightRadius" ), 20.0f, TEXT( "How big our hover light is" ) );
-
+	static FAutoConsoleVariable LaserPointerLightRadius( TEXT( "VREd.LaserPointLightRadius" ), 10.0f, TEXT( "How big our hover light is" ) );
+	static FAutoConsoleVariable LaserRadiusScaleWhenOverUI( TEXT( "VREd.LaserRadiusScaleWhenOverUI" ), 0.25f, TEXT( "How much to scale down the size of the laser pointer radius when over UI" ) );
+	static FAutoConsoleVariable HoverBallRadiusScaleWhenOverUI(TEXT("VREd.HoverBallRadiusScaleWhenOverUI"), 0.4f, TEXT("How much to scale down the size of the hover ball when over UI"));
 	//Trigger
-	static FAutoConsoleVariable TriggerLightlyPressedThreshold_Vive( TEXT( "VI.TriggerLightlyPressedThreshold_Vive" ), 0.03f, TEXT( "Minimum trigger threshold before we consider the trigger at least 'lightly pressed'" ) );
 	static FAutoConsoleVariable TriggerDeadZone_Vive( TEXT( "VI.TriggerDeadZone_Vive" ), 0.01f, TEXT( "Trigger dead zone.  The trigger must be fully released before we'll trigger a new 'light press'" ) );
-	static FAutoConsoleVariable TriggerLightlyPressedThreshold_Rift( TEXT( "VI.TriggerLightlyPressedThreshold_Rift" ), 0.2f, TEXT( "Minimum trigger threshold before we consider the trigger at least 'lightly pressed'" ) );
 	static FAutoConsoleVariable TriggerDeadZone_Rift( TEXT( "VI.TriggerDeadZone_Rift" ), 0.15f, TEXT( "Trigger dead zone.  The trigger must be fully released before we'll trigger a new 'light press'" ) );
-	static FAutoConsoleVariable TriggerFullyPressedThreshold( TEXT( "VI.TriggerFullyPressedThreshold" ), 0.95f, TEXT( "Minimum trigger threshold before we consider the trigger 'fully pressed'" ) );
-	static FAutoConsoleVariable TriggerFullyPressedReleaseThreshold( TEXT( "VI.TriggerFullyPressedReleaseThreshold" ), 0.8f, TEXT( "After fully pressing the trigger, if the axis falls below this threshold we no longer consider it fully pressed" ) );
-
-	// Haptic feedback
-	static FAutoConsoleVariable SleepForRiftHaptics( TEXT( "VI.SleepForRiftHaptics" ), 1, TEXT( "When enabled, we'll sleep the game thread mid-frame to wait for haptic effects to finish.  This can be devasting to performance!" ) );
-	static FAutoConsoleVariable MinHapticTimeForRift( TEXT( "VREd.MinHapticTimeForRift" ), 0.005f, TEXT( "How long to play haptic effects on the Rift" ) );
+	static FAutoConsoleVariable TriggerFullyPressedThreshold_Vive( TEXT( "VI.TriggerFullyPressedThreshold_Vive" ), 0.90f, TEXT( "Minimum trigger threshold before we consider the trigger 'fully pressed'" ) );
+	static FAutoConsoleVariable TriggerFullyPressedThreshold_Rift( TEXT( "VI.TriggerFullyPressedThreshold_Rift" ), 0.99f, TEXT( "Minimum trigger threshold before we consider the trigger 'fully pressed'" ) );
 
 	static FAutoConsoleVariable TrackpadAbsoluteDragSpeed( TEXT( "VREd.TrackpadAbsoluteDragSpeed" ), 40.0f, TEXT( "How fast objects move toward or away when you drag on the touchpad while carrying them" ) );
 	static FAutoConsoleVariable TrackpadRelativeDragSpeed( TEXT( "VREd.TrackpadRelativeDragSpeed" ), 8.0f, TEXT( "How fast objects move toward or away when you hold a direction on an analog stick while carrying them" ) );
 	static FAutoConsoleVariable InvertTrackpadVertical( TEXT( "VREd.InvertTrackpadVertical" ), 1, TEXT( "Toggles inverting the touch pad vertical axis" ) );
-	static FAutoConsoleVariable TriggerLightlyPressedLockTime( TEXT( "VREd.TriggerLightlyPressedLockTime" ), 0.15f, TEXT( "If the trigger remains lightly pressed for longer than this, we'll continue to treat it as a light press in some cases" ) );
 	static FAutoConsoleVariable MinVelocityForInertia( TEXT( "VREd.MinVelocityForMotionControllerInertia" ), 1.0f, TEXT( "Minimum velocity (in cm/frame in unscaled room space) before inertia will kick in when releasing objects (or the world)" ) );
+	static FAutoConsoleVariable MinTrackpadOffsetBeforeRadialMenu( TEXT( "VREd.MinTrackpadOffsetBeforeRadialMenu" ), 0.5f, TEXT( "How far you have to hold the trackpad upward before you can placing objects instantly by pulling the trigger" ) );
 
+	static FAutoConsoleVariable ShowControllerHelpLabels( TEXT( "VREd.ShowControllerHelpLabels" ), 0, TEXT( "Enables help text overlay when controllers are near the viewer" ) );
 	static FAutoConsoleVariable HelpLabelFadeDuration( TEXT( "VREd.HelpLabelFadeDuration" ), 0.4f, TEXT( "Duration to fade controller help labels in and out" ) );
 	static FAutoConsoleVariable HelpLabelFadeDistance( TEXT( "VREd.HelpLabelFadeDistance" ), 30.0f, TEXT( "Distance at which controller help labels should appear (in cm)" ) );
+
+	static FAutoConsoleVariable MinJoystickOffsetBeforeFlick(TEXT("VREd.MinJoystickOffsetBeforeFlick"), 0.4f, TEXT("Dead zone for flick actions on the motion controller"));
+	
+	static FAutoConsoleVariable SequencerScrubMax(TEXT("VREd.SequencerScrubMax"), 2.0f, TEXT("Max fast forward or fast reverse magnitude"));
+
 }
 
 const FName UVREditorMotionControllerInteractor::TrackpadPositionX = FName( "TrackpadPositionX" );
 const FName UVREditorMotionControllerInteractor::TrackpadPositionY = FName( "TrackpadPositionY" );
 const FName UVREditorMotionControllerInteractor::TriggerAxis = FName( "TriggerAxis" );
+const FName UVREditorMotionControllerInteractor::MotionController_Left_PressedTriggerAxis = FName( "MotionController_Left_PressedTriggerAxis" );
+const FName UVREditorMotionControllerInteractor::MotionController_Right_PressedTriggerAxis = FName( "MotionController_Right_PressedTriggerAxis" );
 const FName UVREditorMotionControllerInteractor::MotionController_Left_FullyPressedTriggerAxis = FName( "MotionController_Left_FullyPressedTriggerAxis" );
 const FName UVREditorMotionControllerInteractor::MotionController_Right_FullyPressedTriggerAxis = FName( "MotionController_Right_FullyPressedTriggerAxis" );
-const FName UVREditorMotionControllerInteractor::MotionController_Left_LightlyPressedTriggerAxis = FName( "MotionController_Left_LightlyPressedTriggerAxis" );
-const FName UVREditorMotionControllerInteractor::MotionController_Right_LightlyPressedTriggerAxis = FName( "MotionController_Right_LightlyPressedTriggerAxis" );
 
 namespace VREditorKeyNames
 {
 	// @todo vreditor input: Ideally these would not be needed, but SteamVR fires off it's "trigger pressed" event
 	// well before the trigger is fully down (*click*)
+	static const FName MotionController_Left_PressedTriggerAxis( "MotionController_Left_PressedTriggerAxis" );
+	static const FName MotionController_Right_PressedTriggerAxis( "MotionController_Right_PressedTriggerAxis" );
 	static const FName MotionController_Left_FullyPressedTriggerAxis( "MotionController_Left_FullyPressedTriggerAxis" );
 	static const FName MotionController_Right_FullyPressedTriggerAxis( "MotionController_Right_FullyPressedTriggerAxis" );
-
-	static const FName MotionController_Left_LightlyPressedTriggerAxis( "MotionController_Left_LightlyPressedTriggerAxis" );
-	static const FName MotionController_Right_LightlyPressedTriggerAxis( "MotionController_Right_LightlyPressedTriggerAxis" );
 }
 
 namespace SteamVRControllerKeyNames
@@ -95,7 +103,6 @@ UVREditorMotionControllerInteractor::UVREditorMotionControllerInteractor( const
 	Super( Initializer ),
 	MotionControllerComponent( nullptr ),
 	HandMeshComponent( nullptr ),
-	LaserPointerMeshComponent( nullptr ),
 	LaserPointerMID( nullptr ),
 	TranslucentLaserPointerMID( nullptr ),
 	HoverMeshComponent( nullptr ),
@@ -104,15 +111,20 @@ UVREditorMotionControllerInteractor::UVREditorMotionControllerInteractor( const
 	ControllerHandSide( EControllerHand::Pad ),
 	bHaveMotionController( false ),
 	bIsTriggerFullyPressed( false ),
-	bIsTriggerAtLeastLightlyPressed( false ),
-	RealTimeTriggerWasLightlyPressed( 0.0f ),
+	bIsTriggerPressed( false ),
 	bHasTriggerBeenReleasedSinceLastPress( true ),
-	LastHapticTime( 0.0f ),
 	bIsTouchingTrackpad( false ),
+	bIsPressingTrackpad( false ),
 	TrackpadPosition( FVector2D::ZeroVector ),
 	LastTrackpadPosition( FVector2D::ZeroVector ),
 	bIsTrackpadPositionValid { false, false },
-	LastTrackpadPositionUpdateTime( FTimespan::MinValue() )
+	LastTrackpadPositionUpdateTime( FTimespan::MinValue() ),
+	LastActiveTrackpadUpdateTime( FTimespan::MinValue() ), 
+	bForceShowLaser(false),
+	ForceLaserColor(),
+	bIsScrubbingSequence(false),
+	LastSwipe(ETouchSwipeDirection::None),
+	InitialTouchPosition(FVector2D::ZeroVector)
 {
 }
 
@@ -133,8 +145,8 @@ void UVREditorMotionControllerInteractor::Init( class UVREditorMode* InVRMode )
 	if ( ControllerHandSide == EControllerHand::Left )
 	{
 		AddKeyAction( EKeys::MotionController_Left_Grip1, FViewportActionKeyInput( ViewportWorldActionTypes::WorldMovement ) );
-		AddKeyAction( UVREditorMotionControllerInteractor::MotionController_Left_FullyPressedTriggerAxis, FViewportActionKeyInput( ViewportWorldActionTypes::SelectAndMove ) );
-		AddKeyAction( UVREditorMotionControllerInteractor::MotionController_Left_LightlyPressedTriggerAxis, FViewportActionKeyInput( ViewportWorldActionTypes::SelectAndMove_LightlyPressed ) );
+		AddKeyAction( UVREditorMotionControllerInteractor::MotionController_Left_FullyPressedTriggerAxis, FViewportActionKeyInput( ViewportWorldActionTypes::SelectAndMove_FullyPressed ) );
+		AddKeyAction( UVREditorMotionControllerInteractor::MotionController_Left_PressedTriggerAxis, FViewportActionKeyInput( ViewportWorldActionTypes::SelectAndMove ) );
 		AddKeyAction( SteamVRControllerKeyNames::Touch0, FViewportActionKeyInput( VRActionTypes::Touch ) );
 		AddKeyAction( EKeys::MotionController_Left_TriggerAxis, FViewportActionKeyInput( UVREditorMotionControllerInteractor::TriggerAxis ) );
 		AddKeyAction( EKeys::MotionController_Left_Thumbstick_X, FViewportActionKeyInput( UVREditorMotionControllerInteractor::TrackpadPositionX ) );
@@ -154,8 +166,8 @@ void UVREditorMotionControllerInteractor::Init( class UVREditorMode* InVRMode )
 	else if ( ControllerHandSide == EControllerHand::Right )
 	{
 		AddKeyAction( EKeys::MotionController_Right_Grip1, FViewportActionKeyInput( ViewportWorldActionTypes::WorldMovement ) );
-		AddKeyAction( UVREditorMotionControllerInteractor::MotionController_Right_FullyPressedTriggerAxis, FViewportActionKeyInput( ViewportWorldActionTypes::SelectAndMove ) );
-		AddKeyAction( UVREditorMotionControllerInteractor::MotionController_Right_LightlyPressedTriggerAxis, FViewportActionKeyInput( ViewportWorldActionTypes::SelectAndMove_LightlyPressed ) );
+		AddKeyAction( UVREditorMotionControllerInteractor::MotionController_Right_FullyPressedTriggerAxis, FViewportActionKeyInput( ViewportWorldActionTypes::SelectAndMove_FullyPressed ) );
+		AddKeyAction( UVREditorMotionControllerInteractor::MotionController_Right_PressedTriggerAxis, FViewportActionKeyInput( ViewportWorldActionTypes::SelectAndMove ) );
 		AddKeyAction( SteamVRControllerKeyNames::Touch1, FViewportActionKeyInput( VRActionTypes::Touch ) );
 		AddKeyAction( EKeys::MotionController_Right_TriggerAxis, FViewportActionKeyInput( UVREditorMotionControllerInteractor::TriggerAxis ) );
 		AddKeyAction( EKeys::MotionController_Right_Thumbstick_X, FViewportActionKeyInput( UVREditorMotionControllerInteractor::TrackpadPositionX ) );
@@ -196,69 +208,35 @@ void UVREditorMotionControllerInteractor::SetupComponent( AActor* OwningActor )
 
 		MotionControllerComponent->Hand = ControllerHandSide;
 
-		MotionControllerComponent->bDisableLowLatencyUpdate = false;
+		// @todo vreditor: Reenable late frame updates after we've sorted out why they cause popping artifacts on Rift
+		MotionControllerComponent->bDisableLowLatencyUpdate = true;
 	}
 
+	const UVREditorAssetContainer& AssetContainer = VRMode->GetAssetContainer();
+
 	// Hand mesh
 	{
-		HandMeshComponent = NewObject<UStaticMeshComponent>( OwningActor );
-		OwningActor->AddOwnedComponent( HandMeshComponent );
-		HandMeshComponent->SetupAttachment( MotionControllerComponent );
-		HandMeshComponent->RegisterComponent();
-
-		// @todo vreditor extensibility: We need this to be able to be overridden externally, or simply based on the HMD name (but allowing external folders)
-		FString MeshName;
-		FString MaterialName;
-		if ( GetVRMode().GetHMDDeviceType() == EHMDDeviceType::DT_SteamVR )
-		{
-			MeshName = TEXT( "/Engine/VREditor/Devices/Vive/VivePreControllerMesh" );
-			MaterialName = TEXT( "/Engine/VREditor/Devices/Vive/VivePreControllerMaterial_Inst" );
-		}
-		else
-		{
-			MeshName = TEXT( "/Engine/VREditor/Devices/Oculus/OculusControllerMesh" );
-			MaterialName = TEXT( "/Engine/VREditor/Devices/Oculus/OculusControllerMaterial_Inst" );
-		}
+		HandMeshComponent = VRMode->CreateMotionControllerMesh( OwningActor, MotionControllerComponent ); 
+		check( HandMeshComponent != nullptr );
+		HandMeshComponent->SetCastShadow(false);
 
-		UStaticMesh* HandMesh = LoadObject<UStaticMesh>( nullptr, *MeshName );
-		check( HandMesh != nullptr );
-
-		HandMeshComponent->SetStaticMesh( HandMesh );
-		HandMeshComponent->SetMobility( EComponentMobility::Movable );
-		HandMeshComponent->SetCollisionEnabled( ECollisionEnabled::NoCollision );
-
-		UMaterialInstance* HandMeshMaterialInst = LoadObject<UMaterialInstance>( nullptr, *MaterialName );
-		check( HandMeshMaterialInst != nullptr );
-		HandMeshMID = UMaterialInstanceDynamic::Create( HandMeshMaterialInst, GetTransientPackage() );
+		UMaterialInterface* HandMeshMaterial = GetVRMode().GetHMDDeviceType() == EHMDDeviceType::DT_SteamVR ? AssetContainer.VivePreControllerMaterial : AssetContainer.OculusControllerMaterial;
+		check( HandMeshMaterial != nullptr );
+		HandMeshMID = UMaterialInstanceDynamic::Create( HandMeshMaterial, GetTransientPackage() );
 		check( HandMeshMID != nullptr );
 		HandMeshComponent->SetMaterial( 0, HandMeshMID );
 	}
 
-	// Laser pointer
 	{
-		LaserPointerMeshComponent = NewObject<UStaticMeshComponent>( OwningActor );
-		OwningActor->AddOwnedComponent( LaserPointerMeshComponent );
-		LaserPointerMeshComponent->SetupAttachment( MotionControllerComponent );
-		LaserPointerMeshComponent->RegisterComponent();
-
-		UStaticMesh* LaserPointerMesh = LoadObject<UStaticMesh>( nullptr, TEXT( "/Engine/VREditor/LaserPointer/LaserPointerMesh" ) );
-		check( LaserPointerMesh != nullptr );
-
-		LaserPointerMeshComponent->SetStaticMesh( LaserPointerMesh );
-		LaserPointerMeshComponent->SetMobility( EComponentMobility::Movable );
-		LaserPointerMeshComponent->SetCollisionEnabled( ECollisionEnabled::NoCollision );
-
-		UMaterialInstance* LaserPointerMaterialInst = LoadObject<UMaterialInstance>( nullptr, TEXT( "/Engine/VREditor/LaserPointer/LaserPointerMaterialInst" ) );
-		check( LaserPointerMaterialInst != nullptr );
-		LaserPointerMID = UMaterialInstanceDynamic::Create( LaserPointerMaterialInst, GetTransientPackage() );
-		check( LaserPointerMID != nullptr );
-		LaserPointerMeshComponent->SetMaterial( 0, LaserPointerMID );
-
-		UMaterialInstance* TranslucentLaserPointerMaterialInst = LoadObject<UMaterialInstance>( nullptr, TEXT( "/Engine/VREditor/LaserPointer/TranslucentLaserPointerMaterialInst" ) );
-		check( TranslucentLaserPointerMaterialInst != nullptr );
-		TranslucentLaserPointerMID = UMaterialInstanceDynamic::Create( TranslucentLaserPointerMaterialInst, GetTransientPackage() );
-		check( TranslucentLaserPointerMID != nullptr );
-		LaserPointerMeshComponent->SetMaterial( 1, TranslucentLaserPointerMID );
+		UMaterialInterface* LaserPointerMaterial = AssetContainer.LaserPointerMaterial;
+		check(LaserPointerMaterial != nullptr);
+		LaserPointerMID = UMaterialInstanceDynamic::Create(LaserPointerMaterial, GetTransientPackage());
+		check(LaserPointerMID != nullptr);
+
+		UMaterialInterface* TranslucentLaserPointerMaterial = AssetContainer.LaserPointerTranslucentMaterial;
+		check(TranslucentLaserPointerMaterial != nullptr);
+		TranslucentLaserPointerMID = UMaterialInstanceDynamic::Create(TranslucentLaserPointerMaterial, GetTransientPackage());
+		check(TranslucentLaserPointerMID != nullptr);
 	}
 
 	// Hover cue for laser pointer
@@ -268,11 +246,15 @@ void UVREditorMotionControllerInteractor::SetupComponent( AActor* OwningActor )
 		HoverMeshComponent->SetupAttachment( OwningActor->GetRootComponent() );
 		HoverMeshComponent->RegisterComponent();
 
-		UStaticMesh* HoverMesh = LoadObject<UStaticMesh>( nullptr, TEXT( "/Engine/VREditor/LaserPointer/HoverMesh" ) );
+		UStaticMesh* HoverMesh = AssetContainer.LaserPointerHoverMesh;
 		check( HoverMesh != nullptr );
 		HoverMeshComponent->SetStaticMesh( HoverMesh );
 		HoverMeshComponent->SetMobility( EComponentMobility::Movable );
 		HoverMeshComponent->SetCollisionEnabled( ECollisionEnabled::NoCollision );
+		HoverMeshComponent->SetCastShadow(false);
+
+		HoverMeshComponent->SetMaterial( 0, LaserPointerMID );
+		HoverMeshComponent->SetMaterial( 1, TranslucentLaserPointerMID );
 
 		// Add a light!
 		{
@@ -290,6 +272,58 @@ void UVREditorMotionControllerInteractor::SetupComponent( AActor* OwningActor )
 			HoverPointLightComponent->SetCastShadows( false );
 		}
 	}
+
+	{
+		const int32 NumLaserSplinePoints = 12;
+
+		UStaticMesh* MiddleSplineMesh = AssetContainer.LaserPointerMesh;
+		check (MiddleSplineMesh != nullptr);
+		UStaticMesh* StartSplineMesh = AssetContainer.LaserPointerStartMesh;
+		check (StartSplineMesh != nullptr);
+		UStaticMesh* EndSplineMesh = AssetContainer.LaserPointerEndMesh;
+		check (EndSplineMesh != nullptr);
+
+		LaserSplineComponent = NewObject<USplineComponent>(OwningActor);
+		OwningActor->AddOwnedComponent(LaserSplineComponent);
+		LaserSplineComponent->SetupAttachment(MotionControllerComponent);
+		LaserSplineComponent->RegisterComponent();
+		LaserSplineComponent->SetVisibility(false);
+
+		for (int32 i = 0; i < NumLaserSplinePoints; i++)
+		{
+			USplineMeshComponent* SplineSegment = NewObject<USplineMeshComponent>(OwningActor);
+			OwningActor->AddOwnedComponent(SplineSegment);
+			SplineSegment->SetMobility(EComponentMobility::Movable);
+			SplineSegment->SetCollisionEnabled(ECollisionEnabled::NoCollision);
+			SplineSegment->PostPhysicsComponentTick.bCanEverTick = false;
+			SplineSegment->SetupAttachment(MotionControllerComponent);
+
+			UStaticMesh* StaticMesh = nullptr;
+			if (i == 0)
+			{
+				StaticMesh = StartSplineMesh;
+			}
+			else if (i == NumLaserSplinePoints)
+			{
+				StaticMesh = EndSplineMesh;
+			}
+			else
+			{
+				StaticMesh = MiddleSplineMesh;
+			}
+			
+			SplineSegment->SetStaticMesh(StaticMesh);
+			SplineSegment->bTickInEditor = true;
+			SplineSegment->bCastDynamicShadow = false;
+			SplineSegment->CastShadow = false;
+			SplineSegment->SetMaterial(0, LaserPointerMID);
+			SplineSegment->SetMaterial(1, TranslucentLaserPointerMID);
+			SplineSegment->SetVisibility(true);
+			SplineSegment->RegisterComponent();
+
+			LaserSplineMeshComponents.Add(SplineSegment);
+		}
+	}
 }
 
 
@@ -305,7 +339,6 @@ void UVREditorMotionControllerInteractor::Shutdown()
 
 	MotionControllerComponent = nullptr;
 	HandMeshComponent = nullptr;
-	LaserPointerMeshComponent = nullptr;
 	LaserPointerMID = nullptr;
 	TranslucentLaserPointerMID = nullptr;
 	HoverMeshComponent = nullptr;
@@ -351,11 +384,35 @@ void UVREditorMotionControllerInteractor::Tick( const float DeltaTime )
 			HandMeshComponent->SetVisibility( false );
 		}
 
+		// Offset the beginning of the laser pointer a bit, so that it doesn't overlap the hand mesh
+		const float LaserPointerStartOffset =
+			WorldScaleFactor *
+			( GetVRMode().GetHMDDeviceType() == EHMDDeviceType::DT_OculusRift ? VREd::OculusLaserPointerStartOffset->GetFloat() : VREd::ViveLaserPointerStartOffset->GetFloat() );
+
+		// The laser pointer needs to stay the same size relative to our tracking space, so we inverse compensate for world to meters scale here
+		float LaserPointerRadius = VREd::LaserPointerRadius->GetFloat() * WorldScaleFactor;
+		float HoverMeshRadius = VREd::LaserPointerHoverBallRadius->GetFloat() * WorldScaleFactor;
+
+		// If we're hovering over something really close to the camera, go ahead and shrink the effect
+		// @todo vreditor: Can we make this actually just sized based on distance automatically?  The beam and impact point are basically a cone.
+		if( IsHoveringOverUI() )
+		{
+			LaserPointerRadius *= VREd::LaserRadiusScaleWhenOverUI->GetFloat();
+			HoverMeshRadius *= VREd::HoverBallRadiusScaleWhenOverUI->GetFloat();
+		}
+
+		const bool bEvenIfBlocked = false;
+
+		// If we're currently grabbing the world with this interactor
+		const bool bDraggingWorld = (InteractorData.DraggingMode == EViewportInteractionDraggingMode::World || 
+			(GetOtherInteractor() != nullptr && GetOtherInteractor()->GetInteractorData().DraggingMode == EViewportInteractionDraggingMode::World && InteractorData.DraggingMode == EViewportInteractionDraggingMode::AssistingDrag));
+
 		FVector LaserPointerStart, LaserPointerEnd;
-		if ( GetLaserPointer( /* Out */ LaserPointerStart, /* Out */ LaserPointerEnd ) )
+		const bool bHasLaser = GetLaserPointer(/* Out */ LaserPointerStart, /* Out */ LaserPointerEnd, bEvenIfBlocked);
+		if (bForceShowLaser || (bHasLaser && !bDraggingWorld))
 		{
 			// Only show the laser if we're actually in VR
-			LaserPointerMeshComponent->SetVisibility( GetVRMode().IsActuallyUsingVR() );
+			SetLaserVisibility(GetVRMode().IsActuallyUsingVR());
 
 			// NOTE: We don't need to set the laser pointer location and rotation, as the MotionControllerComponent will do
 			// that later in the frame.  
@@ -363,46 +420,19 @@ void UVREditorMotionControllerInteractor::Tick( const float DeltaTime )
 			// If we're actively dragging something around, then we'll crop the laser length to the hover impact
 			// point.  Otherwise, we always want the laser to protrude through hovered objects, so that you can
 			// interact with translucent gizmo handles that are occluded by geometry
-			FVector LaserPointerImpactPoint = LaserPointerEnd;
-			if ( IsHoveringOverGizmo() || IsHoveringOverUI() || 
-				 ( GetHoverComponent() != nullptr && GetHoverComponent()->GetOwner()->IsA<AVREditorDockableWindow>() ) )
+			if( IsHoveringOverGizmo() ||
+				IsHoveringOverUI() || 
+				IsHovering() )
 			{
-				LaserPointerImpactPoint = GetHoverLocation();
+				LaserPointerEnd = GetHoverLocation();
 			}
 
-			// Apply rotation offset to the laser direction
-			const float LaserPointerRotationOffset = 0.0f;
-			LaserPointerMeshComponent->SetRelativeRotation( FRotator( LaserPointerRotationOffset, 0.0f, 0.0f ) );
-
-			const FVector LaserPointerDirection = ( LaserPointerImpactPoint - LaserPointerStart ).GetSafeNormal();
-
-			// Offset the beginning of the laser pointer a bit, so that it doesn't overlap the hand mesh
-			const float LaserPointerStartOffset =
-				WorldScaleFactor *
-				( GetVRMode().GetHMDDeviceType() == EHMDDeviceType::DT_OculusRift ? VREd::OculusLaserPointerStartOffset->GetFloat() : VREd::ViveLaserPointerStartOffset->GetFloat() );
-
-			const float LaserPointerLength = FMath::Max( 0.000001f, ( LaserPointerImpactPoint - LaserPointerStart ).Size() - LaserPointerStartOffset );
-
-			LaserPointerMeshComponent->SetRelativeLocation( FRotator( LaserPointerRotationOffset, 0.0f, 0.0f ).RotateVector( FVector( LaserPointerStartOffset, 0.0f, 0.0f ) ) );
-
-			// The laser pointer needs to stay the same size relative to our tracking space, so we inverse compensate for world to meters scale here
-			float LaserPointerRadius = VREd::LaserPointerRadius->GetFloat() * WorldScaleFactor;
-			float HoverMeshRadius = VREd::LaserPointerHoverBallRadius->GetFloat() * WorldScaleFactor;
-
-			// If we're hovering over something really close to the camera, go ahead and shrink the effect
-			// @todo vreditor: Can we make this actually just sized based on distance automatically?  The beam and impact point are basically a cone.
-			if ( IsHoveringOverUI() )
+			if( IsHovering() && !GetIsLaserBlocked() )
 			{
-				LaserPointerRadius *= 0.35f;	// @todo vreditor tweak
-				HoverMeshRadius *= 0.35f;	// @todo vreditor tweak
-			}
+				const FVector DirectionTowardHoverLocation = (GetHoverLocation() - LaserPointerStart).GetSafeNormal();
 
-			LaserPointerMeshComponent->SetRelativeScale3D( FVector( LaserPointerLength, LaserPointerRadius * 2.0f, LaserPointerRadius * 2.0f ) );
-
-			if ( IsHovering() )
-			{
 				// The hover effect needs to stay the same size relative to our tracking space, so we inverse compensate for world to meters scale here
-				HoverMeshComponent->SetRelativeScale3D( FVector( HoverMeshRadius * 2.0f ) );
+				HoverMeshComponent->SetRelativeScale3D( FVector( HoverMeshRadius * 2.0f ) * ( 0.25f + 1.0f - this->GetSelectAndMoveTriggerValue() * 0.75f ) );
 				HoverMeshComponent->SetVisibility( true );
 				HoverMeshComponent->SetWorldLocation( GetHoverLocation() );
 
@@ -414,7 +444,7 @@ void UVREditorMotionControllerInteractor::Tick( const float DeltaTime )
 
 				// Pull hover light back a bit from the end of the ray
 				const float PullBackAmount = VREd::LaserPointerLightPullBackDistance->GetFloat() * WorldInteraction->GetWorldScaleFactor();
-				HoverPointLightComponent->SetWorldLocation( GetHoverLocation() - PullBackAmount * LaserPointerDirection );
+				HoverPointLightComponent->SetWorldLocation( GetHoverLocation() - PullBackAmount * DirectionTowardHoverLocation );
 			}
 			else
 			{
@@ -424,63 +454,106 @@ void UVREditorMotionControllerInteractor::Tick( const float DeltaTime )
 		}
 		else
 		{
-			LaserPointerMeshComponent->SetVisibility( false );
+			SetLaserVisibility(false);
 			HoverMeshComponent->SetVisibility( false );
 			HoverPointLightComponent->SetVisibility( false );
 		}
+
+		// Update the curved laser. No matter if we actually show the laser it needs to update, 
+		// so if in the next frame it needs to be visible it won't interpolate from a previous location.
+		{
+			// Get the hand transform and forward vector.
+			FTransform InteractorTransform;
+			FVector InteractorForwardVector;
+			GetTransformAndForwardVector( /* Out */ InteractorTransform, /* Out */ InteractorForwardVector);
+
+			// Offset the start point of the laser.
+			LaserPointerStart = InteractorTransform.GetLocation() + (InteractorForwardVector * FVector(LaserPointerStartOffset, 0.0f, 0.0f));
+
+			UpdateSplineLaser(LaserPointerStart, LaserPointerEnd, InteractorForwardVector);
+		}
+
+		bForceShowLaser = false;
 	}
 
 	// Updating laser colors for both hands
 	{
-		UVREditorMode::EColors ResultColorType = UVREditorMode::EColors::DefaultColor;
+		FLinearColor ResultColor;
 		float CrawlSpeed = 0.0f;
 		float CrawlFade = 0.0f;
 
-		const EViewportInteractionDraggingMode DraggingMode = GetDraggingMode();
-
-		const bool bIsDraggingWorldWithTwoHands = 
-			GetOtherInteractor() != nullptr &&
-			( ( DraggingMode == EViewportInteractionDraggingMode::AssistingDrag && GetOtherInteractor()->GetDraggingMode() == EViewportInteractionDraggingMode::World ) ||
-			  ( DraggingMode == EViewportInteractionDraggingMode::World && GetOtherInteractor()->GetDraggingMode() == EViewportInteractionDraggingMode::AssistingDrag ) );
-
-		if ( bIsDraggingWorldWithTwoHands )
-		{
-			ResultColorType = UVREditorMode::EColors::WorldDraggingColor_TwoHands;
-		}
-		else if ( DraggingMode == EViewportInteractionDraggingMode::World )
+		if (ForceLaserColor.IsSet())
 		{
-			//@todo VREditor find a way to make the teleporter do it
-			// 			if ( GetVRMode().GetWorldInteraction().IsTeleporting() ) //@todo VREditor: Make teleport an action
-			// 			{
-			// 				ResultColorType = FVREditorMode::EColors::TeleportColor;
-			// 			}
-			// 			else
-			// 			{
-			// We can teleport in this mode, so animate the laser a bit
-			CrawlFade = 1.0f;
-			CrawlSpeed = 5.0f;
-			ResultColorType = UVREditorMode::EColors::WorldDraggingColor_OneHand;
-			//			}
+			ResultColor = ForceLaserColor.GetValue();
+			ForceLaserColor.Reset();
 		}
-		else if ( DraggingMode == EViewportInteractionDraggingMode::TransformablesAtLaserImpact ||
-			DraggingMode == EViewportInteractionDraggingMode::Material ||
-			DraggingMode == EViewportInteractionDraggingMode::TransformablesFreely ||
-			DraggingMode == EViewportInteractionDraggingMode::TransformablesWithGizmo ||
-			DraggingMode == EViewportInteractionDraggingMode::AssistingDrag ||
-			DraggingMode == EViewportInteractionDraggingMode::Interactable ||
-			( GetVRMode().GetUISystem().IsInteractorDraggingDockUI( this ) && GetVRMode().GetUISystem().IsDraggingDockUI() ) )
+		else
 		{
-			ResultColorType = UVREditorMode::EColors::SelectionColor;
+			if (InteractorData.HoveringOverTransformGizmoComponent != nullptr)
+			{
+				ResultColor = WorldInteraction->GetColor(UViewportWorldInteraction::EColors::GizmoHover);
+			}
+			else
+			{
+				const EViewportInteractionDraggingMode DraggingMode = GetDraggingMode();
+				if (DraggingMode == EViewportInteractionDraggingMode::World || 
+					(DraggingMode == EViewportInteractionDraggingMode::AssistingDrag && GetOtherInteractor() != nullptr && GetOtherInteractor()->GetDraggingMode() == EViewportInteractionDraggingMode::World))
+				{
+					// We can teleport in this mode, so animate the laser a bit
+					CrawlFade = 1.0f;
+					CrawlSpeed = 5.0f;
+					ResultColor = GetVRMode().GetColor(UVREditorMode::EColors::WorldDraggingColor);
+				}
+				else if (DraggingMode == EViewportInteractionDraggingMode::TransformablesAtLaserImpact ||
+						 DraggingMode == EViewportInteractionDraggingMode::Material ||
+						 DraggingMode == EViewportInteractionDraggingMode::TransformablesFreely ||
+						 DraggingMode == EViewportInteractionDraggingMode::AssistingDrag)
+				{
+					ResultColor = GetVRMode().GetColor(UVREditorMode::EColors::SelectionColor);
+				}
+				else if (DraggingMode == EViewportInteractionDraggingMode::TransformablesWithGizmo)
+				{
+					ResultColor = WorldInteraction->GetColor(UViewportWorldInteraction::EColors::GizmoHover);
+				}
+				else if (DraggingMode == EViewportInteractionDraggingMode::Interactable ||
+						(GetVRMode().GetUISystem().IsInteractorDraggingDockUI(this) && GetVRMode().GetUISystem().IsDraggingDockUI()))
+				{
+					ResultColor = GetVRMode().GetColor(UVREditorMode::EColors::UIColor);
+				}
+				else if (GetLastHoverComponent() != nullptr && GetLastHoverComponent()->GetOwner() != nullptr && GetLastHoverComponent()->GetOwner()->IsA(AVREditorDockableWindow::StaticClass()))
+				{
+					AVREditorDockableWindow* HoveredDockWindow = Cast<AVREditorDockableWindow>(GetLastHoverComponent()->GetOwner());
+					if (HoveredDockWindow && HoveredDockWindow->GetSelectionBarMeshComponent() == GetLastHoverComponent())
+					{
+						ResultColor = GetVRMode().GetColor(UVREditorMode::EColors::UIColor);
+					}
+					else
+					{
+						ResultColor = GetVRMode().GetColor(UVREditorMode::EColors::DefaultColor);
+					}
+				}
+				else if (ControllerType == EControllerType::UI)
+				{
+					ResultColor = GetVRMode().GetColor(UVREditorMode::EColors::UIColor);
+				}
+				else if (IsHoveringOverSelectedActor())
+				{
+					ResultColor = GetVRMode().GetColor(UVREditorMode::EColors::SelectionColor);
+				}
+				else
+				{
+					ResultColor = GetVRMode().GetColor(UVREditorMode::EColors::DefaultColor);
+				}
+			}
 		}
 
-		const FLinearColor ResultColor = GetVRMode().GetColor( ResultColorType );
-		SetLaserVisuals( ResultColor, CrawlFade, CrawlSpeed );
+		SetLaserVisuals(ResultColor, CrawlFade, CrawlSpeed);
 	}
 
 	UpdateHelpLabels();
 }
 
-void UVREditorMotionControllerInteractor::CalculateDragRay()
+void UVREditorMotionControllerInteractor::CalculateDragRay( float& InOutDragRayLength, float& InOutDragRayVelocity)
 {
 	const FTimespan CurrentTime = FTimespan::FromSeconds( FPlatformTime::Seconds() );
 	const float WorldScaleFactor = WorldInteraction->GetWorldScaleFactor();
@@ -494,9 +567,9 @@ void UVREditorMotionControllerInteractor::CalculateDragRay()
 
 		if ( !FMath::IsNearlyZero( SlideDelta ) )
 		{
-			InteractorData.DragRayLength += SlideDelta;
+			InOutDragRayLength += SlideDelta;
 
-			InteractorData.DragRayLengthVelocity = 0.0f;
+			InOutDragRayVelocity = 0.0f;
 
 			// Don't apply inertia unless the user dragged a decent amount this frame
 			if ( bIsAbsolute && FMath::Abs( SlideDelta ) >= VREd::MinVelocityForInertia->GetFloat() * WorldScaleFactor )
@@ -504,41 +577,41 @@ void UVREditorMotionControllerInteractor::CalculateDragRay()
 				// Don't apply inertia if our data is sort of old
 				if ( CurrentTime - LastTrackpadPositionUpdateTime <= FTimespan::FromSeconds( 1.0f / 30.0f ) )
 				{
-					InteractorData.DragRayLengthVelocity = SlideDelta;
+					InOutDragRayVelocity = SlideDelta;
 				}
 			}
 
 			// Don't go too far
-			if ( InteractorData.DragRayLength < 0.0f )
+			if ( InOutDragRayLength < 0.0f )
 			{
-				InteractorData.DragRayLength = 0.0f;
-				InteractorData.DragRayLengthVelocity = 0.0f;
+				InOutDragRayLength = 0.0f;
+				InOutDragRayVelocity = 0.0f;
 			}
 		}
 	}
 	else
 	{
-		if ( !FMath::IsNearlyZero( InteractorData.DragRayLengthVelocity ) )
+		if ( !FMath::IsNearlyZero( InOutDragRayVelocity ) )
 		{
 			// Apply drag ray length inertia
-			InteractorData.DragRayLength += InteractorData.DragRayLengthVelocity;
+			InOutDragRayLength += InOutDragRayVelocity;
 
 			// Don't go too far!
-			if ( InteractorData.DragRayLength < 0.0f )
+			if ( InOutDragRayLength < 0.0f )
 			{
-				InteractorData.DragRayLength = 0.0f;
-				InteractorData.DragRayLengthVelocity = 0.0f;
+				InOutDragRayLength = 0.0f;
+				InOutDragRayVelocity = 0.0f;
 			}
 
 			// Apply damping
-			FVector RayVelocityVector( InteractorData.DragRayLengthVelocity, 0.0f, 0.0f );
+			FVector RayVelocityVector( InOutDragRayVelocity, 0.0f, 0.0f );
 			const bool bVelocitySensitive = true;
 			WorldInteraction->ApplyVelocityDamping( RayVelocityVector, bVelocitySensitive );
-			InteractorData.DragRayLengthVelocity = RayVelocityVector.X;
+			InOutDragRayVelocity = RayVelocityVector.X;
 		}
 		else
 		{
-			InteractorData.DragRayLengthVelocity = 0.0f;
+			InOutDragRayVelocity = 0.0f;
 		}
 	}
 }
@@ -550,87 +623,107 @@ EHMDDeviceType::Type UVREditorMotionControllerInteractor::GetHMDDeviceType() con
 }
 
 
-void UVREditorMotionControllerInteractor::HandleInputKey( FEditorViewportClient& ViewportClient, FViewportActionKeyInput& Action, const FKey Key, const EInputEvent Event, bool& bOutWasHandled )
+void UVREditorMotionControllerInteractor::PreviewInputKey( FEditorViewportClient& ViewportClient, FViewportActionKeyInput& Action, const FKey Key, const EInputEvent Event, bool& bOutWasHandled )
 {
-	if ( !bOutWasHandled )
+	if (bIsScrubbingSequence && 
+		GetControllerType() == EControllerType::UI && 
+		Action.ActionType == ViewportWorldActionTypes::SelectAndMove &&
+		Action.Event == IE_Pressed)
 	{
-		bool bShouldAbsorbEvent = false;
+		ToggleSequencerScrubbingMode();
+		GetVRMode().GetUISystem().TryToSpawnRadialMenu(this, true);
+	}
 
-		// If "SelectAndMove" was pressed, we need to make sure that "SelectAndMove_LightlyPressed" is no longer
-		// in effect before we start handling "SelectAndMove".
-		if ( Action.ActionType == ViewportWorldActionTypes::SelectAndMove && IsTriggerAtLeastLightlyPressed() )
+	// Update touch state
+	if( Action.ActionType == VRActionTypes::Touch )
+	{
+		if( Event == IE_Pressed )
 		{
-			if ( Event == IE_Pressed )
+			bIsTouchingTrackpad = true;
+
+			// Set initial position when starting to touch the trackpad
+			InitialTouchPosition = TrackpadPosition;
+		}
+		else if( Event == IE_Released )
+		{
+			bIsTouchingTrackpad = false;
+			bIsTrackpadPositionValid[ 0 ] = false;
+			bIsTrackpadPositionValid[ 1 ] = false;
+
+			// Detect swipe on trackpad.
+			const FVector2D SwipeDelta = LastTrackpadPosition - InitialTouchPosition;
+			const float AbsSwipeDeltaX = FMath::Abs(SwipeDelta.X);
+			const float AbsSwipeDeltaY = FMath::Abs(SwipeDelta.Y);
+			if (!FMath::IsNearlyZero(SwipeDelta.X, 1.0f) && AbsSwipeDeltaX > AbsSwipeDeltaY)
 			{
-				// How long since the trigger was lightly pressed?
-				const float TimeSinceLightlyPressed = ( float ) ( FPlatformTime::Seconds() - GetRealTimeTriggerWasLightlyPressed() );
-				if ( !IsClickingOnUI() &&	// @todo vreditor: UI clicks happen with light presses; we never want to convert to a hard press!
-					GetDraggingMode() != EViewportInteractionDraggingMode::Material &&	// @todo vreditor: Material dragging happens with light presses, don't convert to a hard press!
-					GetDraggingMode() != EViewportInteractionDraggingMode::TransformablesAtLaserImpact &&	// @todo vreditor: Actor dragging happens with light presses, don't convert to a hard press!
-					AllowTriggerFullPress() &&
-					( !AllowTriggerLightPressLocking() || TimeSinceLightlyPressed < VREd::TriggerLightlyPressedLockTime->GetFloat() ) )
+				if (SwipeDelta.X > 0)
 				{
-					SetTriggerAtLeastLightlyPressed( false );
-					SetTriggerBeenReleasedSinceLastPress( false );
-
-					// Synthesize an input key for releasing the light press
-					// NOTE: Intentionally re-entrant call here.
-					const EInputEvent InputEvent = IE_Released;
-					const bool bWasLightReleaseHandled = UViewportInteractor::HandleInputKey( ViewportClient, GetControllerSide() == EControllerHand::Left ? VREditorKeyNames::MotionController_Left_LightlyPressedTriggerAxis : VREditorKeyNames::MotionController_Right_LightlyPressedTriggerAxis, InputEvent );
+					LastSwipe = ETouchSwipeDirection::Right;
+					UndoRedoFromSwipe(LastSwipe);
 				}
-				else
+				else if (SwipeDelta.X < 0)
 				{
-					// The button was held in a "lightly pressed" state for a long enough time that we should just continue to
-					// treat it like a light press.  This makes it much easier to hold it in this state when you need to!
-					bShouldAbsorbEvent = true;
+					LastSwipe = ETouchSwipeDirection::Left;
+					UndoRedoFromSwipe(LastSwipe);
 				}
 			}
-			else
+			else if (!FMath::IsNearlyZero(SwipeDelta.Y, 1.0f))
 			{
-				// Absorb the release of the SelectAndMove event if we are still in a lightly pressed state
-				bShouldAbsorbEvent = true;
+				if (SwipeDelta.Y > 0)
+				{
+					LastSwipe = ETouchSwipeDirection::Up;
+				}
+				else if (SwipeDelta.Y < 0)
+				{
+					LastSwipe = ETouchSwipeDirection::Down;
+				}
 			}
 		}
+	}
 
-		if ( bShouldAbsorbEvent )
-		{
-			bOutWasHandled = true;
-		}
+	if( Action.ActionType == VRActionTypes::ConfirmRadialSelection )
+	{
+		bIsPressingTrackpad = Event == IE_Released ? false : true;
 	}
 
-	if ( !bOutWasHandled )
+	if( GetControllerType() == EControllerType::Laser )
 	{
-		// Update touch state
-		if ( Action.ActionType == VRActionTypes::Touch )
-		{
-			if ( Event == IE_Pressed )
-			{
-				bIsTouchingTrackpad = true;
-			}
-			else if ( Event == IE_Released )
-			{
-				bIsTouchingTrackpad = false;
-				bIsTrackpadPositionValid[ 0 ] = false;
-				bIsTrackpadPositionValid[ 1 ] = false;
-			}
-		}
+		// Are we holding "up" on the trackpad?
+		const bool bIsHoldingUpOnTrackpad =
+			this->bIsTrackpadPositionValid &&
+			this->TrackpadPosition.Y >= VREd::MinTrackpadOffsetBeforeRadialMenu->GetFloat() &&
+			( GetHMDDeviceType() == EHMDDeviceType::DT_OculusRift || bIsPressingTrackpad );
 
-		// Update modifier state
-		if ( Action.ActionType == VRActionTypes::Modifier )
+		if( bIsHoldingUpOnTrackpad && Action.ActionType == ViewportWorldActionTypes::SelectAndMove && Event == IE_Pressed )
 		{
 			bOutWasHandled = true;
-			if ( Event == IE_Pressed )
+
+			// Try to place the object currently selected
+			TArray<UObject*> ObjectsToPlace;
 			{
-				bIsModifierPressed = true;
+				FEditorDelegates::LoadSelectedAssetsIfNeeded.Broadcast();
+				GEditor->GetSelectedObjects()->GetSelectedObjects( /* Out */ ObjectsToPlace );
 			}
-			else if ( Event == IE_Released )
+
+			if( ObjectsToPlace.Num() > 0 )
 			{
-				bIsModifierPressed = false;
+				Action.bIsInputCaptured = true;
+
+				const bool bShouldInterpolateFromDragLocation = false;
+				UActorFactory* FactoryToUse = nullptr;	// Use default factory
+				GetVRMode().GetVREditorWorldInteraction()->StartPlacingObjects( ObjectsToPlace, FactoryToUse, this, bShouldInterpolateFromDragLocation );
 			}
 		}
 	}
 
-	if( Action.ActionType == VRActionTypes::ConfirmRadialSelection )
+
+	Super::PreviewInputKey( ViewportClient, Action, Key, Event, bOutWasHandled );
+}
+
+
+void UVREditorMotionControllerInteractor::HandleInputKey( FEditorViewportClient& ViewportClient, FViewportActionKeyInput& Action, const FKey Key, const EInputEvent Event, bool& bOutWasHandled )
+{
+	if( !bOutWasHandled && Action.ActionType == VRActionTypes::ConfirmRadialSelection )
 	{
 		bOutWasHandled = true;
 		const EViewportInteractionDraggingMode DraggingMode = GetDraggingMode();
@@ -642,8 +735,10 @@ void UVREditorMotionControllerInteractor::HandleInputKey( FEditorViewportClient&
 			{
 				FVector PlaceAt = GetHoverLocation();
 				const bool bIsPlacingActors = true;
+				const bool bAllowInterpolationWhenPlacing = true;
 				const bool bStartTransaction = true;
-				WorldInteraction->StartDragging( this, WorldInteraction->GetTransformGizmoActor()->GetRootComponent(), PlaceAt, bIsPlacingActors, bStartTransaction );
+				const bool bWithGrabberSphere = false;	// Never use the grabber sphere when dragging at laser impact
+				WorldInteraction->StartDragging( this, WorldInteraction->GetTransformGizmoActor()->GetRootComponent(), PlaceAt, bIsPlacingActors, bAllowInterpolationWhenPlacing, bStartTransaction, bWithGrabberSphere );
 			}
 		}
 		else if( Event == IE_Released )
@@ -663,7 +758,7 @@ void UVREditorMotionControllerInteractor::HandleInputAxis( FEditorViewportClient
 {
 	if ( Action.ActionType == TriggerAxis )
 	{
-		const float TriggerLightlyPressedThreshold = ( GetHMDDeviceType() == EHMDDeviceType::DT_OculusRift ) ? VREd::TriggerLightlyPressedThreshold_Rift->GetFloat() : VREd::TriggerLightlyPressedThreshold_Vive->GetFloat();
+		const float TriggerPressedThreshold = ( GetHMDDeviceType() == EHMDDeviceType::DT_OculusRift ) ? GetDefault<UVRModeSettings>()->TriggerPressedThreshold_Rift : GetDefault<UVRModeSettings>()->TriggerPressedThreshold_Vive;
 		const float TriggerDeadZone = ( GetHMDDeviceType() == EHMDDeviceType::DT_OculusRift ) ? VREd::TriggerDeadZone_Rift->GetFloat() : VREd::TriggerDeadZone_Vive->GetFloat();
 
 		// Synthesize "lightly pressed" events for the trigger
@@ -671,31 +766,24 @@ void UVREditorMotionControllerInteractor::HandleInputAxis( FEditorViewportClient
 			// Store latest trigger value amount
 			SelectAndMoveTriggerValue = Delta;
 
-			FViewportActionKeyInput* OptionalAction = GetActionWithName( ViewportWorldActionTypes::SelectAndMove );
-			const bool bIsFullPressAlreadyCapturing = OptionalAction && OptionalAction->bIsInputCaptured;
-
-			if ( !bIsFullPressAlreadyCapturing &&		// Don't fire if we're already capturing input for a full press
-				!bIsTriggerAtLeastLightlyPressed &&	// Don't fire unless we are already pressed
+			if( !bIsTriggerPressed &&			// Don't fire if we are already pressed
 				bHasTriggerBeenReleasedSinceLastPress &&	// Only if we've been fully released since the last time we fired
-				Delta >= TriggerLightlyPressedThreshold )
+				Delta >= TriggerPressedThreshold )
 			{
-				bIsTriggerAtLeastLightlyPressed = true;
-				SetAllowTriggerLightPressLocking( true );
-				SetAllowTriggerFullPress( true );
-				RealTimeTriggerWasLightlyPressed = FPlatformTime::Seconds();
+				bIsTriggerPressed = true;
 				bHasTriggerBeenReleasedSinceLastPress = false;
 
 				// Synthesize an input key for this light press
 				const EInputEvent InputEvent = IE_Pressed;
-				const bool bWasLightPressHandled = UViewportInteractor::HandleInputKey( ViewportClient, ControllerHandSide == EControllerHand::Left ? MotionController_Left_LightlyPressedTriggerAxis : MotionController_Right_LightlyPressedTriggerAxis, InputEvent );
+				const bool bWasLightPressHandled = UViewportInteractor::HandleInputKey( ViewportClient, ControllerHandSide == EControllerHand::Left ? MotionController_Left_PressedTriggerAxis : MotionController_Right_PressedTriggerAxis, InputEvent );
 			}
-			else if ( bIsTriggerAtLeastLightlyPressed && Delta < TriggerLightlyPressedThreshold )
+			else if ( bIsTriggerPressed && Delta < TriggerPressedThreshold )
 			{
-				bIsTriggerAtLeastLightlyPressed = false;
+				bIsTriggerPressed = false;
 
 				// Synthesize an input key for this light press
 				const EInputEvent InputEvent = IE_Released;
-				const bool bWasLightReleaseHandled = UViewportInteractor::HandleInputKey( ViewportClient, ControllerHandSide == EControllerHand::Left ? MotionController_Left_LightlyPressedTriggerAxis : MotionController_Right_LightlyPressedTriggerAxis, InputEvent );
+				const bool bWasLightReleaseHandled = UViewportInteractor::HandleInputKey( ViewportClient, ControllerHandSide == EControllerHand::Left ? MotionController_Left_PressedTriggerAxis : MotionController_Right_PressedTriggerAxis, InputEvent );
 			}
 		}
 
@@ -706,21 +794,20 @@ void UVREditorMotionControllerInteractor::HandleInputAxis( FEditorViewportClient
 
 		// Synthesize "fully pressed" events for the trigger
 		{
-			if ( !bIsTriggerFullyPressed &&	// Don't fire unless we are already pressed
-				 AllowTriggerFullPress() &&
-				 Delta >= VREd::TriggerFullyPressedThreshold->GetFloat() )
+			const float TriggerFullyPressedThreshold = (GetHMDDeviceType() == EHMDDeviceType::DT_OculusRift) ? VREd::TriggerFullyPressedThreshold_Rift->GetFloat() : VREd::TriggerFullyPressedThreshold_Vive->GetFloat();
+
+			if ( !bIsTriggerFullyPressed &&	// Don't fire if we are already pressed
+				 Delta >= TriggerFullyPressedThreshold )
 			{
 				bIsTriggerFullyPressed = true;
 
-				// Synthesize an input key for this full press
 				const EInputEvent InputEvent = IE_Pressed;
 				UViewportInteractor::HandleInputKey( ViewportClient, ControllerHandSide == EControllerHand::Left ? MotionController_Left_FullyPressedTriggerAxis : MotionController_Right_FullyPressedTriggerAxis, InputEvent );
 			}
-			else if ( bIsTriggerFullyPressed && Delta < VREd::TriggerFullyPressedReleaseThreshold->GetFloat() )
+			else if ( bIsTriggerFullyPressed && Delta < TriggerPressedThreshold )
 			{
 				bIsTriggerFullyPressed = false;
 
-				// Synthesize an input key for this full press
 				const EInputEvent InputEvent = IE_Released;
 				UViewportInteractor::HandleInputKey( ViewportClient, ControllerHandSide == EControllerHand::Left ? MotionController_Left_FullyPressedTriggerAxis : MotionController_Right_FullyPressedTriggerAxis, InputEvent );
 			}
@@ -780,26 +867,6 @@ void UVREditorMotionControllerInteractor::PollInput()
 	}
 }
 
-bool UVREditorMotionControllerInteractor::IsTriggerAtLeastLightlyPressed() const
-{
-	return bIsTriggerAtLeastLightlyPressed;
-}
-
-double UVREditorMotionControllerInteractor::GetRealTimeTriggerWasLightlyPressed() const
-{
-	return RealTimeTriggerWasLightlyPressed;
-}
-
-void UVREditorMotionControllerInteractor::SetTriggerAtLeastLightlyPressed( const bool bInTriggerAtLeastLightlyPressed )
-{
-	bIsTriggerAtLeastLightlyPressed = bInTriggerAtLeastLightlyPressed;
-}
-
-void UVREditorMotionControllerInteractor::SetTriggerBeenReleasedSinceLastPress( const bool bInTriggerBeenReleasedSinceLastPress )
-{
-	bHasTriggerBeenReleasedSinceLastPress = bInTriggerBeenReleasedSinceLastPress;
-}
-
 void UVREditorMotionControllerInteractor::PlayHapticEffect( const float Strength )
 {
 	IInputInterface* InputInterface = FSlateApplication::Get().GetInputInterface();
@@ -807,96 +874,13 @@ void UVREditorMotionControllerInteractor::PlayHapticEffect( const float Strength
 	{
 		const double CurrentTime = FPlatformTime::Seconds();
 
-		// If we're dealing with an Oculus Rift, we have to setup haptic feedback directly.  Otherwise we can use our
-		// generic force feedback system
-		if ( GetHMDDeviceType() == EHMDDeviceType::DT_OculusRift )
-		{
-			// Haptics are a little strong on Oculus Touch, so we scale them down a bit
-			const float HapticScaleForRift = 0.8f;
-
-			FHapticFeedbackValues HapticFeedbackValues;
-			HapticFeedbackValues.Amplitude = Strength * HapticScaleForRift;
-			HapticFeedbackValues.Frequency = 0.5f;
-			InputInterface->SetHapticFeedbackValues( WorldInteraction->GetMotionControllerID(), ( int32 )ControllerHandSide, HapticFeedbackValues );
+		//@todo viewportinteration
+		FForceFeedbackValues ForceFeedbackValues;
+		ForceFeedbackValues.LeftLarge = ControllerHandSide == EControllerHand::Left ? Strength : 0;
+		ForceFeedbackValues.RightLarge = ControllerHandSide == EControllerHand::Right ? Strength : 0;
 
-			LastHapticTime = CurrentTime;
-		}
-		else
-		{
-			//@todo viewportinteration
-			FForceFeedbackValues ForceFeedbackValues;
-			ForceFeedbackValues.LeftLarge = ControllerHandSide == EControllerHand::Left ? Strength : 0;
-			ForceFeedbackValues.RightLarge = ControllerHandSide == EControllerHand::Right ? Strength : 0;
-
-			// @todo vreditor: If an Xbox controller is plugged in, this causes both the motion controllers and the Xbox controller to vibrate!
-			InputInterface->SetForceFeedbackChannelValues( WorldInteraction->GetMotionControllerID(), ForceFeedbackValues );
-
-			if ( ForceFeedbackValues.LeftLarge > KINDA_SMALL_NUMBER )
-			{
-				LastHapticTime = CurrentTime;
-			}
-
-			if ( ForceFeedbackValues.RightLarge > KINDA_SMALL_NUMBER )
-			{
-				LastHapticTime = CurrentTime;
-			}
-		}
-	}
-
-	// @todo vreditor: We'll stop haptics right away because if the frame hitches, the controller will be left vibrating
-	StopOldHapticEffects();
-}
-
-void UVREditorMotionControllerInteractor::StopOldHapticEffects()
-{
-	const float MinHapticTime = VREd::MinHapticTimeForRift->GetFloat();
-
-	IInputInterface* InputInterface = FSlateApplication::Get().GetInputInterface();
-	if ( InputInterface )
-	{
-		// If we're dealing with an Oculus Rift, we have to setup haptic feedback directly.  Otherwise we can use our
-		// generic force feedback system
-		if ( GetHMDDeviceType() == EHMDDeviceType::DT_OculusRift )
-		{
-			bool bWaitingForMoreHaptics = false;
-
-			if ( LastHapticTime != 0.0f )
-			{
-				do
-				{
-					bWaitingForMoreHaptics = false;
-
-					const double CurrentTime = FPlatformTime::Seconds();
-
-					// Left hand
-					if ( CurrentTime - LastHapticTime >= MinHapticTime )
-					{
-						FHapticFeedbackValues HapticFeedbackValues;
-						HapticFeedbackValues.Amplitude = 0.0f;
-						HapticFeedbackValues.Frequency = 0.0f;
-						InputInterface->SetHapticFeedbackValues( WorldInteraction->GetMotionControllerID(), ( int32 ) ControllerHandSide, HapticFeedbackValues );
-
-						LastHapticTime = 0.0;
-					}
-					else if ( LastHapticTime != 0.0 )
-					{
-						bWaitingForMoreHaptics = true;
-					}
-
-					if ( bWaitingForMoreHaptics && VREd::SleepForRiftHaptics->GetInt() != 0 )
-					{
-						FPlatformProcess::Sleep( 0 );
-					}
-				}
-				// @todo vreditor urgent: This is needed so that haptics don't play too long.  Our Oculus code doesn't currently support 
-				// multi-threading, so we need to delay the main thread to make sure we stop it before it rumbles for more than an instant!
-				while ( bWaitingForMoreHaptics && VREd::SleepForRiftHaptics->GetInt() != 0 );
-			}
-		}
-		else
-		{
-			// @todo vreditor: Do we need to cancel haptics for non-Rift devices?  Doesn't seem like it
-		}
+		// @todo vreditor: If an Xbox controller is plugged in, this causes both the motion controllers and the Xbox controller to vibrate!
+		InputInterface->SetForceFeedbackChannelValues( WorldInteraction->GetMotionControllerID(), ForceFeedbackValues );
 	}
 }
 
@@ -915,7 +899,7 @@ bool UVREditorMotionControllerInteractor::GetTransformAndForwardVector( FTransfo
 	return false;
 }
 
-float UVREditorMotionControllerInteractor::GetTrackpadSlideDelta()
+float UVREditorMotionControllerInteractor::GetTrackpadSlideDelta( const bool Axis )
 {
 	const bool bIsAbsolute = ( GetVRMode().GetHMDDeviceType() == EHMDDeviceType::DT_SteamVR );
 	float SlideDelta = 0.0f;
@@ -923,11 +907,11 @@ float UVREditorMotionControllerInteractor::GetTrackpadSlideDelta()
 	{
 		if ( bIsAbsolute )
 		{
-			SlideDelta = ( ( TrackpadPosition.Y - LastTrackpadPosition.Y ) * VREd::TrackpadAbsoluteDragSpeed->GetFloat() );
+			SlideDelta = ( ( TrackpadPosition[Axis] - LastTrackpadPosition[Axis] ) * VREd::TrackpadAbsoluteDragSpeed->GetFloat() );
 		}
 		else
 		{
-			SlideDelta = ( TrackpadPosition.Y * VREd::TrackpadRelativeDragSpeed->GetFloat() );
+			SlideDelta = ( TrackpadPosition[Axis] * VREd::TrackpadRelativeDragSpeed->GetFloat() );
 		}
 	}
 
@@ -969,6 +953,20 @@ FTimespan& UVREditorMotionControllerInteractor::GetLastTrackpadPositionUpdateTim
 	return LastTrackpadPositionUpdateTime;
 }
 
+FTimespan& UVREditorMotionControllerInteractor::GetLastActiveTrackpadUpdateTime()
+{
+	return LastActiveTrackpadUpdateTime;
+}
+
+void UVREditorMotionControllerInteractor::SetForceShowLaser(const bool bInForceShow)
+{
+	bForceShowLaser = bInForceShow;
+}
+
+void UVREditorMotionControllerInteractor::SetForceLaserColor(const FLinearColor& InColor)
+{
+	ForceLaserColor = InColor;
+}
 
 /** Changes the color of the buttons on the handmesh */
 void UVREditorMotionControllerInteractor::ApplyButtonPressColors( const FViewportActionKeyInput& Action )
@@ -978,8 +976,7 @@ void UVREditorMotionControllerInteractor::ApplyButtonPressColors( const FViewpor
 	const EInputEvent Event = Action.Event;
 
 	// Trigger
-	if ( ActionType == ViewportWorldActionTypes::SelectAndMove ||
-		ActionType == ViewportWorldActionTypes::SelectAndMove_LightlyPressed )
+	if ( ActionType == ViewportWorldActionTypes::SelectAndMove )
 	{
 		static FName StaticTriggerParameter( "B1" );
 		SetMotionControllerButtonPressedVisuals( Event, StaticTriggerParameter, PressStrength );
@@ -1069,15 +1066,15 @@ void UVREditorMotionControllerInteractor::ShowHelpForHand( const bool bShowIt )
 						LabelText = LOCTEXT( "WorldMovementHelp", "Move World" );
 						ComponentName = TEXT( "WorldMovementHelp" );
 					}
-					else if ( Action.ActionType == ViewportWorldActionTypes::SelectAndMove )
+					else if ( Action.ActionType == ViewportWorldActionTypes::SelectAndMove_FullyPressed )
 					{
-						LabelText = LOCTEXT( "SelectAndMoveHelp", "Select & Move" );
-						ComponentName = TEXT( "SelectAndMoveHelp" );
+						LabelText = LOCTEXT( "SelectAndMove_FullyPressedHelp", "Select & Move" );
+						ComponentName = TEXT( "SelectAndMove_FullyPressedHelp" );
 					}
-					else if ( Action.ActionType == ViewportWorldActionTypes::SelectAndMove_LightlyPressed )
+					else if ( Action.ActionType == ViewportWorldActionTypes::SelectAndMove )
 					{
-						LabelText = LOCTEXT( "SelectAndMove_LightlyPressedHelp", "Select & Move" );
-						ComponentName = TEXT( "SelectAndMove_LightlyPressedHelp" );
+						LabelText = LOCTEXT( "SelectAndMove_Help", "Select & Move" );
+						ComponentName = TEXT( "SelectAndMove_Help" );
 					}
 					else if ( Action.ActionType == VRActionTypes::Touch )
 					{
@@ -1107,7 +1104,7 @@ void UVREditorMotionControllerInteractor::ShowHelpForHand( const bool bShowIt )
 					
 					const bool bWithSceneComponent = false;	// Nope, we'll spawn our own inside AFloatingText
 					check( VRMode );
-					AFloatingText* FloatingText = UViewportWorldInteraction::SpawnTransientSceneActor<AFloatingText>( GetVRMode().GetWorld(), ComponentName, bWithSceneComponent );
+					AFloatingText* FloatingText = GetVRMode().SpawnTransientSceneActor<AFloatingText>(ComponentName);
 					FloatingText->SetText( LabelText );
 
 					HelpLabels.Add( Key, FloatingText );
@@ -1127,10 +1124,10 @@ void UVREditorMotionControllerInteractor::UpdateHelpLabels()
 	// Only show help labels if the hand is pretty close to the face
 	const float DistanceToHead = ( GetTransform().GetLocation() - HeadTransform.GetLocation() ).Size();
 	const float MinDistanceToHeadForHelp = VREd::HelpLabelFadeDistance->GetFloat() * GetVRMode().GetWorldScaleFactor();	// (in cm)
-	bool bShowHelp = DistanceToHead <= MinDistanceToHeadForHelp;
+	bool bShowHelp = VREd::ShowControllerHelpLabels->GetInt() != 0 && DistanceToHead <= MinDistanceToHeadForHelp;
 
 	// Don't show help if a UI is summoned on that hand
-	if ( HasUIInFront() || HasUIOnForearm() || GetVRMode().GetUISystem().IsShowingRadialMenu( this ) )
+	if ( HasUIOnForearm() || GetVRMode().GetUISystem().IsShowingRadialMenu( this ) )
 	{
 		bShowHelp = false;
 	}
@@ -1146,7 +1143,7 @@ void UVREditorMotionControllerInteractor::UpdateHelpLabels()
 		for ( auto& KeyAndValue : HelpLabels )
 		{
 			AFloatingText* FloatingText = KeyAndValue.Value;
-			UViewportWorldInteraction::DestroyTransientActor( GetVRMode().GetWorld(), FloatingText );
+			GetVRMode().DestroyTransientActor(FloatingText);
 		}
 		HelpLabels.Reset();
 	}
@@ -1208,7 +1205,7 @@ UStaticMeshSocket* UVREditorMotionControllerInteractor::FindMeshSocketForKey( US
 	}
 	else if ( Key == EKeys::MotionController_Left_Trigger || Key == EKeys::MotionController_Right_Trigger ||
 		Key == VREditorKeyNames::MotionController_Left_FullyPressedTriggerAxis || Key == VREditorKeyNames::MotionController_Right_FullyPressedTriggerAxis ||
-		Key == VREditorKeyNames::MotionController_Left_LightlyPressedTriggerAxis || Key == VREditorKeyNames::MotionController_Right_LightlyPressedTriggerAxis )
+		Key == VREditorKeyNames::MotionController_Left_PressedTriggerAxis || Key == VREditorKeyNames::MotionController_Right_PressedTriggerAxis )
 	{
 		static FName TriggerSocketName( "Trigger" );
 		SocketName = TriggerSocketName;
@@ -1285,6 +1282,69 @@ UStaticMeshSocket* UVREditorMotionControllerInteractor::FindMeshSocketForKey( US
 	return nullptr;
 };
 
+void UVREditorMotionControllerInteractor::UpdateSplineLaser(const FVector& InStartLocation, const FVector& InEndLocation, const FVector& InForward)
+{
+	if (LaserSplineComponent)
+	{
+		// Clear the segments before updating it
+		LaserSplineComponent->ClearSplinePoints(true);
+
+		const FVector SmoothLaserDirection = InEndLocation - InStartLocation;
+		float Distance = SmoothLaserDirection.Size();
+		const FVector StraightLaserEndLocation = InStartLocation + (InForward * Distance);
+		const int32 NumLaserSplinePoints = LaserSplineMeshComponents.Num();
+
+		LaserSplineComponent->AddSplinePoint(InStartLocation, ESplineCoordinateSpace::World, false);
+		for (int32 Index = 1; Index < NumLaserSplinePoints; Index++)
+		{
+			float Alpha = (float)Index / (float)NumLaserSplinePoints;
+			Alpha = FMath::Sin(Alpha * PI * 0.5f);
+			const FVector PointOnStraightLaser = FMath::Lerp(InStartLocation, StraightLaserEndLocation, Alpha);
+			const FVector PointOnSmoothLaser = FMath::Lerp(InStartLocation, InEndLocation, Alpha);
+			const FVector PointBetweenLasers = FMath::Lerp(PointOnStraightLaser, PointOnSmoothLaser, Alpha);
+			LaserSplineComponent->AddSplinePoint(PointBetweenLasers, ESplineCoordinateSpace::World, false);
+		}
+		LaserSplineComponent->AddSplinePoint(InEndLocation, ESplineCoordinateSpace::World, false);
+
+		// Update all the segments of the spline
+		LaserSplineComponent->UpdateSpline();
+		
+		const float LaserPointerRadius = VREd::LaserPointerRadius->GetFloat() * VRMode->GetWorldScaleFactor();
+		Distance *= 0.0001f;
+		for (int32 Index = 0; Index < NumLaserSplinePoints; Index++)
+		{
+			USplineMeshComponent* SplineMeshComponent = LaserSplineMeshComponents[Index];
+			check (SplineMeshComponent != nullptr);
+			
+			FVector StartLoc, StartTangent, EndLoc, EndTangent;
+			LaserSplineComponent->GetLocationAndTangentAtSplinePoint(Index, StartLoc, StartTangent, ESplineCoordinateSpace::Local);
+			LaserSplineComponent->GetLocationAndTangentAtSplinePoint(Index + 1, EndLoc, EndTangent, ESplineCoordinateSpace::Local);
+
+			const float AlphaIndex = (float)Index / (float)NumLaserSplinePoints;
+			const float AlphaDistance = Distance * AlphaIndex;
+			float Radius = LaserPointerRadius * ((AlphaIndex * AlphaDistance) + 1);
+			FVector2D LaserScale(Radius, Radius); 
+			SplineMeshComponent->SetStartScale(LaserScale, false);
+			
+			const float NextAlphaIndex = (float)(Index + 1) / (float)NumLaserSplinePoints;
+			const float NextAlphaDistance = Distance * NextAlphaIndex;
+			Radius = LaserPointerRadius * ((NextAlphaIndex * NextAlphaDistance) + 1);
+			LaserScale = FVector2D(Radius, Radius);
+			SplineMeshComponent->SetEndScale(LaserScale, false);
+
+			SplineMeshComponent->SetStartAndEnd(StartLoc, StartTangent, EndLoc, EndTangent, true);
+		}
+	}
+}
+
+void UVREditorMotionControllerInteractor::SetLaserVisibility(const bool bVisible)
+{
+	for (USplineMeshComponent* SplineMeshComponent : LaserSplineMeshComponents)
+	{
+		SplineMeshComponent->SetVisibility(bVisible);
+	}
+}
+
 void UVREditorMotionControllerInteractor::SetLaserVisuals( const FLinearColor& NewColor, const float CrawlFade, const float CrawlSpeed )
 {
 	static FName StaticLaserColorParameterName( "LaserColor" );
@@ -1311,32 +1371,116 @@ void UVREditorMotionControllerInteractor::UpdateRadialMenuInput( const float Del
 	const EHMDDeviceType::Type HMDDeviceType = GetVRMode().GetHMDDeviceType();
 	//Update the radial menu
 	EViewportInteractionDraggingMode DraggingMode = GetDraggingMode();
-	if ( !HasUIInFront() &&
-		( bIsTrackpadPositionValid[ 0 ] && bIsTrackpadPositionValid[ 1 ] ) &&
+	if (ControllerType == EControllerType::UI)
+	{
+		if ((bIsTrackpadPositionValid[0] && bIsTrackpadPositionValid[1]) &&
+			DraggingMode != EViewportInteractionDraggingMode::AssistingDrag)
+		{
+			if (bIsScrubbingSequence)
+			{
+				const FVector2D ReturnToCenter = FVector2D::ZeroVector;
+				UISystem.GetRadialMenuFloatingUI()->HighlightSlot(ReturnToCenter);
+
+				const float NewPlayRate = FMath::GetMappedRangeValueClamped(FVector2D(-1.0f, 1.0f), FVector2D(-1.0f*VREd::SequencerScrubMax->GetFloat(), VREd::SequencerScrubMax->GetFloat()), TrackpadPosition.X);
+				FVREditorActionCallbacks::PlaySequenceAtRate(VRMode, NewPlayRate);
+			}
+			else
+			{
+				// Update the radial menu if we are already showing the radial menu
+				if (UISystem.IsShowingRadialMenu(this))
+				{
+					if (UISystem.GetRadialMenuFloatingUI()->GetWidgetComponents().Num() > 0)
+					{
+						UISystem.GetRadialMenuFloatingUI()->HighlightSlot(TrackpadPosition);
+
+						if (TrackpadPosition.GetAbsMax() > VREd::MinJoystickOffsetBeforeFlick->GetFloat())
+						{
+							LastActiveTrackpadUpdateTime = FTimespan::FromSeconds(FPlatformTime::Seconds());
+						}
+					}
+				}
+				else if (!UISystem.IsShowingRadialMenu(this) && 
+					TrackpadPosition.GetAbsMax() > VREd::MinJoystickOffsetBeforeFlick->GetFloat())
+				{
+					const bool bForceRefresh = false;
+					UISystem.TryToSpawnRadialMenu(this, bForceRefresh);
+					LastActiveTrackpadUpdateTime = FTimespan::FromSeconds(FPlatformTime::Seconds());
+				}
+			}
+		}
+		// If we are not currently touching the Vive touchpad, reset the highlighted button
+		else if (HMDDeviceType == EHMDDeviceType::DT_SteamVR && !bIsTouchingTrackpad)
+		{
+			if (UISystem.IsShowingRadialMenu(this))
+			{
+				const FVector2D ReturnToCenter = FVector2D::ZeroVector;
+				UISystem.GetRadialMenuFloatingUI()->HighlightSlot(ReturnToCenter);
+			}
+		}
+	}
+
+	else if (ControllerType == EControllerType::Laser)
+	{
+		if (VRMode->GetHMDDeviceType() != EHMDDeviceType::DT_SteamVR &&
+			(bIsTrackpadPositionValid[0] && bIsTrackpadPositionValid[1]) &&
+			DraggingMode != EViewportInteractionDraggingMode::TransformablesWithGizmo &&
+			DraggingMode != EViewportInteractionDraggingMode::TransformablesFreely &&
+			DraggingMode != EViewportInteractionDraggingMode::TransformablesAtLaserImpact &&
+			DraggingMode != EViewportInteractionDraggingMode::AssistingDrag)
+		{
+			// Move thumbstick left to undo
+			if (TrackpadPosition.X < -1 * VREd::MinJoystickOffsetBeforeFlick->GetFloat() && 
+				bFlickActionExecuted == false && 
+				!IsHoveringOverUI())
+			{
+				VRMode->GetWorldInteraction().Undo();
+				bFlickActionExecuted = true;
+			}
+			// Move thumbstick right to redo
+			if (TrackpadPosition.X > VREd::MinJoystickOffsetBeforeFlick->GetFloat() &&
+				bFlickActionExecuted == false &&
+				!IsHoveringOverUI())
+			{
+				VRMode->GetWorldInteraction().Redo();
+				bFlickActionExecuted = true;
+			}
+			// Center to reset
+			// TODO: Remove finger from touchpad to reset vive
+			if (FMath::IsNearlyZero(TrackpadPosition.X) &&
+				!IsHoveringOverUI())
+			{
+				bFlickActionExecuted = false;
+			}
+		}
+	}
+}
+
+void UVREditorMotionControllerInteractor::UndoRedoFromSwipe(const ETouchSwipeDirection InSwipeDirection)
+{
+	EViewportInteractionDraggingMode DraggingMode = GetDraggingMode();
+	if (ControllerType == EControllerType::Laser &&
+		VRMode->GetHMDDeviceType() == EHMDDeviceType::DT_SteamVR && 
 		DraggingMode != EViewportInteractionDraggingMode::TransformablesWithGizmo &&
 		DraggingMode != EViewportInteractionDraggingMode::TransformablesFreely &&
 		DraggingMode != EViewportInteractionDraggingMode::TransformablesAtLaserImpact &&
-		DraggingMode != EViewportInteractionDraggingMode::AssistingDrag &&
-		!UISystem.IsInteractorDraggingDockUI( this ) &&
-		!UISystem.IsShowingRadialMenu( Cast<UVREditorInteractor>( OtherInteractor ) ) )
+		DraggingMode != EViewportInteractionDraggingMode::AssistingDrag)
 	{
-		// Update the radial menu if we are already showing the radial menu
-		if ( UISystem.IsShowingRadialMenu( this ) )
+		if (InSwipeDirection == ETouchSwipeDirection::Left)
 		{
-			const TSharedPtr<SRadialBox>& RadialBox = StaticCastSharedPtr<SRadialBox>(UISystem.GetRadialWidget());
-			RadialBox->HighlightSlot(TrackpadPosition);
+			VRMode->GetWorldInteraction().Undo();
+			bFlickActionExecuted = true;
 		}
-	}
-	// If we are not currently touching the Vive touchpad, reset the highlighted button
-	else if (HMDDeviceType == EHMDDeviceType::DT_SteamVR && !bIsTouchingTrackpad)
-	{
-		if (UISystem.IsShowingRadialMenu( this ))
+		else if (InSwipeDirection == ETouchSwipeDirection::Right)
 		{
-			const TSharedPtr<SRadialBox>& RadialBox = StaticCastSharedPtr<SRadialBox>( UISystem.GetRadialWidget() );
-			FVector2D ReturnToCenter = FVector2D::ZeroVector;
-			RadialBox->HighlightSlot( ReturnToCenter );
+			VRMode->GetWorldInteraction().Redo();
+			bFlickActionExecuted = true;
 		}
 	}
 }
 
+bool UVREditorMotionControllerInteractor::GetIsLaserBlocked() const
+{
+	return Super::GetIsLaserBlocked() || (ControllerType != EControllerType::Laser && ControllerType != EControllerType::AssistingLaser);
+}
+
 #undef LOCTEXT_NAMESPACE
diff --git a/Engine/Source/Editor/VREditor/VREditorMotionControllerInteractor.h b/Engine/Source/Editor/VREditor/VREditorMotionControllerInteractor.h
--- a/Engine/Source/Editor/VREditor/VREditorMotionControllerInteractor.h
+++ b/Engine/Source/Editor/VREditor/VREditorMotionControllerInteractor.h
@@ -39,24 +39,27 @@ public:
 	// IViewportInteractorInterface overrides
 	virtual void Shutdown() override;
 	virtual void Tick( const float DeltaTime ) override;
-	virtual void CalculateDragRay() override;
+	virtual void CalculateDragRay( float& InOutDragRayLength, float& InOutDragRayVelocity ) override;
 
 	/** @return Returns the type of HMD we're dealing with */
 	EHMDDeviceType::Type GetHMDDeviceType() const;
 
 	// UViewportInteractor
+	virtual void PreviewInputKey( class FEditorViewportClient& ViewportClient, FViewportActionKeyInput& Action, const FKey Key, const EInputEvent Event, bool& bOutWasHandled ) override;
 	virtual void HandleInputKey( class FEditorViewportClient& ViewportClient, FViewportActionKeyInput& Action, const FKey Key, const EInputEvent Event, bool& bOutWasHandled ) override;
 	virtual bool GetTransformAndForwardVector( FTransform& OutHandTransform, FVector& OutForwardVector ) const override;
+	virtual bool GetIsLaserBlocked() const override;
 
-	/** Gets the Y delta of the trackpad */
-	float GetTrackpadSlideDelta();
+	/** 
+	 * Gets the trackpad delta of the axis passed. 
+	 * 
+	 * @param Axis	The axis of which we want the slide delta. 0 is X axis and 1 is Y axis.Default is axis Y
+	 */
+	float GetTrackpadSlideDelta( const bool Axis = 1 );
 
 	/** Starts haptic feedback for physical motion controller */
 	virtual void PlayHapticEffect( const float Strength ) override;
 	
-	/** Stops playing any haptic effects that have been going for a while.  Called every frame. */
-	void StopOldHapticEffects();
-	
 	/** Get the side of the controller */
 	EControllerHand GetControllerSide() const;
 
@@ -78,6 +81,27 @@ public:
 	/** Get when the last time the trackpad position was updated */
 	FTimespan& GetLastTrackpadPositionUpdateTime();
 
+	/** Get when the last time the trackpad position was updated */
+	FTimespan& GetLastActiveTrackpadUpdateTime();
+
+	/** Set if we want to force to show the laser*/
+	void SetForceShowLaser(const bool bInForceShow);
+
+	/** Next frame this will be used as color for the laser */
+	void SetForceLaserColor(const FLinearColor& InColor);
+
+	/** Toggles whether or not this controller is being used to scrub sequencer */
+	void ToggleSequencerScrubbingMode()
+	{
+		bIsScrubbingSequence = !bIsScrubbingSequence;
+	};
+
+	/** Returns whether or not this controller is being used to scrub sequencer */
+	bool IsScrubbingSequencer()
+	{
+		return bIsScrubbingSequence;
+	}
+
 protected:
 
 	// ViewportInteractor
@@ -86,18 +110,6 @@ protected:
 	/** Polls input for the motion controllers transforms */
 	virtual void PollInput() override;
 
-	/** Check if the trigger is lightly pressed */
-	bool IsTriggerAtLeastLightlyPressed() const;
-
-	/** Get the real time the trigger was lightly pressed */
-	double GetRealTimeTriggerWasLightlyPressed() const;
-
-	/** Set that the trigger is at least lighty pressed */
-	void SetTriggerAtLeastLightlyPressed( const bool bInTriggerAtLeastLightlyPressed );
-
-	/** Set that the trigger has been released since last press */
-	void SetTriggerBeenReleasedSinceLastPress( const bool bInTriggerBeenReleasedSinceLastPress );
-
 private:
 
 	/** Changes the color of the buttons on the handmesh */
@@ -115,12 +127,31 @@ private:
 	/** Given a mesh and a key name, tries to find a socket on the mesh that matches a supported key */
 	UStaticMeshSocket* FindMeshSocketForKey( UStaticMesh* StaticMesh, const FKey Key );
 
+	/** Updates all the segments of the curved laser */
+	void UpdateSplineLaser(const FVector& InStartLocation, const FVector& InEndLocation, const FVector& InForward);
+	
+	/** Sets the visibility on all curved laser segments */
+	void SetLaserVisibility(const bool bVisible);
+
 	/** Sets the visuals of the LaserPointer */
 	void SetLaserVisuals( const FLinearColor& NewColor, const float CrawlFade, const float CrawlSpeed );
 
 	/** Updates the radial menu */
 	void UpdateRadialMenuInput( const float DeltaTime );
 
+	/** Directions the trackpad can be swiped to */
+	enum ETouchSwipeDirection
+	{
+		None = 0,
+		Left = 1,
+		Right = 2,
+		Up = 3,
+		Down = 4
+	};
+
+	/** Start undo or redo from swipe for the Vive */
+	void UndoRedoFromSwipe(const ETouchSwipeDirection InSwipeDirection);
+
 protected:
 	
 	/** Motion controller component which handles late-frame transform updates of all parented sub-components */
@@ -135,9 +166,13 @@ protected:
 	UPROPERTY()
 	class UStaticMeshComponent* HandMeshComponent;
 
-	/** Mesh for this hand's laser pointer */
+	/** Spline for this hand's laser pointer */
+	UPROPERTY()
+	class USplineComponent* LaserSplineComponent;
+
+	/** Spline meshes for curved laser */
 	UPROPERTY()
-	class UStaticMeshComponent* LaserPointerMeshComponent;
+	TArray<class USplineMeshComponent*> LaserSplineMeshComponents;
 
 	/** MID for laser pointer material (opaque parts) */
 	UPROPERTY()
@@ -171,8 +206,8 @@ protected:
 	static const FName TriggerAxis;
 	static const FName MotionController_Left_FullyPressedTriggerAxis;
 	static const FName MotionController_Right_FullyPressedTriggerAxis;
-	static const FName MotionController_Left_LightlyPressedTriggerAxis;
-	static const FName MotionController_Right_LightlyPressedTriggerAxis;
+	static const FName MotionController_Left_PressedTriggerAxis;
+	static const FName MotionController_Right_PressedTriggerAxis;
 
 private:
 
@@ -184,22 +219,11 @@ private:
 	bool bIsTriggerFullyPressed;
 
 	/** True if the trigger is currently pulled far enough that we consider it in a "half pressed" state */
-	bool bIsTriggerAtLeastLightlyPressed;
-
-	/** Real time that the trigger became lightly pressed.  If the trigger remains lightly pressed for a reasonably 
-	long duration before being pressed fully, we'll continue to treat it as a light press in some cases */
-	double RealTimeTriggerWasLightlyPressed;
+	bool bIsTriggerPressed;
 
 	/** True if trigger has been fully released since the last press */
 	bool bHasTriggerBeenReleasedSinceLastPress;
 
-	//
-	// Haptic feedback
-	//
-
-	/** The last real time we played a haptic effect.  This is used to turn off haptic effects shortly after they're triggered. */
-	double LastHapticTime;
-
 	//
 	// Trackpad support
 	//
@@ -207,6 +231,9 @@ private:
 	/** True if the trackpad is actively being touched */
 	bool bIsTouchingTrackpad;
 
+	/** True if pressing trackpad button (or analog stick button is down) */
+	bool bIsPressingTrackpad;
+
 	/** Position of the touched trackpad */
 	FVector2D TrackpadPosition;
 
@@ -219,4 +246,28 @@ private:
 	/** Real time that the last trackpad position was last updated.  Used to filter out stale previous data. */
 	FTimespan LastTrackpadPositionUpdateTime;
 
+	/** Real time that the last trackpad position was over the dead zone threshold. */
+	FTimespan LastActiveTrackpadUpdateTime;
+
+	/** Forcing to show laser */
+	bool bForceShowLaser;
+
+	/** The color that will be used for one frame */
+	TOptional<FLinearColor> ForceLaserColor;
+
+	/**Whether a flick action was executed */
+	bool bFlickActionExecuted;
+
+	/** whether or not this controller is being used to scrub sequencer */
+	bool bIsScrubbingSequence;
+
+	//
+	// Swipe
+	//
+
+	/** Latest swipe direction on the trackpad */
+	ETouchSwipeDirection LastSwipe;
+
+	/** Initial position when starting to touch the trackpad */
+	FVector2D InitialTouchPosition;
 };
diff --git a/Engine/Source/Editor/VREditor/VREditorStyle.cpp b/Engine/Source/Editor/VREditor/VREditorStyle.cpp
--- a/Engine/Source/Editor/VREditor/VREditorStyle.cpp
+++ b/Engine/Source/Editor/VREditor/VREditorStyle.cpp
@@ -35,6 +35,18 @@ FName FVREditorStyle::GetStyleSetName()
 	return StyleSetName;
 }
 
+FName FVREditorStyle::GetSecondaryStyleSetName()
+{
+	static FName StyleSetName(TEXT("VRRadialStyle"));
+	return StyleSetName;
+}
+
+FName FVREditorStyle::GetNumpadStyleSetName()
+{
+	static FName StyleSetName(TEXT("VRNumpadRadialStyle"));
+	return StyleSetName;
+}
+
 #define IMAGE_BRUSH( RelativePath, ... ) FSlateImageBrush( Style->RootToContentDir( RelativePath, TEXT(".png") ), __VA_ARGS__ )
 #define BOX_BRUSH( RelativePath, ... ) FSlateBoxBrush( Style->RootToContentDir( RelativePath, TEXT(".png") ), __VA_ARGS__ )
 #define BORDER_BRUSH( RelativePath, ... ) FSlateBorderBrush( Style->RootToContentDir( RelativePath, TEXT(".png") ), __VA_ARGS__ )
@@ -42,27 +54,36 @@ FName FVREditorStyle::GetStyleSetName()
 #define OTF_FONT( RelativePath, ... ) FSlateFontInfo( Style->RootToContentDir( RelativePath, TEXT(".otf") ), __VA_ARGS__ )
 #define TTF_CORE_FONT(RelativePath, ...) FSlateFontInfo(Style->RootToCoreContentDir(RelativePath, TEXT(".ttf") ), __VA_ARGS__)
 
+const FVector2D Icon14x14(14.0f, 14.0f);
 const FVector2D Icon16x16(16.0f, 16.0f);
 const FVector2D Icon20x20(20.0f, 20.0f);
 const FVector2D Icon40x40(40.0f, 40.0f);
+const FVector2D Icon64x64(64.0f, 64.0f);
+const FVector2D Icon512x512(512.0f, 512.0f);
 
 TSharedRef< FSlateStyleSet > FVREditorStyle::Create()
 {
 	TSharedRef< FSlateStyleSet > Style = MakeShareable(new FSlateStyleSet(FVREditorStyle::GetStyleSetName()));
+
+	Style->SetContentRoot(FPaths::EngineContentDir() / TEXT("Editor/Slate"));
+	Style->SetCoreContentRoot(FPaths::EngineContentDir() / TEXT("Slate"));
+
 	// Use the default menu button style, but set the background to dark grey.
 	const FButtonStyle NormalButton = FEditorStyle::GetWidgetStyle<FButtonStyle>("Menu.Button");
 	Style->Set("VREditorStyle.Button", FButtonStyle(NormalButton)
-		.SetNormal(BOX_BRUSH("Common/RoundedSelection_16x", 4.0f / 16.0f, FLinearColor(0.1f,0.1f,0.1f))));
+		.SetNormal(BOX_BRUSH("Common/RoundedSelection_16x", 4.0f / 16.0f, FLinearColor(0.1f, 0.1f, 0.1f))));
+	Style->Set("VREditorStyle.CollapsedButton", FButtonStyle(NormalButton)
+		.SetNormal(BOX_BRUSH("Common/RoundedSelection_16x", 4.0f / 16.0f, FLinearColor(0.1f, 0.1f, 0.1f))));
 
 	const FTextBlockStyle NormalText = FEditorStyle::GetWidgetStyle<FTextBlockStyle>("NormalText");
 	Style->Set("VREditorStyle.Label", FTextBlockStyle(NormalText)
-		.SetFont(TTF_CORE_FONT("Slate/Fonts/Roboto-Regular", 9)));
+		.SetFont(TTF_CORE_FONT("Fonts/Roboto-Regular", 7)));
 
 	// Headings will have a font outline
 	FFontOutlineSettings HeadingOutline;
 	HeadingOutline.OutlineColor = FLinearColor(0.2f, 0.2f, 0.2f, 0.5f);
-	HeadingOutline.OutlineSize = 2.0f;
-	FSlateFontInfo HeadlineFont = TTF_CORE_FONT("Slate/Fonts/Roboto-Regular", 18);
+	HeadingOutline.OutlineSize = 1.5f;
+	FSlateFontInfo HeadlineFont = TTF_CORE_FONT("Fonts/Roboto-Regular", 10);
 	HeadlineFont.OutlineSettings = HeadingOutline;
 
 	Style->Set("VREditorStyle.Heading", FTextBlockStyle(NormalText)
@@ -73,38 +94,152 @@ TSharedRef< FSlateStyleSet > FVREditorStyle::Create()
 	FFontOutlineSettings HelperOutline;
 	HelperOutline.OutlineColor = FLinearColor( 0.2f, 0.2f, 0.2f, 0.5f );
 	HelperOutline.OutlineSize = 3.0f;
-	FSlateFontInfo HelperFont = TTF_CORE_FONT( "Slate/Fonts/Roboto-Regular", 24 );
+	FSlateFontInfo HelperFont = TTF_CORE_FONT( "Fonts/Roboto-Regular", 24 );
 	HelperFont.OutlineSettings = HelperOutline;
 
 	Style->Set( "VREditorStyle.HelperText", FTextBlockStyle( NormalText )
 		.SetFont( HelperFont )
 		.SetColorAndOpacity( FLinearColor::White ) );
-	
-	const FCheckBoxStyle CheckStyle = FEditorStyle::GetWidgetStyle<FCheckBoxStyle>("Menu.Check");
-	Style->Set("VREditorStyle.Check", FCheckBoxStyle(CheckStyle));
 
-	const FCheckBoxStyle CheckboxStyle = FEditorStyle::GetWidgetStyle<FCheckBoxStyle>("Menu.CheckBox");
-	Style->Set("VREditorStyle.CheckBox", FCheckBoxStyle(CheckboxStyle));
+	FCheckBoxStyle VRMenuCheckBoxStyle = FCheckBoxStyle()
+		.SetCheckBoxType(ESlateCheckBoxType::ToggleButton)
+		.SetUncheckedImage(IMAGE_BRUSH("Icons/VREditor/T_Radial_Checkbox", Icon64x64, FLinearColor::Transparent))
+		.SetUncheckedHoveredImage(IMAGE_BRUSH("Icons/VREditor/T_Radial_Checkbox", Icon64x64, FLinearColor::Transparent))
+		.SetUncheckedPressedImage(IMAGE_BRUSH("Icons/VREditor/T_Radial_Checkbox", Icon64x64, FLinearColor::Transparent))
+		.SetCheckedImage(IMAGE_BRUSH("Icons/VREditor/T_Radial_Checkbox", Icon64x64, FLinearColor::White))
+		.SetCheckedHoveredImage(IMAGE_BRUSH("Icons/VREditor/T_Radial_Checkbox", Icon64x64, FLinearColor::White))
+		.SetCheckedPressedImage(IMAGE_BRUSH("Icons/VREditor/T_Radial_Checkbox", Icon64x64, FLinearColor::White));
+
+	Style->Set("VREditorStyle.Check", VRMenuCheckBoxStyle);
+	Style->Set("VRRadialStyle.Check", VRMenuCheckBoxStyle);
+
+	Style->Set("VREditorStyle.CheckBox", VRMenuCheckBoxStyle);
+	Style->Set("VRRadialStyle.CheckBox", VRMenuCheckBoxStyle);
 
 	const FCheckBoxStyle RadioButtonStyle = FEditorStyle::GetWidgetStyle<FCheckBoxStyle>("Menu.RadioButton");
 	Style->Set("VREditorStyle.RadioButton", FCheckBoxStyle(RadioButtonStyle));
+	Style->Set("VRRadialStyle.RadioButton", FCheckBoxStyle(RadioButtonStyle));
+
+	const FCheckBoxStyle ToggleButton = FEditorStyle::GetWidgetStyle<FCheckBoxStyle>("Menu.ToggleButton");
+	Style->Set("VREditorStyle.ToggleButton", FCheckBoxStyle(ToggleButton));
+	Style->Set("VRRadialStyle.ToggleButton", FCheckBoxStyle(ToggleButton));
 
 	const FTextBlockStyle KeybindingStyle = FEditorStyle::GetWidgetStyle<FTextBlockStyle>("Menu.Keybinding");
 	Style->Set("VREditorStyle.Keybinding", FTextBlockStyle(KeybindingStyle));
+	Style->Set("VRRadialStyle.Keybinding", FTextBlockStyle(KeybindingStyle));
+
+	Style->Set("VREditorStyle.AlignActors", new IMAGE_BRUSH("Icons/UMG/Alignment/Horizontal_Left", Icon16x16));
+
+	Style->Set("VRRadialStyle.Button", FButtonStyle(NormalButton)
+		.SetHovered(FSlateNoResource())
+		.SetPressed(FSlateNoResource()));
+	Style->Set("VRRadialStyle.CollapsedButton", FButtonStyle(NormalButton)
+		.SetHovered(FSlateNoResource())
+		.SetPressed(FSlateNoResource()));
+
+	FFontOutlineSettings RadialOutline;
+	RadialOutline.OutlineColor = FLinearColor(0.1f, 0.1f, 0.1f, 1.0f);
+	RadialOutline.OutlineSize = 1.0f;
+
+	FSlateFontInfo RadialFont = TTF_CORE_FONT("Fonts/Roboto-Regular", 8);
+	RadialFont.OutlineSettings = RadialOutline;
+	Style->Set("VRRadialStyle.Label", FTextBlockStyle(NormalText)
+		.SetFont(RadialFont)
+		.SetColorAndOpacity(FLinearColor::White));
+
+	Style->Set("VRRadialStyle.InactiveFont", FSlateFontInfo(RadialFont));
+
+	FSlateFontInfo ActiveRadialFont = TTF_CORE_FONT("Fonts/Roboto-Regular", 10);
+	FFontOutlineSettings ActiveRadialOutline;
+	ActiveRadialOutline.OutlineColor = FLinearColor::Black;
+	ActiveRadialOutline.OutlineSize = 1.0f;
+	ActiveRadialFont.OutlineSettings = ActiveRadialOutline;
+	Style->Set("VRRadialStyle.ActiveFont", FSlateFontInfo(ActiveRadialFont));
+	
+	FSlateFontInfo NumpadRadialFont = TTF_CORE_FONT("Fonts/Roboto-Regular", 24);
+	NumpadRadialFont.OutlineSettings = RadialOutline;
+	Style->Set("VRNumpadRadialStyle.Label", FTextBlockStyle(NormalText)
+		.SetFont(NumpadRadialFont)
+		.SetColorAndOpacity(FLinearColor::White));
+	Style->Set("VRNumpadRadialStyle.Button", FButtonStyle(NormalButton)
+		.SetHovered(FSlateNoResource())
+		.SetPressed(FSlateNoResource()));
+
+	Style->Set("VREditorStyle.EditMenu", new IMAGE_BRUSH("Icons/VREditor/T_Radial_Edit", Icon512x512));
+	Style->Set("VREditorStyle.SnapMenu", new IMAGE_BRUSH("Icons/VREditor/T_Radial_Snapping", Icon512x512));
+	Style->Set("VREditorStyle.GizmoMenu", new IMAGE_BRUSH("Icons/VREditor/T_Radial_Gizmo", Icon512x512));
+	Style->Set("VREditorStyle.ModesMenu", new IMAGE_BRUSH("Icons/VREditor/T_Radial_Modes", Icon512x512));
+	Style->Set("VREditorStyle.ToolsMenu", new IMAGE_BRUSH("Icons/VREditor/T_Radial_Tools", Icon512x512));
+	Style->Set("VREditorStyle.WindowsMenu", new IMAGE_BRUSH("Icons/VREditor/T_Radial_Windows", Icon512x512));
+	Style->Set("VREditorStyle.ActionsMenu", new IMAGE_BRUSH("Icons/VREditor/T_Radial_Actions", Icon512x512));
+	
+	Style->Set("VREditorStyle.ActorsMode", new IMAGE_BRUSH("Icons/VREditor/T_Radial_Actors", Icon512x512));
+	Style->Set("VREditorStyle.FoliageMode", new IMAGE_BRUSH("Icons/VREditor/T_Radial_Foliage", Icon512x512));
+	Style->Set("VREditorStyle.LandscapeMode", new IMAGE_BRUSH("Icons/VREditor/T_Radial_Landscape", Icon512x512));
+	Style->Set("VREditorStyle.MeshPaintMode", new IMAGE_BRUSH("Icons/VREditor/T_Radial_Mesh_Paint", Icon512x512));
+	Style->Set("VREditorStyle.MeshEditMode", new IMAGE_BRUSH("Icons/VREditor/T_Radial_Mesh_All", Icon512x512));
+	Style->Set("VREditorStyle.PolygonEditMode", new IMAGE_BRUSH("Icons/VREditor/T_Radial_Mesh_Poly", Icon512x512));
+	Style->Set("VREditorStyle.EdgeEditMode", new IMAGE_BRUSH("Icons/VREditor/T_Radial_Mesh_Edge", Icon512x512));
+	Style->Set("VREditorStyle.VertexEditMode", new IMAGE_BRUSH("Icons/VREditor/T_Radial_Mesh_Vertex", Icon512x512));
+
+	Style->Set("VREditorStyle.Copy", new IMAGE_BRUSH("Icons/VREditor/T_Radial_Copy", Icon512x512));
+	Style->Set("VREditorStyle.Cut", new IMAGE_BRUSH("Icons/VREditor/T_Radial_Cut", Icon512x512));
+	Style->Set("VREditorStyle.Paste", new IMAGE_BRUSH("Icons/VREditor/T_Radial_Paste", Icon512x512));
+	Style->Set("VREditorStyle.Delete", new IMAGE_BRUSH("Icons/VREditor/T_Radial_Delete", Icon512x512));
+	Style->Set("VREditorStyle.Duplicate", new IMAGE_BRUSH("Icons/VREditor/T_Radial_Duplicate", Icon512x512));
+	Style->Set("VREditorStyle.SnapToFloor", new IMAGE_BRUSH("Icons/VREditor/T_Radial_Snaps", Icon512x512));
+	Style->Set("VREditorStyle.DeselectAll", new IMAGE_BRUSH("Icons/VREditor/T_Radial_Deselect_All", Icon512x512));
+
+	Style->Set("VREditorStyle.Flashlight", new IMAGE_BRUSH("Icons/VREditor/T_Radial_Flashlight", Icon512x512));
+	Style->Set("VREditorStyle.Screenshot", new IMAGE_BRUSH("Icons/VREditor/T_Radial_Screenshot", Icon512x512));
+	Style->Set("VREditorStyle.Simulate", new IMAGE_BRUSH("Icons/VREditor/T_Radial_Simulate", Icon512x512));
+	Style->Set("VREditorStyle.Pause", new IMAGE_BRUSH("Icons/VREditor/T_Radial_Pause", Icon512x512));
+	Style->Set("VREditorStyle.Play", new IMAGE_BRUSH("Icons/VREditor/T_Radial_Play", Icon512x512));
+	Style->Set("VREditorStyle.Resume", new IMAGE_BRUSH("Icons/VREditor/T_Radial_Resume", Icon512x512));
+	Style->Set("VREditorStyle.SaveSimulation", new IMAGE_BRUSH("Icons/VREditor/T_Radial_Save_Actors", Icon512x512));
 
-	const FTextBlockStyle LabelStyle = FEditorStyle::GetWidgetStyle<FTextBlockStyle>("Menu.Label");
-	Style->Set("VREditorStyle.Label", FTextBlockStyle(LabelStyle));
+	Style->Set("VREditorStyle.Translate", new IMAGE_BRUSH("Icons/VREditor/T_Radial_Translate", Icon512x512));
+	Style->Set("VREditorStyle.Rotate", new IMAGE_BRUSH("Icons/VREditor/T_Radial_Rotate", Icon512x512));
+	Style->Set("VREditorStyle.Scale", new IMAGE_BRUSH("Icons/VREditor/T_Radial_Scale", Icon512x512));
+	Style->Set("VREditorStyle.Universal", new IMAGE_BRUSH("Icons/VREditor/T_Radial_Universal", Icon512x512));
+	Style->Set("VREditorStyle.WorldSpace", new IMAGE_BRUSH("Icons/VREditor/T_Radial_World_Space", Icon512x512));
+	Style->Set("VREditorStyle.LocalSpace", new IMAGE_BRUSH("Icons/VREditor/T_Radial_Local_Space", Icon512x512));
+
+	Style->Set("VREditorStyle.AlignActors", new IMAGE_BRUSH("Icons/VREditor/T_Radial_Align_Actors", Icon512x512));
+	Style->Set("VREditorStyle.SetTargets", new IMAGE_BRUSH("Icons/VREditor/T_Radial_Set_Targets", Icon512x512));
+	Style->Set("VREditorStyle.GridNum", new IMAGE_BRUSH("Icons/VREditor/T_Radial_Grid_Num", Icon512x512));
+	Style->Set("VREditorStyle.TranslateSnap", new IMAGE_BRUSH("Icons/VREditor/T_Radial_Translate_Snap", Icon512x512));
+	Style->Set("VREditorStyle.AngleNum", new IMAGE_BRUSH("Icons/VREditor/T_Radial_Angle_Num", Icon512x512));
+	Style->Set("VREditorStyle.RotateSnap", new IMAGE_BRUSH("Icons/VREditor/T_Radial_Rotate_Snap", Icon512x512));
+	Style->Set("VREditorStyle.ScaleNum", new IMAGE_BRUSH("Icons/VREditor/T_Radial_Scale_Num", Icon512x512));
+	Style->Set("VREditorStyle.ScaleSnap", new IMAGE_BRUSH("Icons/VREditor/T_Radial_Scale_Snap", Icon512x512));
+
+	Style->Set("VREditorStyle.ContentBrowser", new IMAGE_BRUSH("Icons/VREditor/T_Radial_Content_Browser", Icon512x512));
+	Style->Set("VREditorStyle.Details", new IMAGE_BRUSH("Icons/VREditor/T_Radial_Details", Icon512x512));
+	Style->Set("VREditorStyle.ModesPanel", new IMAGE_BRUSH("Icons/VREditor/T_Radial_Modes_Panel", Icon512x512));
+	Style->Set("VREditorStyle.Sequencer", new IMAGE_BRUSH("Icons/VREditor/T_Radial_Sequencer", Icon512x512));
+	Style->Set("VREditorStyle.WorldOutliner", new IMAGE_BRUSH("Icons/VREditor/T_Radial_World_Outliner", Icon512x512));
+	Style->Set("VREditorStyle.WorldSettings", new IMAGE_BRUSH("Icons/VREditor/T_Radial_World_Settings", Icon512x512));
+
+	Style->Set("VREditorStyle.SequencerPlay", new IMAGE_BRUSH("Icons/VREditor/T_Radial_Seq_Play", Icon512x512));
+	Style->Set("VREditorStyle.SequencerStop", new IMAGE_BRUSH("Icons/VREditor/T_Radial_Seq_Stop", Icon512x512));
+	Style->Set("VREditorStyle.SequencerReverse", new IMAGE_BRUSH("Icons/VREditor/T_Radial_Seq_Reverse", Icon512x512));
+	Style->Set("VREditorStyle.Scrub", new IMAGE_BRUSH("Icons/VREditor/T_Radial_Seq_Scrub", Icon512x512));
+	Style->Set("VREditorStyle.PlayFromStart", new IMAGE_BRUSH("Icons/VREditor/T_Radial_Seq_Start", Icon512x512));
+	Style->Set("VREditorStyle.ToggleLooping", new IMAGE_BRUSH("Icons/VREditor/T_Radial_Seq_Loop", Icon512x512));
+
+
+
+
+
+
+	Style->Set("VREditorStyle.Home", new IMAGE_BRUSH("Icons/VREditor/T_Radial_Home_v1", Icon512x512));
+	Style->Set("VREditorStyle.OneLevel", new IMAGE_BRUSH("Icons/VREditor/T_Radial_Home_v2", Icon512x512));
+	Style->Set("VREditorStyle.TwoLevel", new IMAGE_BRUSH("Icons/VREditor/T_Radial_Home_v3", Icon512x512));
 
 	return Style;
 }
 
-#undef IMAGE_BRUSH
-#undef BOX_BRUSH
-#undef BORDER_BRUSH
-#undef TTF_FONT
-#undef OTF_FONT
-#undef TTF_CORE_FONT
-
 void FVREditorStyle::ReloadTextures()
 {
 	FSlateApplication::Get().GetRenderer()->ReloadTextureResources();
@@ -114,3 +249,11 @@ const ISlateStyle& FVREditorStyle::Get()
 {
 	return *VREditorStyleInstance;
 }
+
+
+#undef IMAGE_BRUSH
+#undef BOX_BRUSH
+#undef BORDER_BRUSH
+#undef TTF_FONT
+#undef OTF_FONT
+#undef TTF_CORE_FONT
diff --git a/Engine/Source/Editor/VREditor/VREditorStyle.h b/Engine/Source/Editor/VREditor/VREditorStyle.h
--- a/Engine/Source/Editor/VREditor/VREditorStyle.h
+++ b/Engine/Source/Editor/VREditor/VREditorStyle.h
@@ -22,7 +22,14 @@ public:
 
 	static FName GetStyleSetName();
 
-	
+	static FName GetSecondaryStyleSetName();
+
+	static FName GetNumpadStyleSetName();
+
+	static const FSlateBrush* GetBrush(FName PropertyName, const ANSICHAR* Specifier = NULL)
+	{
+		return VREditorStyleInstance->GetBrush(PropertyName, Specifier);
+	}
 
 private:
 
diff --git a/Engine/Source/Editor/VREditor/VREditorWidgetComponent.cpp b/Engine/Source/Editor/VREditor/VREditorWidgetComponent.cpp
--- a/Engine/Source/Editor/VREditor/VREditorWidgetComponent.cpp
+++ b/Engine/Source/Editor/VREditor/VREditorWidgetComponent.cpp
@@ -8,6 +8,7 @@ UVREditorWidgetComponent::UVREditorWidgetComponent(const FObjectInitializer& Obj
 	bIsHovering = false;
 	DrawingPolicy = EVREditorWidgetDrawingPolicy::Always;
 	bHasEverDrawn = false;
+	bSelectable = false;
 }
 
 bool UVREditorWidgetComponent::ShouldDrawWidget() const
