Merging //UE4/Private-Geometry to Dev-VREditor (//UE4/Dev-VREditor) 3339870

diff --git a/Engine/Source/Editor/VREditor/VREditorWorldInteraction.cpp b/Engine/Source/Editor/VREditor/VREditorWorldInteraction.cpp
--- a/Engine/Source/Editor/VREditor/VREditorWorldInteraction.cpp
+++ b/Engine/Source/Editor/VREditor/VREditorWorldInteraction.cpp
@@ -7,6 +7,7 @@
 #include "Misc/CommandLine.h"
 #include "Modules/ModuleManager.h"
 #include "VREditorMode.h"
+#include "VREditorAssetContainer.h"
 #include "ViewportInteractionTypes.h"
 #include "Materials/MaterialInterface.h"
 #include "Engine/Texture.h"
@@ -14,13 +15,13 @@
 #include "ViewportWorldInteraction.h"
 #include "Engine/BrushBuilder.h"
 #include "Components/PrimitiveComponent.h"
-#include "Kismet/GameplayStatics.h"
 #include "Kismet2/ComponentEditorUtils.h"
 #include "Sound/SoundCue.h"
 #include "Editor.h"
 #include "VREditorUISystem.h"
 #include "VREditorFloatingUI.h"
 #include "VREditorInteractor.h"
+#include "VREditorMotionControllerInteractor.h"
 
 // For actor placement
 #include "ObjectTools.h"
@@ -37,7 +38,6 @@
 namespace VREd
 {
 	static FAutoConsoleVariable SizeOfActorsOverContentBrowserThumbnail( TEXT( "VREd.SizeOfActorsOverContentBrowserThumbnail" ), 6.0f, TEXT( "How large objects should be when rendered 'thumbnail size' over the Content Browser" ) );
-	static FAutoConsoleVariable OverrideSizeOfPlacedActors( TEXT( "VREd.OverrideSizeOfPlacedActors" ), 0.0f, TEXT( "If set to a value greater than zero, sets the size in cm of a placed actor (relative to your world's scale)" ) );
 	static FAutoConsoleVariable HoverHapticFeedbackStrength( TEXT( "VREd.HoverHapticFeedbackStrength" ), 0.1f, TEXT( "Default strength for haptic feedback when hovering" ) );
 	static FAutoConsoleVariable HoverHapticFeedbackTime( TEXT( "VREd.HoverHapticFeedbackTime" ), 0.2f, TEXT( "The minimum time between haptic feedback for hovering" ) );
 	static FAutoConsoleVariable PivotPointTransformGizmo( TEXT( "VREd.PivotPointTransformGizmo" ), 1, TEXT( "If the pivot point transform gizmo is used instead of the bounding box gizmo" ) );
@@ -48,13 +48,9 @@ namespace VREd
 UVREditorWorldInteraction::UVREditorWorldInteraction( const FObjectInitializer& Initializer ) : 
 	Super( Initializer ),
 	Owner( nullptr ),
-	DropMaterialOrMaterialSound( nullptr ),
 	FloatingUIAssetDraggedFrom( nullptr ),
 	PlacingMaterialOrTextureAsset( nullptr )
 {
-	// Load sounds
-	DropMaterialOrMaterialSound = LoadObject<USoundCue>( nullptr, TEXT( "/Engine/VREditor/Sounds/VR_grab_Cue" ) );
-	check( DropMaterialOrMaterialSound != nullptr );
 }
 
 void UVREditorWorldInteraction::Init( UVREditorMode* InOwner, UViewportWorldInteraction* InViewportWorldInteraction )
@@ -77,7 +73,6 @@ void UVREditorWorldInteraction::Shutdown()
 
 	PlacingMaterialOrTextureAsset = nullptr;
 	FloatingUIAssetDraggedFrom = nullptr;
-	DropMaterialOrMaterialSound = nullptr;
 	Owner = nullptr;
 }
 
@@ -85,21 +80,31 @@ void UVREditorWorldInteraction::StopDragging( UViewportInteractor* Interactor )
 {
 	EViewportInteractionDraggingMode InteractorDraggingMode = Interactor->GetDraggingMode();
 
-	// If we were placing a material, go ahead and do that now
 	if ( InteractorDraggingMode == EViewportInteractionDraggingMode::Material )
 	{
+		// If we were placing a material, go ahead and do that now
 		PlaceDraggedMaterialOrTexture( Interactor );
 	}	
-
-	// If we were placing something, bring the window back
-	if ( FloatingUIAssetDraggedFrom != nullptr && 
+	else if ( FloatingUIAssetDraggedFrom != nullptr && 
 		( InteractorDraggingMode == EViewportInteractionDraggingMode::TransformablesAtLaserImpact ||
 		  InteractorDraggingMode == EViewportInteractionDraggingMode::Material ) )
 	{
-		Owner->GetUISystem().ShowEditorUIPanel( FloatingUIAssetDraggedFrom, Cast<UVREditorInteractor>( Interactor->GetOtherInteractor() ), 
-			true, false, false, false );
+		// If we were placing something, bring the window back
+		const bool bShouldShow = true;
+		const bool bSpawnInFront = false;
+		const bool bDragFromOpen = false;
+		const bool bPlaySound = false;
+		Owner->GetUISystem().ShowEditorUIPanel(FloatingUIAssetDraggedFrom, Cast<UVREditorInteractor>(Interactor->GetOtherInteractor()), bShouldShow, bSpawnInFront, bDragFromOpen, bPlaySound);
 		FloatingUIAssetDraggedFrom = nullptr;
 	}
+	else if (Interactor->GetDraggingMode() == EViewportInteractionDraggingMode::TransformablesFreely)
+	{
+		UVREditorMotionControllerInteractor* MotionController = Cast<UVREditorMotionControllerInteractor>(Interactor);
+		if (Owner->GetUISystem().GetUIInteractor() == MotionController)
+		{
+			MotionController->SetControllerType(EControllerType::UI);
+		}
+	}
 }
 
 void UVREditorWorldInteraction::UpdateNearClipPlaneOnScaleChange(const float NewWorldToMetersScale)
@@ -127,12 +132,14 @@ void UVREditorWorldInteraction::StartDraggingMaterialOrTexture( UViewportInterac
 		FViewportInteractorData& OtherInteractorData = Interactor->GetOtherInteractor()->GetInteractorData();
 		OtherInteractorData.bWasAssistingDrag = false;
 
+		InteractorData.bDraggingWithGrabberSphere = false;
 		InteractorData.bIsFirstDragUpdate = true;
 		InteractorData.bWasAssistingDrag = false;
 		InteractorData.DragRayLength = ( HitLocation - LaserPointerStart ).Size();
-		InteractorData.LastLaserPointerStart = LaserPointerStart;
 		InteractorData.LastDragToLocation = HitLocation;
-		InteractorData.LaserPointerImpactAtDragStart = HitLocation;
+		InteractorData.InteractorRotationAtDragStart = InteractorData.Transform.GetRotation();
+		InteractorData.GrabberSphereLocationAtDragStart = FVector::ZeroVector;
+		InteractorData.ImpactLocationAtDragStart = HitLocation;
 		InteractorData.DragTranslationVelocity = FVector::ZeroVector;
 		InteractorData.DragRayLengthVelocity = 0.0f;
 		InteractorData.bIsDrivingVelocityOfSimulatedTransformables = false;
@@ -141,6 +148,7 @@ void UVREditorWorldInteraction::StartDraggingMaterialOrTexture( UViewportInterac
 	
 		InteractorData.TransformGizmoInteractionType = ETransformGizmoInteractionType::None;
 		InteractorData.GizmoStartTransform = FTransform::Identity;
+		InteractorData.GizmoLastTransform = InteractorData.GizmoTargetTransform = InteractorData.GizmoUnsnappedTargetTransform = InteractorData.GizmoInterpolationSnapshotTransform = InteractorData.GizmoStartTransform;
 		InteractorData.GizmoStartLocalBounds = FBox( 0 );
 
 		InteractorData.GizmoSpaceFirstDragUpdateOffsetAlongAxis = FVector::ZeroVector;	// Will be determined on first update
@@ -156,74 +164,95 @@ void UVREditorWorldInteraction::StartDraggingMaterialOrTexture( UViewportInterac
 
 void UVREditorWorldInteraction::OnAssetDragStartedFromContentBrowser( const TArray<FAssetData>& DraggedAssets, UActorFactory* FactoryToUse )
 {
-	if( ViewportWorldInteraction == nullptr )
+	FloatingUIAssetDraggedFrom = nullptr;
+	if( ViewportWorldInteraction != nullptr )
 	{
-		return;
-	}
+		// Figure out which controller pressed the button and started dragging
+		// @todo vreditor placement: This logic could misfire.  Ideally we would be routed information from the pointer event, so we can determine the hand.
+		UVREditorInteractor* PlacingWithInteractor = nullptr;
 
-	const bool bIsPreview = false;	 // @todo vreditor placement: Consider supporting a "drop preview" actor (so you can cancel placement interactively)
+		TArray<UViewportInteractor*>& Interactors = ViewportWorldInteraction->GetInteractors();
+		for( UViewportInteractor* Interactor : Interactors )
+		{
+			UVREditorInteractor* VRInteractor = Cast<UVREditorInteractor>( Interactor );
+			FViewportActionKeyInput* SelectAndMove_Action = Interactor->GetActionWithName( ViewportWorldActionTypes::SelectAndMove );
+			if( VRInteractor && SelectAndMove_Action != nullptr && SelectAndMove_Action->bIsInputCaptured )
+			{
+				if( VRInteractor->IsClickingOnUI() && !VRInteractor->IsRightClickingOnUI() )
+				{
+					PlacingWithInteractor = VRInteractor;
+					break;
+				}
+			}
+		}
 
-	bool bTransactionStarted = false;
+		if( PlacingWithInteractor != nullptr )
+		{
+			FEditorDelegates::LoadSelectedAssetsIfNeeded.Broadcast();
 
-	// Figure out which controller pressed the button and started dragging
-	// @todo vreditor placement: This logic could misfire.  Ideally we would be routed information from the pointer event, so we can determine the hand.
-	UVREditorInteractor* PlacingWithInteractor = nullptr;
+			TArray< UObject* > DroppedObjects;
+			DroppedObjects.Reserve( DraggedAssets.Num() );
 
-	TArray<UViewportInteractor*>& Interactors = ViewportWorldInteraction->GetInteractors();
-	for ( UViewportInteractor* Interactor : Interactors )
-	{
-		UVREditorInteractor* VRInteractor = Cast<UVREditorInteractor>( Interactor );
-		FViewportActionKeyInput* SelectAndMoveAction = Interactor->GetActionWithName( ViewportWorldActionTypes::SelectAndMove );
-		FViewportActionKeyInput* SelectAndMoveLightlyPressedAction = Interactor->GetActionWithName( ViewportWorldActionTypes::SelectAndMove_LightlyPressed );
+			for( const FAssetData& AssetData : DraggedAssets )
+			{
+				UObject* AssetObj = AssetData.GetAsset();
+				if( AssetObj != nullptr )
+				{
+					// Don't add the same asset more than once
+					DroppedObjects.AddUnique( AssetObj );
+				}
+			}
 
-		if ( VRInteractor && ( ( SelectAndMoveAction && SelectAndMoveAction->bIsInputCaptured ) 
-			|| ( SelectAndMoveLightlyPressedAction && SelectAndMoveLightlyPressedAction->bIsInputCaptured ) ) )
-		{
-			if ( VRInteractor->IsClickingOnUI() && !VRInteractor->IsRightClickingOnUI() ) 
+			if( DroppedObjects.Num() > 0 )
 			{
-				PlacingWithInteractor = VRInteractor;
-				break;
+				// Hide the UI panel that's being used to drag
+				FloatingUIAssetDraggedFrom = PlacingWithInteractor->GetLastHoveredWidgetComponent();
+				Owner->GetUISystem().ShowEditorUIPanel( FloatingUIAssetDraggedFrom, PlacingWithInteractor, false, false, false );
+
+				const bool bShouldInterpolateFromDragLocation = VREd::PlacementInterpolationDuration->GetFloat() > KINDA_SMALL_NUMBER;
+				StartPlacingObjects( DroppedObjects, FactoryToUse, PlacingWithInteractor, bShouldInterpolateFromDragLocation );
+
+				const UVREditorAssetContainer& AssetContainer = Owner->GetAssetContainer();
+				Owner->PlaySound(AssetContainer.DropFromContentBrowserSound, PlacingWithInteractor->GetTransform().GetLocation());
 			}
 		}
 	}
+}
+
+void UVREditorWorldInteraction::StartPlacingObjects( const TArray<UObject*>& ObjectsToPlace, UActorFactory* FactoryToUse, UVREditorInteractor* PlacingWithInteractor, const bool bShouldInterpolateFromDragLocation )
+{
+	if( ViewportWorldInteraction == nullptr || PlacingWithInteractor == nullptr )
+	{
+		return;
+	}
 
+	const bool bIsPreview = false;	 // @todo vreditor placement: Consider supporting a "drop preview" actor (so you can cancel placement interactively)
+
+	bool bTransactionStarted = false;
 
-	// We're always expecting a hand to be hovering at the time we receive this event
-	if( PlacingWithInteractor )
 	{
 		// Cancel UI input
 		{
-			FViewportActionKeyInput* SelectAndMoveAction = PlacingWithInteractor->GetActionWithName( ViewportWorldActionTypes::SelectAndMove );
-			FViewportActionKeyInput* SelectAndMoveLightlyPressedAction = PlacingWithInteractor->GetActionWithName( ViewportWorldActionTypes::SelectAndMove_LightlyPressed );
-
 			PlacingWithInteractor->SetIsClickingOnUI( false );
 			PlacingWithInteractor->SetIsRightClickingOnUI( false );
 
-			if ( SelectAndMoveAction )
+			FViewportActionKeyInput* SelectAndMoveAction = PlacingWithInteractor->GetActionWithName( ViewportWorldActionTypes::SelectAndMove );
+			if( SelectAndMoveAction != nullptr )
 			{
 				SelectAndMoveAction->bIsInputCaptured = false;
 			}
-
-			if ( SelectAndMoveLightlyPressedAction )
-			{
-				SelectAndMoveLightlyPressedAction->bIsInputCaptured = false;
-			}
 		}
 
-		FloatingUIAssetDraggedFrom = PlacingWithInteractor->GetHoveringOverWidgetComponent();
-		// Hide the UI panel that's being used to drag
-		Owner->GetUISystem().ShowEditorUIPanel( FloatingUIAssetDraggedFrom, PlacingWithInteractor, false, false, true, false );
 
 		TArray< UObject* > DroppedObjects;
 		TArray< AActor* > AllNewActors;
 
 		UObject* DraggingSingleMaterialOrTexture = nullptr;
 
-		FVector PlaceAt = PlacingWithInteractor->GetHoverLocation();
+		FVector PlaceAt = PlacingWithInteractor->GetInteractorData().LastHoverLocationOverUI;
 
 		// Only place the object at the laser impact point if we're NOT going to interpolate to the impact
 		// location.  When interpolation is enabled, it looks much better to blend to the new location
-		const bool bShouldInterpolateFromDragLocation = VREd::PlacementInterpolationDuration->GetFloat() > KINDA_SMALL_NUMBER;
 		if( !bShouldInterpolateFromDragLocation )
 		{
 			FVector HitLocation = FVector::ZeroVector;
@@ -233,12 +262,12 @@ void UVREditorWorldInteraction::OnAssetDragStartedFromContentBrowser( const TArr
 			}
 		}
 
-		for( const FAssetData& AssetData : DraggedAssets )
+		for( UObject* AssetObj : ObjectsToPlace )
 		{
 			bool bCanPlace = true;
 
-			UObject* AssetObj = AssetData.GetAsset();
-			if( !ObjectTools::IsAssetValidForPlacing( ViewportWorldInteraction->GetWorld(), AssetData.ObjectPath.ToString() ) )
+			const FString ObjectPath = AssetObj->GetPathName();
+			if( !ObjectTools::IsAssetValidForPlacing( ViewportWorldInteraction->GetWorld(), ObjectPath ) )
 			{
 				bCanPlace = false;
 			}
@@ -257,13 +286,20 @@ void UVREditorWorldInteraction::OnAssetDragStartedFromContentBrowser( const TArr
 			}
 
 			const bool bIsMaterialOrTexture = ( AssetObj->IsA( UMaterialInterface::StaticClass() ) || AssetObj->IsA( UTexture::StaticClass() ) );
-			if( bIsMaterialOrTexture && DraggedAssets.Num() == 1 )
+			if( bIsMaterialOrTexture && ObjectsToPlace.Num() == 1 )
 			{
-				DraggingSingleMaterialOrTexture = AssetObj;;
+				DraggingSingleMaterialOrTexture = AssetObj;
+			}
+			else
+			{
+				// @todo mesheditor: We're dragging actors, so deactivate mesh editor mode for this.  They'll contend over transformables.
+				static FName MeshEditorModeName( TEXT( "MeshEditor" ) );
+				GLevelEditorModeTools().DeactivateMode( MeshEditorModeName );
 			}
 
 			// Check if the asset has an actor factory
-			bool bHasActorFactory = FActorFactoryAssetProxy::GetFactoryForAsset( AssetData ) != NULL;
+			bool bHasActorFactory = FActorFactoryAssetProxy::GetFactoryForAssetObject( AssetObj ) != nullptr;
+
 
 			if( !( AssetObj->IsA( AActor::StaticClass() ) || bHasActorFactory ) &&
 				!AssetObj->IsA( UBrushBuilder::StaticClass() ) )
@@ -296,26 +332,10 @@ void UVREditorWorldInteraction::OnAssetDragStartedFromContentBrowser( const TArr
 
 				TArray<AActor*> NewActors = FLevelEditorViewportClient::TryPlacingActorFromObject( ViewportWorldInteraction->GetWorld()->GetCurrentLevel(), AssetObj, bSelectNewActors, NewObjectFlags, FactoryToUse );
 
+				AllNewActors.Append( NewActors );
 				if( NewActors.Num() > 0 )
 				{
 					DroppedObjects.Add( AssetObj );
-					for( AActor* NewActor : NewActors )
-					{
-						AllNewActors.Add( NewActor );
-
-						if( VREd::OverrideSizeOfPlacedActors->GetFloat() > KINDA_SMALL_NUMBER )
-						{
-							const FBox LocalSpaceBounds = NewActor->CalculateComponentsBoundingBoxInLocalSpace();
-							const float LocalBoundsSize = LocalSpaceBounds.GetSize().GetAbsMax();
-
-							const float DesiredSize = 
-								( bShouldInterpolateFromDragLocation ? 
-									VREd::SizeOfActorsOverContentBrowserThumbnail->GetFloat() :
-									VREd::OverrideSizeOfPlacedActors->GetFloat() )
-								* ViewportWorldInteraction->GetWorldScaleFactor();
-							NewActor->SetActorScale3D( FVector( DesiredSize / LocalBoundsSize ) );
-						}
-					}
 				}
 			}
 		}
@@ -329,7 +349,7 @@ void UVREditorWorldInteraction::OnAssetDragStartedFromContentBrowser( const TArr
 			GEditor->DisableDeltaModification( false );
 		}
 
-		if( AllNewActors.Num() > 0 )	// @todo vreditor: Should we do this for dragged materials too?
+		if( AllNewActors.Num() > 0 )
 		{
 			if( !bIsPreview )
 			{
@@ -340,52 +360,54 @@ void UVREditorWorldInteraction::OnAssetDragStartedFromContentBrowser( const TArr
 
 				FEditorDelegates::OnNewActorsDropped.Broadcast( DroppedObjects, AllNewActors );
 			}
-		}
 
-		if( AllNewActors.Num() > 0 )
-		{
-			// Start dragging the new actor(s)
-			UPrimitiveComponent* ClickedTransformGizmoComponent = nullptr;
-			const bool bIsPlacingNewObjects = true;
-			const bool bStartTransaction = false;
-			if( ViewportWorldInteraction->StartDragging( PlacingWithInteractor, ClickedTransformGizmoComponent, PlaceAt, bIsPlacingNewObjects, bStartTransaction ) )
+			FBox BoundsOfAllActors;
 			{
-				// Dragging started!
-
-				// If we're interpolating, update the target transform of the actors to use our overridden size.  When
-				// we placed them we set their size to be 'thumbnail sized', and we want them to interpolate to
-				// their actual size in the world
-				if( bShouldInterpolateFromDragLocation )
+				BoundsOfAllActors.Init();
+				for( AActor* NewActor : AllNewActors )
 				{
-					for( TUniquePtr<FViewportTransformable>& TransformablePtr : ViewportWorldInteraction->GetTransformables() )
-					{
-						FViewportTransformable& Transformable = *TransformablePtr;
-
-						float ObjectSize = 1.0f;
+					BoundsOfAllActors += NewActor->CalculateComponentsBoundingBoxInLocalSpace();
+				}
+			}
+			const float BoundsOfAllActorsSize = BoundsOfAllActors.GetSize().GetAbsMax();
+			const float DesiredScale = ( VREd::SizeOfActorsOverContentBrowserThumbnail->GetFloat() / BoundsOfAllActorsSize ) * ViewportWorldInteraction->GetWorldScaleFactor();
 
-						if( VREd::OverrideSizeOfPlacedActors->GetFloat() > KINDA_SMALL_NUMBER )
-						{
-							ObjectSize = 1.0f * ViewportWorldInteraction->GetWorldScaleFactor();
+			// Start the placed objects off scaled down to match the content browser thumbnail
+			if( bShouldInterpolateFromDragLocation )
+			{
+ 				for( AActor* NewActor : AllNewActors )
+ 				{
+ 					NewActor->SetActorScale3D( FVector( DesiredScale ) );
+ 				}
+			}
 
-							const FBox LocalSpaceBounds = Transformable.GetLocalSpaceBoundingBox();
-							const float LocalBoundsSize = LocalSpaceBounds.GetSize().GetAbsMax();
+			// We changed the initial scale of selected actors, so make sure our transformables and gizmo start transform are up to date.
+			GEditor->NoteSelectionChange();
 
-							const float DesiredSize = VREd::OverrideSizeOfPlacedActors->GetFloat() * ViewportWorldInteraction->GetWorldScaleFactor();
-							const FVector NewScale( DesiredSize / LocalBoundsSize );
+			// Start dragging the new actor(s)
+			UPrimitiveComponent* ClickedTransformGizmoComponent = nullptr;
+			const bool bIsPlacingNewObjects = true;
+			const bool bAllowInterpolationWhenPlacing = bShouldInterpolateFromDragLocation;
+			const bool bStartTransaction = false;
+			const bool bWithGrabberSphere = false;	// Always place using the laser, not the grabber sphere
+			ViewportWorldInteraction->StartDragging( PlacingWithInteractor, ClickedTransformGizmoComponent, PlaceAt, bIsPlacingNewObjects, bAllowInterpolationWhenPlacing, bStartTransaction, bWithGrabberSphere );
 
-							Transformable.UnsnappedTargetTransform.SetScale3D( NewScale );
-							Transformable.StartTransform.SetScale3D( NewScale );
-							Transformable.LastTransform.SetScale3D( NewScale );
-							Transformable.TargetTransform.SetScale3D( NewScale );
-						}
-					}
-				}
+			// If we're interpolating, update the target transform of the actors to use our overridden size.  When
+			// we placed them we set their size to be 'thumbnail sized', and we want them to interpolate to
+			// their actual size in the world
+			if( bShouldInterpolateFromDragLocation )
+			{
+				const FVector NewScale( 1.0f / DesiredScale );
+				FViewportInteractorData& InteractorData = PlacingWithInteractor->GetInteractorData();
+				InteractorData.GizmoUnsnappedTargetTransform.SetScale3D( NewScale );
+				InteractorData.GizmoLastTransform.SetScale3D( NewScale );
+				InteractorData.GizmoTargetTransform.SetScale3D( NewScale );
 			}
 		}
 
 		if( DraggingSingleMaterialOrTexture != nullptr )
 		{
-			const FViewportActionKeyInput Action( ViewportWorldActionTypes::SelectAndMove_LightlyPressed );
+			const FViewportActionKeyInput Action( ViewportWorldActionTypes::SelectAndMove );
 
 			// Start dragging the material
 			StartDraggingMaterialOrTexture( PlacingWithInteractor,	Action, PlacingWithInteractor->GetHoverLocation(), DraggingSingleMaterialOrTexture );
@@ -401,6 +423,7 @@ void UVREditorWorldInteraction::PlaceDraggedMaterialOrTexture( UViewportInteract
 	{
 		// Check to see if the laser pointer is over something we can drop on
 		UPrimitiveComponent* HitComponent = nullptr;
+		FVector HitLocation = FVector::ZeroVector;
 		{
 			const bool bIgnoreGizmos = true;	// Never place on top of gizmos, just ignore them
 			const bool bEvenIfUIIsInFront = true;	// Don't let the UI block placement
@@ -413,6 +436,7 @@ void UVREditorWorldInteraction::PlaceDraggedMaterialOrTexture( UViewportInteract
 					if( Cast<AVREditorFloatingUI>( HitResult.GetComponent()->GetOwner() ) == nullptr )
 					{
 						HitComponent = HitResult.GetComponent();
+						HitLocation = HitResult.ImpactPoint;
 					}
 				}
 			}
@@ -436,13 +460,17 @@ void UVREditorWorldInteraction::PlaceDraggedMaterialOrTexture( UViewportInteract
 			{
 				// @todo vreditor placement: How do we get the material ID that was dropped on?  Regular editor uses hit proxies.  We may need to augment FHitResult.
 				// @todo vreditor placement: Support optionally dropping on all materials, not only the impacted material
-				const int32 TargetMaterialSlot = -1;	// All materials
-				const bool bPlaced = FComponentEditorUtils::AttemptApplyMaterialToComponent( HitComponent, DroppedObjAsMaterial, TargetMaterialSlot );
-			}
+				bool bPlaced = false;
+				check( Owner );
+				Owner->OnPlaceDraggedMaterial().Broadcast( HitComponent, DroppedObjAsMaterial, bPlaced );
+				if( !bPlaced )
+				{
+					const int32 TargetMaterialSlot = -1;	// All materials
+					bool bPlaced = FComponentEditorUtils::AttemptApplyMaterialToComponent( HitComponent, DroppedObjAsMaterial, TargetMaterialSlot );
 
-			if( DroppedObjAsMaterial || DroppedObjAsTexture )
-			{
-				UGameplayStatics::PlaySound2D( ViewportWorldInteraction->GetWorld(), DropMaterialOrMaterialSound );
+					const UVREditorAssetContainer& AssetContainer = Owner->GetAssetContainer();
+					Owner->PlaySound(AssetContainer.DropFromContentBrowserSound, Interactor->GetTransform().GetLocation());
+				}
 			}
 		}
 	}
@@ -461,8 +489,8 @@ void UVREditorWorldInteraction::SnapSelectedActorsToGround()
 
 		CommandList->ExecuteAction( Commands.SnapBottomCenterBoundsToFloor.ToSharedRef() );
 
-		// Selected objects have been moved, so make sure our transformables are all reset to the object's new location
-		ViewportWorldInteraction->ResetTransformables();
+		// Force transformables to refresh
+		GEditor->NoteSelectionChange();
 	}
 }
 
diff --git a/Engine/Source/Editor/VREditor/VREditorWorldInteraction.h b/Engine/Source/Editor/VREditor/VREditorWorldInteraction.h
--- a/Engine/Source/Editor/VREditor/VREditorWorldInteraction.h
+++ b/Engine/Source/Editor/VREditor/VREditorWorldInteraction.h
@@ -33,6 +33,9 @@ public:
 	/** Snaps the selected objects to the ground */
 	void SnapSelectedActorsToGround();
 
+	/** Try spawn and start placing the specified objects */
+	void StartPlacingObjects( const TArray<UObject*>& ObjectsToPlace, class UActorFactory* FactoryToUse, class UVREditorInteractor* PlacingWithInteractor, const bool bShouldInterpolateFromDragLocation );
+
 protected:
 	/** When an interactor stops dragging */
 	void StopDragging( UViewportInteractor* Interactor );
@@ -59,10 +62,6 @@ protected:
 	UPROPERTY()
 	UViewportWorldInteraction* ViewportWorldInteraction;
 
-	/** Sound for dropping materials and textures */
-	UPROPERTY()
-	class USoundCue* DropMaterialOrMaterialSound;
-
 	//
 	// Dragging object from UI
 	//
diff --git a/Engine/Source/Editor/ViewportInteraction/ActorTransformer.cpp b/Engine/Source/Editor/ViewportInteraction/ActorTransformer.cpp
--- a/Engine/Source/Editor/ViewportInteraction/ActorTransformer.cpp
+++ b/Engine/Source/Editor/ViewportInteraction/ActorTransformer.cpp
@@ -4,7 +4,8 @@
 #include "ActorViewportTransformable.h"
 #include "ViewportWorldInteraction.h"
 #include "Engine/Selection.h"
-
+#include "ViewportInteractionAssetContainer.h"
+#include "ViewportInteractor.h"
 
 void UActorTransformer::Init( UViewportWorldInteraction* InitViewportWorldInteraction )
 {
@@ -23,7 +24,41 @@ void UActorTransformer::Shutdown()
 }
 
 
-void UActorTransformer::OnActorSelectionChanged( UObject* ChangedObject )
+void UActorTransformer::OnStartDragging(class UViewportInteractor* Interactor)
+{
+	const UViewportInteractionAssetContainer& AssetContainer = ViewportWorldInteraction->GetAssetContainer();
+	const FVector& SoundLocation = Interactor->GetInteractorData().GizmoLastTransform.GetLocation();
+	const EViewportInteractionDraggingMode DraggingMode = Interactor->GetDraggingMode();
+
+	if (DraggingMode == EViewportInteractionDraggingMode::TransformablesWithGizmo)
+	{
+		ViewportWorldInteraction->PlaySound(AssetContainer.GizmoHandleSelectedSound, SoundLocation);
+	}
+	else if (DraggingMode == EViewportInteractionDraggingMode::TransformablesFreely ||
+		DraggingMode == EViewportInteractionDraggingMode::TransformablesAtLaserImpact)
+	{
+		ViewportWorldInteraction->PlaySound(AssetContainer.SelectionStartDragSound, SoundLocation);
+	}
+}
+
+void UActorTransformer::OnStopDragging(class UViewportInteractor* Interactor)
+{
+	const UViewportInteractionAssetContainer& AssetContainer = ViewportWorldInteraction->GetAssetContainer();
+	const FVector& SoundLocation = Interactor->GetInteractorData().GizmoLastTransform.GetLocation();
+	const EViewportInteractionDraggingMode DraggingMode = Interactor->GetDraggingMode();
+
+	if (DraggingMode == EViewportInteractionDraggingMode::TransformablesWithGizmo)
+	{
+		ViewportWorldInteraction->PlaySound(AssetContainer.GizmoHandleDropSound, SoundLocation);
+	}
+	else if (DraggingMode == EViewportInteractionDraggingMode::TransformablesFreely ||
+		DraggingMode == EViewportInteractionDraggingMode::TransformablesAtLaserImpact)
+	{
+		ViewportWorldInteraction->PlaySound(AssetContainer.SelectionDropSound, SoundLocation);
+	}
+}
+
+void UActorTransformer::UpdateTransformables()
 {
 	TArray<TUniquePtr<FViewportTransformable>> NewTransformables;
 
@@ -35,16 +70,28 @@ void UActorTransformer::OnActorSelectionChanged( UObject* ChangedObject )
 
 	for( UObject* SelectedActorObject : SelectedActorObjects )
 	{
-		AActor* SelectedActor = CastChecked<AActor>( SelectedActorObject );
-
-		FActorViewportTransformable* Transformable = new FActorViewportTransformable();
-
-		NewTransformables.Add( TUniquePtr<FViewportTransformable>( Transformable ) );
-
-		Transformable->ActorWeakPtr = SelectedActor;
-		Transformable->StartTransform = Transformable->LastTransform = Transformable->TargetTransform = Transformable->UnsnappedTargetTransform = Transformable->InterpolationSnapshotTransform = SelectedActor->GetTransform();
+		AActor* SelectedActor = Cast<AActor>( SelectedActorObject );
+		if( SelectedActor != nullptr )
+		{
+			// We only are able to move objects that have a root scene component
+			if( SelectedActor->GetRootComponent() != nullptr && Cast<USceneComponent>( SelectedActor->GetRootComponent() ) != nullptr )
+			{
+				FActorViewportTransformable* Transformable = new FActorViewportTransformable();
+
+				NewTransformables.Add( TUniquePtr<FViewportTransformable>( Transformable ) );
+
+				Transformable->ActorWeakPtr = SelectedActor;
+				Transformable->StartTransform = SelectedActor->GetTransform();
+			}
+		}
 	}
 
 	ViewportWorldInteraction->SetTransformables( MoveTemp( NewTransformables ) );
 }
 
+
+void UActorTransformer::OnActorSelectionChanged( UObject* ChangedObject )
+{
+	UpdateTransformables();
+}
+
diff --git a/Engine/Source/Editor/ViewportInteraction/ActorViewportTransformable.cpp b/Engine/Source/Editor/ViewportInteraction/ActorViewportTransformable.cpp
--- a/Engine/Source/Editor/ViewportInteraction/ActorViewportTransformable.cpp
+++ b/Engine/Source/Editor/ViewportInteraction/ActorViewportTransformable.cpp
@@ -11,6 +11,16 @@ void FActorViewportTransformable::ApplyTransform( const FTransform& NewTransform
 		const FTransform& ExistingTransform = Actor->GetActorTransform();
 		if( !ExistingTransform.Equals( NewTransform, 0.0f ) )
 		{
+			// If we're moving a non-movable actor while in simulate mode, go ahead and make it movable.  We're only
+			// editing the PIE copy of the actor here, so this won't affect the actual editor world.
+			if( GEditor->bIsSimulatingInEditor && Actor->GetWorld()->IsPlayInEditor() )
+			{
+				if( Actor->GetRootComponent() != nullptr && Actor->GetRootComponent()->Mobility != EComponentMobility::Movable )
+				{
+					Actor->GetRootComponent()->SetMobility( EComponentMobility::Movable );
+				}
+			}
+
 			GEditor->BroadcastBeginObjectMovement(*Actor);
 			const bool bOnlyTranslationChanged =
 				ExistingTransform.GetRotation() == NewTransform.GetRotation() &&
@@ -19,9 +29,14 @@ void FActorViewportTransformable::ApplyTransform( const FTransform& NewTransform
 			Actor->SetActorTransform( NewTransform, bSweep );
 			//GWarn->Logf( TEXT( "SMOOTH: Actor %s to %s" ), *Actor->GetName(), *Transformable.TargetTransform.ToString() );
 
-			Actor->InvalidateLightingCacheDetailed( bOnlyTranslationChanged );
+			// @todo vreditor: InvalidateLightingCacheDetailed() causes static mesh components to re-create their physics state, 
+			// cancelling all velocity on the rigid body.  So we currently avoid calling it for simulated actors.
+			if( !IsPhysicallySimulated() )
+			{
+				Actor->InvalidateLightingCacheDetailed( bOnlyTranslationChanged );
+			}
 
-			const bool bFinished = false;	// @todo gizmo now: PostEditChange never called; and bFinished=true never known!!
+			const bool bFinished = false;	// @todo gizmo: PostEditChange never called; and bFinished=true never known!!
 			Actor->PostEditMove( bFinished );
 			GEditor->BroadcastEndObjectMovement(*Actor);
 		}
@@ -43,17 +58,24 @@ const FTransform FActorViewportTransformable::GetTransform() const
 }
 
 
-FBox FActorViewportTransformable::GetLocalSpaceBoundingBox() const
+FBox FActorViewportTransformable::BuildBoundingBox( const FTransform& BoundingBoxToWorld ) const
 {
-	FBox LocalSpaceBounds( 0 );
-
+	FBox BoundingBox( 0 );
 	AActor* Actor = Cast<AActor>( ActorWeakPtr.Get() );
+
 	if( Actor != nullptr )
 	{
+		const FTransform WorldToBoundingBox = BoundingBoxToWorld.Inverse();
+		const FTransform& ActorToWorld = Actor->GetTransform();
+		const FTransform ActorToBoundingBox = ActorToWorld * WorldToBoundingBox;
+
 		const bool bIncludeNonCollidingComponents = false;	// @todo gizmo: Disabled this because it causes lights to have huge bounds
-		LocalSpaceBounds = Actor->CalculateComponentsBoundingBoxInLocalSpace( bIncludeNonCollidingComponents );
+		const FBox ActorSpaceBoundingBox = Actor->CalculateComponentsBoundingBoxInLocalSpace( bIncludeNonCollidingComponents );
+
+		BoundingBox = ActorSpaceBoundingBox.TransformBy( ActorToBoundingBox );
 	}
-	return LocalSpaceBounds;
+
+	return BoundingBox;
 }
 
 
@@ -92,6 +114,17 @@ void FActorViewportTransformable::SetLinearVelocity( const FVector& NewVelocity
 }
 
 
+FVector FActorViewportTransformable::GetLinearVelocity() const
+{
+	const AActor* Actor = Cast<AActor>(ActorWeakPtr.Get());
+	if (Actor != nullptr)
+	{
+		return Actor->GetVelocity();
+	}
+
+	return FVector::ZeroVector;
+}
+
 void FActorViewportTransformable::UpdateIgnoredActorList( TArray<AActor*>& IgnoredActors )
 {
 	AActor* Actor = Cast<AActor>( ActorWeakPtr.Get() );
diff --git a/Engine/Source/Editor/ViewportInteraction/Gizmo/VIGizmoHandle.cpp b/Engine/Source/Editor/ViewportInteraction/Gizmo/VIGizmoHandle.cpp
--- a/Engine/Source/Editor/ViewportInteraction/Gizmo/VIGizmoHandle.cpp
+++ b/Engine/Source/Editor/ViewportInteraction/Gizmo/VIGizmoHandle.cpp
@@ -192,42 +192,42 @@ void UGizmoHandleGroup::UpdateHandleColor( const int32 AxisIndex, FGizmoHandle&
 	UMaterialInstanceDynamic* MID1 = CastChecked<UMaterialInstanceDynamic>( HandleMesh->GetMaterial( 1 ) );
 
 	ABaseTransformGizmo* GizmoActor = CastChecked<ABaseTransformGizmo>( GetOwner() );
-	if ( GizmoActor )
+	if (GizmoActor)
 	{
 		UViewportWorldInteraction* WorldInteraction = GizmoActor->GetOwnerWorldInteraction();
-		if ( WorldInteraction )
+		if (WorldInteraction)
 		{
-			FLinearColor HandleColor = WorldInteraction->GetColor( UViewportWorldInteraction::EColors::DefaultColor );
-			if ( HandleMesh == DraggingHandle )
+			FLinearColor HandleColor = WorldInteraction->GetColor(UViewportWorldInteraction::EColors::DefaultColor);
+			if (HandleMesh == DraggingHandle)
 			{
-				HandleColor = WorldInteraction->GetColor( UViewportWorldInteraction::EColors::Dragging );
+				HandleColor = WorldInteraction->GetColor(UViewportWorldInteraction::EColors::GizmoDragging);
 			}
-			else if ( AxisIndex != INDEX_NONE )
+			else if (AxisIndex != INDEX_NONE)
 			{
-				switch ( AxisIndex )
+				switch (AxisIndex)
 				{
 				case 0:
-					HandleColor = WorldInteraction->GetColor( UViewportWorldInteraction::EColors::Forward );
+					HandleColor = WorldInteraction->GetColor(UViewportWorldInteraction::EColors::Forward);
 					break;
 
 				case 1:
-					HandleColor = WorldInteraction->GetColor( UViewportWorldInteraction::EColors::Right );
+					HandleColor = WorldInteraction->GetColor(UViewportWorldInteraction::EColors::Right);
 					break;
 
 				case 2:
-					HandleColor = WorldInteraction->GetColor( UViewportWorldInteraction::EColors::Up );
+					HandleColor = WorldInteraction->GetColor(UViewportWorldInteraction::EColors::Up);
 					break;
 				}
 
-				if ( HoveringOverHandles.Contains( HandleMesh ) )
+				if (HoveringOverHandles.Contains(HandleMesh))
 				{
-					HandleColor = FLinearColor::LerpUsingHSV( HandleColor, WorldInteraction->GetColor( UViewportWorldInteraction::EColors::Hover ), Handle.HoverAlpha );
+					HandleColor = FLinearColor::LerpUsingHSV(HandleColor, WorldInteraction->GetColor(UViewportWorldInteraction::EColors::GizmoHover), Handle.HoverAlpha);
 				}
 			}
 
-			static FName StaticHandleColorParameter( "Color" );
-			MID0->SetVectorParameterValue( StaticHandleColorParameter, HandleColor );
-			MID1->SetVectorParameterValue( StaticHandleColorParameter, HandleColor );
+			static FName StaticHandleColorParameter("Color");
+			MID0->SetVectorParameterValue(StaticHandleColorParameter, HandleColor);
+			MID1->SetVectorParameterValue(StaticHandleColorParameter, HandleColor);
 		}
 	}
 }
@@ -373,47 +373,48 @@ void UAxisGizmoHandleGroup::UpdateHandlesRelativeTransformOnAxis( const FTransfo
 			int32 CenterHandleCount, FacingAxisIndex, CenterAxisIndex;
 			HandlePlacement.GetCenterHandleCountAndFacingAxisIndex( /* Out */ CenterHandleCount, /* Out */ FacingAxisIndex, /* Out */ CenterAxisIndex);
 			
-			FVector GizmoSpaceFacingAxisVector = GetAxisVector( FacingAxisIndex, HandlePlacement.Axes[FacingAxisIndex] );
-			
-			// Check on which side we are relative to the gizmo
-			const FVector GizmoSpaceViewLocation = GetOwner()->GetTransform().InverseTransformPosition( ViewLocation );
-			if ( GizmoSpaceViewLocation[ FacingAxisIndex ] < 0 )
+			if (DraggingHandle == nullptr)
 			{
-				GizmoSpaceFacingAxisVector[ FacingAxisIndex ] *= -1.0f;
-			}
+				FVector GizmoSpaceFacingAxisVector = GetAxisVector( FacingAxisIndex, HandlePlacement.Axes[FacingAxisIndex] );
+			
+				// Check on which side we are relative to the gizmo
+				const FVector GizmoSpaceViewLocation = GetOwner()->GetTransform().InverseTransformPosition( ViewLocation );
+				if ( GizmoSpaceViewLocation[ FacingAxisIndex ] < 0 )
+				{
+					GizmoSpaceFacingAxisVector[ FacingAxisIndex ] *= -1.0f;
+				}
 
-			const FTransform GizmoOriginToFacingAxisRotation( GizmoSpaceFacingAxisVector.ToOrientationQuat() );
+				const FTransform GizmoOriginToFacingAxisRotation( GizmoSpaceFacingAxisVector.ToOrientationQuat() );
 	
-			FTransform HandleToGizmoOrigin = HandleToCenter * GizmoOriginToFacingAxisRotation;
+				FTransform HandleToGizmoOrigin = HandleToCenter * GizmoOriginToFacingAxisRotation;
 				
-			// Check on the axis if the offset is on the other side of the object compared to the view. Switch the offset to the side of the view if it does
-			FVector GizmoSpaceFacingAxisOffset = HandleToGizmoOrigin.GetLocation();
-			for ( int32 AxisIndex = 0; AxisIndex < 3; ++AxisIndex )
-			{		
-				if ( AxisIndex != FacingAxisIndex &&
-					( ( GizmoSpaceFacingAxisOffset[ AxisIndex ] > 0 && GizmoSpaceViewLocation[ AxisIndex ] < 0 ) ||
-					  ( GizmoSpaceFacingAxisOffset[ AxisIndex ] < 0 && GizmoSpaceViewLocation[ AxisIndex ] > 0 ) ) )
-				{
-					GizmoSpaceFacingAxisOffset[ AxisIndex ] *= -1.0f;
-				}	
-			}
+				// Check on the axis if the offset is on the other side of the object compared to the view. Switch the offset to the side of the view if it does
+				FVector GizmoSpaceFacingAxisOffset = HandleToGizmoOrigin.GetLocation();
+				for ( int32 AxisIndex = 0; AxisIndex < 3; ++AxisIndex )
+				{		
+					if ( AxisIndex != FacingAxisIndex &&
+						( ( GizmoSpaceFacingAxisOffset[ AxisIndex ] > 0 && GizmoSpaceViewLocation[ AxisIndex ] < 0 ) ||
+						  ( GizmoSpaceFacingAxisOffset[ AxisIndex ] < 0 && GizmoSpaceViewLocation[ AxisIndex ] > 0 ) ) )
+					{
+						GizmoSpaceFacingAxisOffset[ AxisIndex ] *= -1.0f;
+					}	
+				}
 
-			GizmoSpaceFacingAxisOffset *= AnimationAlpha;
+				GizmoSpaceFacingAxisOffset *= AnimationAlpha;
 
-			HandleToGizmoOrigin.SetLocation( GizmoSpaceFacingAxisOffset * GizmoScale );
+				HandleToGizmoOrigin.SetLocation( GizmoSpaceFacingAxisOffset * GizmoScale );
 
-			// Set the final transform
-			GizmoHandleMeshComponent->SetRelativeTransform( HandleToGizmoOrigin );
+				// Set the final transform
+				GizmoHandleMeshComponent->SetRelativeTransform( HandleToGizmoOrigin );
 			
-			float GizmoHandleScale = GizmoScale;
-			if ( GizmoScale != GizmoHoverScale )
-			{
+				float GizmoHandleScale = GizmoScale;
+
 				// Make the handle bigger while hovered (but don't affect the offset -- we want it to scale about it's origin)
 				GizmoHandleScale *= FMath::Lerp( 1.0f, GizmoHoverScale, Handle.HoverAlpha );
-			}
 			
-			GizmoHandleScale *= AnimationAlpha;
-			GizmoHandleMeshComponent->SetRelativeScale3D( FVector( GizmoHandleScale ) );
+				GizmoHandleScale *= AnimationAlpha;
+				GizmoHandleMeshComponent->SetRelativeScale3D( FVector( GizmoHandleScale ) );
+			}
 
 			// Update material
 			UpdateHandleColor( FacingAxisIndex, Handle, DraggingHandle, HoveringOverHandles );
diff --git a/Engine/Source/Editor/ViewportInteraction/Gizmo/VIPivotTransformGizmo.cpp b/Engine/Source/Editor/ViewportInteraction/Gizmo/VIPivotTransformGizmo.cpp
--- a/Engine/Source/Editor/ViewportInteraction/Gizmo/VIPivotTransformGizmo.cpp
+++ b/Engine/Source/Editor/ViewportInteraction/Gizmo/VIPivotTransformGizmo.cpp
@@ -8,21 +8,34 @@
 #include "VIStretchGizmoHandle.h"
 #include "VIUniformScaleGizmoHandle.h"
 #include "ViewportWorldInteraction.h"
+#include "ViewportInteractor.h"
+#include "ViewportInteractionAssetContainer.h"
+#include "Materials/Material.h"
+#include "Engine/Font.h"
+#include "Engine/CollisionProfile.h"
+#include "Materials/MaterialInstance.h"
+#include "Materials/MaterialInstanceDynamic.h"
+#include "Components/TextRenderComponent.h"
+#include "Math/UnitConversion.h"
 
 namespace VREd //@todo VREditor: Duplicates of TransformGizmo
 {
 	// @todo vreditor tweak: Tweak out console variables
 	static FAutoConsoleVariable PivotGizmoMinDistanceForScaling( TEXT( "VI.PivotGizmoMinDistanceForScaling" ), 0.0f, TEXT( "How far away the camera needs to be from an object before we'll start scaling it based on distance"));
-	static FAutoConsoleVariable PivotGizmoDistanceScaleFactor( TEXT( "VI.PivotGizmoDistanceScaleFactor" ), 0.003f, TEXT( "How much the gizmo handles should increase in size with distance from the camera, to make it easier to select" ) );
+	static FAutoConsoleVariable PivotGizmoDistanceScaleFactor( TEXT( "VI.PivotGizmoDistanceScaleFactor" ), 0.0035f, TEXT( "How much the gizmo handles should increase in size with distance from the camera, to make it easier to select" ) );
 	static FAutoConsoleVariable PivotGizmoTranslationPivotOffsetX( TEXT("VI.PivotGizmoTranslationPivotOffsetX" ), 30.0f, TEXT( "How much the translation cylinder is offsetted from the pivot" ) );
 	static FAutoConsoleVariable PivotGizmoScalePivotOffsetX( TEXT( "VI.PivotGizmoScalePivotOffsetX" ), 120.0f, TEXT( "How much the non-uniform scale is offsetted from the pivot" ) );
 	static FAutoConsoleVariable PivotGizmoPlaneTranslationPivotOffsetYZ(TEXT("VI.PivotGizmoPlaneTranslationPivotOffsetYZ" ), 40.0f, TEXT( "How much the plane translation is offsetted from the pivot" ) );
 	static FAutoConsoleVariable PivotGizmoTranslationScaleMultiply( TEXT( "VI.PivotGizmoTranslationScaleMultiply" ), 2.0f, TEXT( "Multiplies translation handles scale" ) );
 	static FAutoConsoleVariable PivotGizmoTranslationHoverScaleMultiply( TEXT( "VI.PivotGizmoTranslationHoverScaleMultiply" ), 0.75f, TEXT( "Multiplies translation handles hover scale" ) );
+	static FAutoConsoleVariable PivotGizmoAimAtShrinkSize( TEXT( "VI.PivotGizmoAimAtShrinkSize" ), 0.3f, TEXT( "The minimum size when not aiming at the gizmo (0 to 1)" ) );
+	static FAutoConsoleVariable PivotGizmoAimAtAnimationSpeed( TEXT( "VI.PivotGizmoAimAtAnimationSpeed" ), 0.15f, TEXT( "The speed to animate to the gizmo full size when aiming at it" ) );
 }
 
 APivotTransformGizmo::APivotTransformGizmo() :
-	Super()
+	Super(),
+	AimingAtGizmoScaleAlpha(0.0f),
+	LastDraggingHandle(nullptr)
 {
 	UniformScaleGizmoHandleGroup = CreateDefaultSubobject<UUniformScaleGizmoHandleGroup>( TEXT( "UniformScaleHandles" ), true );
 	UniformScaleGizmoHandleGroup->SetOwningTransformGizmo(this);
@@ -76,34 +89,78 @@ void APivotTransformGizmo::UpdateGizmo(const EGizmoHandleTypes InGizmoType, cons
 {
 	GizmoType = InGizmoType;
 
-	const float WorldScaleFactor = GetWorld()->GetWorldSettings()->WorldToMeters / 100.0f;
-
 	// Position the gizmo at the location of the first selected actor
-	const bool bSweep = false;
-	this->SetActorTransform(InLocalToWorld, bSweep);
-
+	SetActorTransform(InLocalToWorld);
 
 	// Increase scale with distance, to make gizmo handles easier to click on
 	const float WorldSpaceDistanceToToPivot = FMath::Max(VREd::PivotGizmoMinDistanceForScaling->GetFloat(), FMath::Sqrt(FVector::DistSquared( GetActorLocation(), InViewLocation )));
+	const float WorldScaleFactor = WorldInteraction->GetWorldScaleFactor();
+	float GizmoScale = (InScaleMultiplier * ((WorldSpaceDistanceToToPivot / WorldScaleFactor) * VREd::PivotGizmoDistanceScaleFactor->GetFloat())) * WorldScaleFactor;
+	const bool bIsWorldSpaceGizmo = (WorldInteraction->GetTransformGizmoCoordinateSpace() == COORD_World);
+
+	if (LastDraggingHandle != nullptr && DraggingHandle == nullptr)
+	{
+		AimingAtGizmoScaleAlpha = 0.0f;
+	}
 
-	// @todo gizmo: Scale causes bounds-based features to look wrong
-	const float GizmoScale(InScaleMultiplier * ((WorldSpaceDistanceToToPivot / WorldScaleFactor) * VREd::PivotGizmoDistanceScaleFactor->GetFloat() ) * WorldScaleFactor);
+	// Only scale the gizmo down when not aiming at it for VR implementations
+	if (WorldInteraction->IsInVR())
+	{
+		bool bIsAimingTowards = false;
+		bool bAnyInteractorInCorrectState = false;
+		const float GizmoRadius = (GizmoScale * 350) * 0.5f;
+
+		// Check if any interactor has a laser close to the gizmo
+		for (UViewportInteractor* Interactor : WorldInteraction->GetInteractors())
+		{	
+			// We only want the interactor to affect the size when aiming at the gizmo if they are dragging nothing
+			if (Interactor->GetDraggingMode() == EViewportInteractionDraggingMode::Nothing)
+			{
+				bAnyInteractorInCorrectState = true;			
+
+				FVector LaserPointerStart, LaserPointerEnd;
+				if (Interactor->GetLaserPointer(/*Out*/ LaserPointerStart, /*Out*/ LaserPointerEnd))
+				{
+					const FVector ClosestPointOnLaser = FMath::ClosestPointOnLine(LaserPointerStart, LaserPointerEnd, InLocalToWorld.GetLocation());
+					const float ClosestPointDistance = (ClosestPointOnLaser - InLocalToWorld.GetLocation()).Size();
+					if (ClosestPointDistance <= GizmoRadius)
+					{
+						bIsAimingTowards = true;
+						break;
+					}
+				}
+			}
+		}
+
+		if (bAnyInteractorInCorrectState)
+		{
+			const float DeltaTime = WorldInteraction->GetCurrentDeltaTime();
+			if (bIsAimingTowards)
+			{
+				AimingAtGizmoScaleAlpha += DeltaTime / VREd::PivotGizmoAimAtAnimationSpeed->GetFloat();
+			}
+			else
+			{
+				AimingAtGizmoScaleAlpha -= DeltaTime / VREd::PivotGizmoAimAtAnimationSpeed->GetFloat();
+			}
+			AimingAtGizmoScaleAlpha = FMath::Clamp(AimingAtGizmoScaleAlpha, VREd::PivotGizmoAimAtShrinkSize->GetFloat(), 1.0f);
+
+			GizmoScale *= AimingAtGizmoScaleAlpha;
+		}
+	}
 
 	// Update animation
 	float AnimationAlpha = GetAnimationAlpha();
 
 	// Update all the handles
-	bool bIsHoveringOrDraggingScaleGizmo = false;
 	for (UGizmoHandleGroup* HandleGroup : AllHandleGroups)
 	{
 		if (HandleGroup != nullptr)
 		{
-			if (DraggingHandle == nullptr || HandleGroup == StretchGizmoHandleGroup || HandleGroup == RotationGizmoHandleGroup)
-			{
-				bool bIsHoveringOrDraggingThisHandleGroup = false;
-				HandleGroup->UpdateGizmoHandleGroup(InLocalToWorld, InLocalBounds, InViewLocation, bInAllHandlesVisible, DraggingHandle,
-					InHoveringOverHandles, AnimationAlpha, GizmoScale, InGizmoHoverScale, InGizmoHoverAnimationDuration, /* Out */ bIsHoveringOrDraggingThisHandleGroup );
-			}
+			//const FTransform HandleGroupTransform = bIsWorldSpaceGizmo == true && HandleGroup != RotationGizmoHandleGroup == true ? TransformGizmoToWorld : InLocalToWorld;
+			bool bIsHoveringOrDraggingThisHandleGroup = false;
+			HandleGroup->UpdateGizmoHandleGroup(InLocalToWorld, InLocalBounds, InViewLocation, bInAllHandlesVisible, DraggingHandle,
+				InHoveringOverHandles, AnimationAlpha, GizmoScale, InGizmoHoverScale, InGizmoHoverAnimationDuration, /* Out */ bIsHoveringOrDraggingThisHandleGroup);
 			
 			if (HandleGroup != RotationGizmoHandleGroup)
 			{
@@ -111,6 +168,8 @@ void APivotTransformGizmo::UpdateGizmo(const EGizmoHandleTypes InGizmoType, cons
 			}
 		}
 	}
+
+	LastDraggingHandle = DraggingHandle;
 }
 
 /************************************************************************/
@@ -119,14 +178,8 @@ void APivotTransformGizmo::UpdateGizmo(const EGizmoHandleTypes InGizmoType, cons
 UPivotTranslationGizmoHandleGroup::UPivotTranslationGizmoHandleGroup() :
 	Super()
 {
-	UStaticMesh* TranslationHandleMesh = nullptr;
-	{
-		static ConstructorHelpers::FObjectFinder<UStaticMesh> ObjectFinder( TEXT( "/Engine/VREditor/TransformGizmo/TranslateArrowHandle" ) );
-		TranslationHandleMesh = ObjectFinder.Object;
-		check( TranslationHandleMesh != nullptr );
-	}
-
-	CreateHandles( TranslationHandleMesh, FString( "PivotTranslationHandle" ) );
+	const UViewportInteractionAssetContainer& AssetContainer = UViewportWorldInteraction::LoadAssetContainer();
+	CreateHandles( AssetContainer.TranslationHandleMesh, FString( "PivotTranslationHandle" ) );
 }
 
 
@@ -159,14 +212,8 @@ EGizmoHandleTypes UPivotTranslationGizmoHandleGroup::GetHandleType() const
 UPivotScaleGizmoHandleGroup::UPivotScaleGizmoHandleGroup() :
 	Super()
 {
-	UStaticMesh* ScaleHandleMesh = nullptr;
-	{
-		static ConstructorHelpers::FObjectFinder<UStaticMesh> ObjectFinder( TEXT( "/Engine/VREditor/TransformGizmo/UniformScaleHandle" ) );
-		ScaleHandleMesh = ObjectFinder.Object;
-		check( ScaleHandleMesh != nullptr );
-	}
-
-	CreateHandles( ScaleHandleMesh, FString( "PivotScaleHandle" ) );	
+	const UViewportInteractionAssetContainer& AssetContainer = UViewportWorldInteraction::LoadAssetContainer();
+	CreateHandles( AssetContainer.UniformScaleHandleMesh, FString( "PivotScaleHandle" ) );	
 }
 
 void UPivotScaleGizmoHandleGroup::UpdateGizmoHandleGroup( const FTransform& LocalToWorld, const FBox& LocalBounds, const FVector ViewLocation, const bool bAllHandlesVisible, class UActorComponent* DraggingHandle, 
@@ -200,14 +247,8 @@ bool UPivotScaleGizmoHandleGroup::SupportsWorldCoordinateSpace() const
 UPivotPlaneTranslationGizmoHandleGroup::UPivotPlaneTranslationGizmoHandleGroup() :
 	Super()
 {
-	UStaticMesh* TranslationHandleMesh = nullptr;
-	{
-		static ConstructorHelpers::FObjectFinder<UStaticMesh> ObjectFinder( TEXT( "/Engine/VREditor/TransformGizmo/PlaneTranslationHandle" ) );
-		TranslationHandleMesh = ObjectFinder.Object;
-		check( TranslationHandleMesh != nullptr );
-	}
-
-	CreateHandles( TranslationHandleMesh, FString( "PlaneTranslationHandle" ) );
+	const UViewportInteractionAssetContainer& AssetContainer = UViewportWorldInteraction::LoadAssetContainer();
+	CreateHandles( AssetContainer.PlaneTranslationHandleMesh, FString( "PlaneTranslationHandle" ) );
 }
 
 void UPivotPlaneTranslationGizmoHandleGroup::UpdateGizmoHandleGroup( const FTransform& LocalToWorld, const FBox& LocalBounds, const FVector ViewLocation, const bool bAllHandlesVisible, class UActorComponent* DraggingHandle,
@@ -237,23 +278,60 @@ EGizmoHandleTypes UPivotPlaneTranslationGizmoHandleGroup::GetHandleType() const
 UPivotRotationGizmoHandleGroup::UPivotRotationGizmoHandleGroup() :
 	Super()
 {
-	UStaticMesh* QuarterRotationHandleMesh = nullptr;
-	{
-		static ConstructorHelpers::FObjectFinder<UStaticMesh> ObjectFinder(TEXT("/Engine/VREditor/TransformGizmo/RotationHandleQuarter"));
-		QuarterRotationHandleMesh = ObjectFinder.Object;
-		check(QuarterRotationHandleMesh != nullptr);
-	}
+	const UViewportInteractionAssetContainer& AssetContainer = UViewportWorldInteraction::LoadAssetContainer();
 
+	UStaticMesh* QuarterRotationHandleMesh = AssetContainer.RotationHandleMesh;
 	CreateHandles(QuarterRotationHandleMesh, FString("RotationHandle"));
 
 	{
-		static ConstructorHelpers::FObjectFinder<UStaticMesh> ObjectFinder(TEXT("/Engine/VREditor/TransformGizmo/RotationHandleFull"));
-		UStaticMesh* FullRotationHandleMesh = ObjectFinder.Object;
+		RootFullRotationHandleComponent = CreateDefaultSubobject<USceneComponent>(TEXT("RootFullRotationHandleComponent"));
+		RootFullRotationHandleComponent->SetMobility(EComponentMobility::Movable);
+		RootFullRotationHandleComponent->SetupAttachment(this);
+
+		UStaticMesh* FullRotationHandleMesh = AssetContainer.RotationHandleSelectedMesh;
 		check(FullRotationHandleMesh != nullptr);
 
 		FullRotationHandleMeshComponent = CreateMeshHandle(FullRotationHandleMesh, FString("FullRotationHandle"));
 		FullRotationHandleMeshComponent->SetVisibility(false);
 		FullRotationHandleMeshComponent->SetCollisionEnabled(ECollisionEnabled::NoCollision);
+		FullRotationHandleMeshComponent->SetupAttachment(RootFullRotationHandleComponent);
+	}
+
+	{
+		UStaticMesh* RotationHandleIndicatorMesh = AssetContainer.StartRotationIndicatorMesh;
+		check(RotationHandleIndicatorMesh != nullptr);
+
+		//Start rotation indicator
+		RootStartRotationIdicatorComponent = CreateDefaultSubobject<USceneComponent>(TEXT("RootStartRotationIndicator"));
+		StartRotationIndicatorMeshComponent = CreateDefaultSubobject<UStaticMeshComponent>(TEXT("StartRotationIndicator"));
+		SetupIndicator(RootStartRotationIdicatorComponent, StartRotationIndicatorMeshComponent, RotationHandleIndicatorMesh);
+	}
+
+	{
+		UStaticMesh* RotationHandleIndicatorMesh = AssetContainer.CurrentRotationIndicatorMesh;
+		check(RotationHandleIndicatorMesh != nullptr);
+
+		//Delta rotation indicator
+		RootDeltaRotationIndicatorComponent = CreateDefaultSubobject<USceneComponent>(TEXT("RootDeltaRotationIndicator"));
+		DeltaRotationIndicatorMeshComponent = CreateDefaultSubobject<UStaticMeshComponent>(TEXT("DeltaRotationIndicator"));
+		SetupIndicator(RootDeltaRotationIndicatorComponent, DeltaRotationIndicatorMeshComponent, RotationHandleIndicatorMesh);
+	}
+
+	{
+		UMaterialInstanceDynamic* DynamicMaterialInst = UMaterialInstanceDynamic::Create(AssetContainer.TransformGizmoMaterial, GetTransientPackage());
+		check(DynamicMaterialInst != nullptr);
+
+		DeltaRotationIndicatorMeshComponent->SetMaterial(0, DynamicMaterialInst);
+		StartRotationIndicatorMeshComponent->SetMaterial(0, DynamicMaterialInst);
+		FullRotationHandleMeshComponent->SetMaterial(0, DynamicMaterialInst);
+
+		UMaterialInstanceDynamic* TranslucentDynamicMaterialInst = UMaterialInstanceDynamic::Create(AssetContainer.TranslucentTransformGizmoMaterial, GetTransientPackage());
+		check(TranslucentDynamicMaterialInst != nullptr);
+
+		DeltaRotationIndicatorMeshComponent->SetMaterial(1, TranslucentDynamicMaterialInst);
+		StartRotationIndicatorMeshComponent->SetMaterial(1, TranslucentDynamicMaterialInst);
+		FullRotationHandleMeshComponent->SetMaterial(1, TranslucentDynamicMaterialInst);
+
 	}
 }
 
@@ -265,9 +343,8 @@ void UPivotRotationGizmoHandleGroup::UpdateGizmoHandleGroup( const FTransform& L
 		GizmoScale, GizmoHoverScale, GizmoHoverAnimationDuration, bOutIsHoveringOrDraggingThisHandleGroup);
 
 	bool bShowFullRotationDragHandle = false;
-	const ECoordSystem CoordSystem = OwningTransformGizmoActor->GetOwnerWorldInteraction()->GetTransformGizmoCoordinateSpace();
-	
-	const bool bIsDraggingHandle = DraggingHandle == nullptr ? true : false;
+	UViewportWorldInteraction* WorldInteraction = OwningTransformGizmoActor->GetOwnerWorldInteraction();
+	const ECoordSystem CoordSystem = WorldInteraction->GetTransformGizmoCoordinateSpace();
 	const bool bIsTypeSupported = (OwningTransformGizmoActor->GetGizmoType() == EGizmoHandleTypes::All && GetShowOnUniversalGizmo()) || GetHandleType() == OwningTransformGizmoActor->GetGizmoType();
 	const bool bSupportsCurrentCoordinateSpace = SupportsWorldCoordinateSpace() || CoordSystem != COORD_World;
 	const bool bShowAnyRotationHandle = (bIsTypeSupported && bSupportsCurrentCoordinateSpace && bAllHandlesVisible);
@@ -279,7 +356,9 @@ void UPivotRotationGizmoHandleGroup::UpdateGizmoHandleGroup( const FTransform& L
 		if (GizmoHandleMeshComponent != nullptr)	// Can be null if no handle for this specific placement
 		{
 			const bool bDraggingCurrentHandle = (DraggingHandle != nullptr && DraggingHandle == GizmoHandleMeshComponent);
-			GizmoHandleMeshComponent->SetVisibility(bShowAnyRotationHandle && !bDraggingCurrentHandle);
+			const bool bShouldShowHandle = bShowAnyRotationHandle && !bDraggingCurrentHandle;
+			GizmoHandleMeshComponent->SetVisibility(bShouldShowHandle);
+			GizmoHandleMeshComponent->SetCollisionEnabled(bShouldShowHandle ? ECollisionEnabled::QueryOnly : ECollisionEnabled::NoCollision);
 
 			int32 CenterHandleCount, FacingAxisIndex, CenterAxisIndex;
 			const FTransformGizmoHandlePlacement HandlePlacement = MakeHandlePlacementForIndex(HandleIndex);
@@ -287,19 +366,40 @@ void UPivotRotationGizmoHandleGroup::UpdateGizmoHandleGroup( const FTransform& L
 
 			if (bDraggingCurrentHandle)
 			{
+				bool bInitialDrag = false;
 				bShowFullRotationDragHandle = true;
-				if (!StartActorRotation.IsSet())
+				FVector GizmoSpaceFacingAxisVector = GetAxisVector(FacingAxisIndex, HandlePlacement.Axes[FacingAxisIndex]);
+				
+				// Set the initial rotation
+				if (!StartDragRotation.IsSet())
 				{
-					StartActorRotation = OwningTransformGizmoActor->GetTransform().GetRotation();
+					StartDragRotation = LocalToWorld.GetRotation();
+					bInitialDrag = true;
 				}
-
+				
+				// Set the root of the full rotation handles to the rotation we had when starting the drag.
 				if (CoordSystem == ECoordSystem::COORD_Local)
 				{
-					OwningTransformGizmoActor->SetActorRotation(StartActorRotation.GetValue());
+					RootFullRotationHandleComponent->SetWorldRotation(StartDragRotation.GetValue());	
 				}
 
-				const FVector GizmoSpaceFacingAxisVector = GetAxisVector(FacingAxisIndex, HandlePlacement.Axes[FacingAxisIndex]);
-				FullRotationHandleMeshComponent->SetRelativeTransform(FTransform(GizmoSpaceFacingAxisVector.ToOrientationQuat(), FVector::ZeroVector, FVector(GizmoScale * AnimationAlpha)));
+				FullRotationHandleMeshComponent->SetRelativeTransform(FTransform(GizmoSpaceFacingAxisVector.ToOrientationQuat(), FVector::ZeroVector, FVector(GizmoScale)));
+
+				const FVector LocalIntersectPoint = OwningTransformGizmoActor->GetOwnerWorldInteraction()->GetLocalIntersectPointOnRotationGizmo();
+				UpdateIndicator(RootDeltaRotationIndicatorComponent, LocalIntersectPoint, FacingAxisIndex);
+
+				// Update the start indicator only if this was the first time dragging
+				if (bInitialDrag)
+				{
+					UpdateIndicator(RootStartRotationIdicatorComponent, LocalIntersectPoint, FacingAxisIndex);
+
+					FLinearColor Color = WorldInteraction->GetColor(UViewportWorldInteraction::EColors::GizmoHover);
+					SetIndicatorColor(FullRotationHandleMeshComponent, Color);
+					SetIndicatorColor(DeltaRotationIndicatorMeshComponent, Color);
+					SetIndicatorColor(StartRotationIndicatorMeshComponent, Color);
+
+					bInitialDrag = false;
+				}
 			}
 			else if(DraggingHandle == nullptr)
 			{
@@ -346,22 +446,19 @@ void UPivotRotationGizmoHandleGroup::UpdateGizmoHandleGroup( const FTransform& L
 				}
 
 				// Set the final transform
-				GizmoHandleMeshComponent->SetRelativeTransform(FTransform(Rotation, FVector::ZeroVector, FVector(GizmoScale * AnimationAlpha )));
+				GizmoHandleMeshComponent->SetRelativeTransform(FTransform(Rotation, FVector::ZeroVector, FVector(GizmoScale)));
 
 				// Update material
 				UpdateHandleColor(FacingAxisIndex, Handle, DraggingHandle, HoveringOverHandles);
+
+				// Reset the start rotation
+				StartDragRotation.Reset();
 			}
 		}
 	}
 
-	// Show or hide the visuals for full rotation and set the collision accordingly
-	FullRotationHandleMeshComponent->SetCollisionEnabled(bShowFullRotationDragHandle == true ? ECollisionEnabled::QueryOnly : ECollisionEnabled::NoCollision);
-	FullRotationHandleMeshComponent->SetVisibility(bShowFullRotationDragHandle);
-
-	if (!bShowFullRotationDragHandle)
-	{
-		StartActorRotation.Reset();
-	}
+	// Show or hide the visuals for when rotating
+	ShowRotationVisuals(bShowFullRotationDragHandle);
 }
 
 ETransformGizmoInteractionType UPivotRotationGizmoHandleGroup::GetInteractionType() const
@@ -373,3 +470,63 @@ EGizmoHandleTypes UPivotRotationGizmoHandleGroup::GetHandleType() const
 {
 	return EGizmoHandleTypes::Rotate;
 }
+
+void UPivotRotationGizmoHandleGroup::UpdateIndicator(USceneComponent* IndicatorRoot, const FVector& Direction, const uint32 FacingAxisIndex)
+{
+	float Y = 0, X = 0;
+	if (FacingAxisIndex == 0)
+	{
+		Y = Direction.Y;
+		X = Direction.Z;
+	}
+	else if (FacingAxisIndex == 1)
+	{
+		Y = -Direction.X;
+		X = Direction.Z;
+	}
+	else if (FacingAxisIndex == 2)
+	{
+		Y = Direction.Y;
+		X = -Direction.X;
+	}
+
+	const float Angle = FMath::RadiansToDegrees(FMath::Atan2(Y, X));
+	IndicatorRoot->SetRelativeRotation(FRotator(0, 0, Angle));
+}
+
+void UPivotRotationGizmoHandleGroup::ShowRotationVisuals(const bool bInShow)
+{
+	FullRotationHandleMeshComponent->SetCollisionEnabled(bInShow == true ? ECollisionEnabled::QueryOnly : ECollisionEnabled::NoCollision);
+	FullRotationHandleMeshComponent->SetVisibility(bInShow);
+	StartRotationIndicatorMeshComponent->SetVisibility(bInShow);
+	DeltaRotationIndicatorMeshComponent->SetVisibility(bInShow);
+}
+
+void UPivotRotationGizmoHandleGroup::SetupIndicator(USceneComponent* RootComponent, UStaticMeshComponent* IndicatorMeshComponent, UStaticMesh* Mesh)
+{
+	RootComponent->SetMobility(EComponentMobility::Movable);
+	RootComponent->SetupAttachment(FullRotationHandleMeshComponent);
+
+	IndicatorMeshComponent->SetStaticMesh(Mesh);
+	IndicatorMeshComponent->SetMobility(EComponentMobility::Movable);
+	IndicatorMeshComponent->SetupAttachment(RootComponent);
+	IndicatorMeshComponent->SetCollisionEnabled(ECollisionEnabled::NoCollision);
+	IndicatorMeshComponent->SetVisibility(false);
+	IndicatorMeshComponent->bGenerateOverlapEvents = false;
+	IndicatorMeshComponent->SetCanEverAffectNavigation(false);
+	IndicatorMeshComponent->bCastDynamicShadow = true;
+	IndicatorMeshComponent->bCastStaticShadow = false;
+	IndicatorMeshComponent->bAffectDistanceFieldLighting = false;
+	IndicatorMeshComponent->bAffectDynamicIndirectLighting = false;
+	IndicatorMeshComponent->SetRelativeLocation(FVector(0, 0, 100));
+}
+
+void UPivotRotationGizmoHandleGroup::SetIndicatorColor(UStaticMeshComponent* InMeshComponent, const FLinearColor& InHandleColor)
+{
+	UMaterialInstanceDynamic* MID0 = Cast<UMaterialInstanceDynamic>(InMeshComponent->GetMaterial(0));
+	UMaterialInstanceDynamic* MID1 = Cast<UMaterialInstanceDynamic>(InMeshComponent->GetMaterial(1));
+
+	static FName StaticHandleColorParameter("Color");
+	MID0->SetVectorParameterValue(StaticHandleColorParameter, InHandleColor);
+	MID1->SetVectorParameterValue(StaticHandleColorParameter, InHandleColor);
+}
\ No newline at end of file
diff --git a/Engine/Source/Editor/ViewportInteraction/Gizmo/VIPivotTransformGizmo.h b/Engine/Source/Editor/ViewportInteraction/Gizmo/VIPivotTransformGizmo.h
--- a/Engine/Source/Editor/ViewportInteraction/Gizmo/VIPivotTransformGizmo.h
+++ b/Engine/Source/Editor/ViewportInteraction/Gizmo/VIPivotTransformGizmo.h
@@ -20,7 +20,7 @@ public:
 
 	/** Default constructor that sets up CDO properties */
 	APivotTransformGizmo();
-
+	
 	/** Called by the world interaction system after we've been spawned into the world, to allow
 	    us to create components and set everything up nicely for the selected objects that we'll be
 		used to manipulate */
@@ -52,6 +52,13 @@ private:
 	/** Stretch handle group component */
 	UPROPERTY()
 	class UStretchGizmoHandleGroup* StretchGizmoHandleGroup;
+
+	/** The alpha for gizmo animation when aiming at it with a laser */
+	float AimingAtGizmoScaleAlpha;
+
+	/** Handle from previous tick that was dragged */
+	UPROPERTY()
+	UActorComponent* LastDraggingHandle;
 };
 
 /**
@@ -155,9 +162,40 @@ public:
 
 private:
 
+	/** Updates the root of an indicator to rotate the indicator itself */
+	void UpdateIndicator(USceneComponent* IndicatorRoot, const FVector& Direction, const uint32 FacingAxisIndex);
+
+	/** Make the components visible when dragging rotation */
+	void ShowRotationVisuals(const bool bInShow);
+
+	void SetupIndicator(USceneComponent* RootComponent, UStaticMeshComponent* IndicatorMeshComponent, UStaticMesh* Mesh);
+	
+	void SetIndicatorColor(UStaticMeshComponent* InMeshComponent, const FLinearColor& InHandleColor);
+
+	/** Root component of all the mesh components that are used to visualize the rotation when dragging */
+	UPROPERTY()
+	USceneComponent* RootFullRotationHandleComponent;
+
+	/** When dragging a rotation handle the full rotation circle appears */
 	UPROPERTY()
 	UStaticMeshComponent* FullRotationHandleMeshComponent;
 
-	// The rotation of the owning transform gizmo when starting to use a rotation handle
-	TOptional<FQuat> StartActorRotation;
+	/** The mesh that indicated the start rotation */
+	UPROPERTY()
+	UStaticMeshComponent* StartRotationIndicatorMeshComponent;
+
+	/** The root component of the start rotation indicator */
+	UPROPERTY()
+	USceneComponent* RootStartRotationIdicatorComponent;
+
+	/** The mesh that indicated the delta rotation */
+	UPROPERTY()
+	UStaticMeshComponent* DeltaRotationIndicatorMeshComponent;
+
+	/** The root component of the delta rotation indicator */
+	UPROPERTY()
+	USceneComponent* RootDeltaRotationIndicatorComponent;
+
+	/** The rotation when starting to drag the gizmo */
+	TOptional<FQuat> StartDragRotation;
 };
diff --git a/Engine/Source/Editor/ViewportInteraction/Gizmo/VIUniformScaleGizmoHandle.cpp b/Engine/Source/Editor/ViewportInteraction/Gizmo/VIUniformScaleGizmoHandle.cpp
--- a/Engine/Source/Editor/ViewportInteraction/Gizmo/VIUniformScaleGizmoHandle.cpp
+++ b/Engine/Source/Editor/ViewportInteraction/Gizmo/VIUniformScaleGizmoHandle.cpp
@@ -70,16 +70,16 @@ void UUniformScaleGizmoHandleGroup::UpdateGizmoHandleGroup( const FTransform& Lo
 			if( GizmoActor )
 			{
 				UViewportWorldInteraction* WorldInteraction = GizmoActor->GetOwnerWorldInteraction();
-				if( WorldInteraction )
+				if( WorldInteraction && WorldInteraction->IsActive() )
 				{
 					FLinearColor HandleColor = WorldInteraction->GetColor( UViewportWorldInteraction::EColors::DefaultColor );
 					if (UniformScaleHandle == DraggingHandle)
 					{
-						HandleColor = WorldInteraction->GetColor( UViewportWorldInteraction::EColors::Dragging );
+						HandleColor = WorldInteraction->GetColor( UViewportWorldInteraction::EColors::GizmoDragging );
 					}
 					else if (HoveringOverHandles.Contains( UniformScaleHandle ))
 					{
-						HandleColor = FLinearColor::LerpUsingHSV( HandleColor, WorldInteraction->GetColor( UViewportWorldInteraction::EColors::Hover ), Handle.HoverAlpha );
+						HandleColor = FLinearColor::LerpUsingHSV( HandleColor, WorldInteraction->GetColor( UViewportWorldInteraction::EColors::GizmoHover ), Handle.HoverAlpha );
 					}
 
 					MID0->SetVectorParameterValue( "Color", HandleColor );
diff --git a/Engine/Source/Editor/ViewportInteraction/MouseCursorInteractor.cpp b/Engine/Source/Editor/ViewportInteraction/MouseCursorInteractor.cpp
--- a/Engine/Source/Editor/ViewportInteraction/MouseCursorInteractor.cpp
+++ b/Engine/Source/Editor/ViewportInteraction/MouseCursorInteractor.cpp
@@ -22,7 +22,6 @@ void UMouseCursorInteractor::Init()
 	// Setup keys
 	{
 		AddKeyAction( EKeys::LeftMouseButton, FViewportActionKeyInput( ViewportWorldActionTypes::SelectAndMove ) );
-		AddKeyAction( EKeys::MiddleMouseButton, FViewportActionKeyInput( ViewportWorldActionTypes::SelectAndMove_LightlyPressed ) );
 	}
 }
 
@@ -38,6 +37,8 @@ void UMouseCursorInteractor::PollInput()
 	{
 		FEditorViewportClient& ViewportClient = *ViewportClientPtr;
 
+		bIsControlKeyPressed = ViewportClient.Viewport->KeyState( EKeys::LeftControl ) || ViewportClient.Viewport->KeyState( EKeys::RightControl );
+
 		// Only if we're not tracking (RMB looking)
 		if( !ViewportClient.IsTracking() )
 		{
@@ -74,3 +75,15 @@ void UMouseCursorInteractor::PollInput()
 	}
 	InteractorData.RoomSpaceTransform = InteractorData.Transform * WorldInteraction->GetRoomTransform().Inverse();
 }
+
+
+bool UMouseCursorInteractor::IsModifierPressed() const
+{
+	return bIsControlKeyPressed;
+}
+
+bool UMouseCursorInteractor::AllowLaserSmoothing() const
+{
+	return false;
+}
+
diff --git a/Engine/Source/Editor/ViewportInteraction/Public/ActorTransformer.h b/Engine/Source/Editor/ViewportInteraction/Public/ActorTransformer.h
--- a/Engine/Source/Editor/ViewportInteraction/Public/ActorTransformer.h
+++ b/Engine/Source/Editor/ViewportInteraction/Public/ActorTransformer.h
@@ -13,17 +13,23 @@ class VIEWPORTINTERACTION_API UActorTransformer : public UViewportTransformer
 
 public:
 
-	UFUNCTION()
+	// UViewportTransformer overrides
 	virtual void Init( class UViewportWorldInteraction* InitViewportWorldInteraction ) override;
-
-	UFUNCTION()
 	virtual void Shutdown() override;
-
+	virtual bool CanAlignToActors() const override
+	{
+		return true;
+	}
+	virtual void OnStartDragging(class UViewportInteractor* Interactor) override;
+	virtual void OnStopDragging(class UViewportInteractor* Interactor) override;
 
 protected:
 
 	/** Called when selection changes in the level */
 	void OnActorSelectionChanged( UObject* ChangedObject );
 
+	/** Creates a list of transformables with the current selected actors */
+	void UpdateTransformables();
+
 };
 
diff --git a/Engine/Source/Editor/ViewportInteraction/Public/ActorViewportTransformable.h b/Engine/Source/Editor/ViewportInteraction/Public/ActorViewportTransformable.h
--- a/Engine/Source/Editor/ViewportInteraction/Public/ActorViewportTransformable.h
+++ b/Engine/Source/Editor/ViewportInteraction/Public/ActorViewportTransformable.h
@@ -21,9 +21,10 @@ public:
 	// FViewportTransformable overrides
 	virtual const FTransform GetTransform() const override;
 	virtual void ApplyTransform( const FTransform& NewTransform, const bool bSweep ) override;
-	virtual FBox GetLocalSpaceBoundingBox() const override;
+	virtual FBox BuildBoundingBox( const FTransform& BoundingBoxToWorld ) const override;
 	virtual bool IsPhysicallySimulated() const override;
 	virtual void SetLinearVelocity( const FVector& NewVelocity ) override;
+	virtual FVector GetLinearVelocity() const override;
 	virtual void UpdateIgnoredActorList( TArray<class AActor*>& IgnoredActors ) override;
 
 	/** The actual actor object */
diff --git a/Engine/Source/Editor/ViewportInteraction/Public/MouseCursorInteractor.h b/Engine/Source/Editor/ViewportInteraction/Public/MouseCursorInteractor.h
--- a/Engine/Source/Editor/ViewportInteraction/Public/MouseCursorInteractor.h
+++ b/Engine/Source/Editor/ViewportInteraction/Public/MouseCursorInteractor.h
@@ -22,8 +22,15 @@ public:
 
 	// ViewportInteractor overrides
 	virtual void PollInput() override;
+	virtual bool IsModifierPressed() const override;
+
+protected:
+
+	virtual bool AllowLaserSmoothing() const;
 
 private:
 
-	// ...
+	/** Whether the control key was pressed the last time input was polled */
+	bool bIsControlKeyPressed;
+
 };
diff --git a/Engine/Source/Editor/ViewportInteraction/Public/ViewportDragOperation.h b/Engine/Source/Editor/ViewportInteraction/Public/ViewportDragOperation.h
--- a/Engine/Source/Editor/ViewportInteraction/Public/ViewportDragOperation.h
+++ b/Engine/Source/Editor/ViewportInteraction/Public/ViewportDragOperation.h
@@ -52,6 +52,9 @@ public:
 	/* Destroys the dragoperation */
 	void ClearDragOperation();
 
+	/** If this operation is currently active. */
+	bool IsDragging() const;
+
 private:
 
 	/** The actual dragging implementation */
diff --git a/Engine/Source/Editor/ViewportInteraction/Public/ViewportInteractionTypes.h b/Engine/Source/Editor/ViewportInteraction/Public/ViewportInteractionTypes.h
--- a/Engine/Source/Editor/ViewportInteraction/Public/ViewportInteractionTypes.h
+++ b/Engine/Source/Editor/ViewportInteraction/Public/ViewportInteractionTypes.h
@@ -62,9 +62,6 @@ enum class EViewportInteractionDraggingMode : uint8
 	/** Moving a custom interactable */
 	Interactable,
 
-	/** Custom implementation for dragging */
-	Custom,
-
 	/** Dragging a material */
 	Material
 };
diff --git a/Engine/Source/Editor/ViewportInteraction/Public/ViewportInteractionUtils.h b/Engine/Source/Editor/ViewportInteraction/Public/ViewportInteractionUtils.h
new file mode 100644
--- /dev/null
+++ b/Engine/Source/Editor/ViewportInteraction/Public/ViewportInteractionUtils.h
@@ -0,0 +1,69 @@
+#pragma once
+
+namespace ViewportInteractionUtils
+{
+	/************************************************************************/
+	/* 1 Euro filter smoothing algorithm									*/
+	/* http://cristal.univ-lille.fr/~casiez/1euro/							*/
+	/************************************************************************/
+	class VIEWPORTINTERACTION_API FOneEuroFilter
+	{
+	private:
+
+		class FLowpassFilter
+		{
+		public:
+
+			/** Default constructor */
+			FLowpassFilter();
+
+			/** Calculate */
+			FVector Filter(const FVector& InValue, const FVector& InAlpha);
+
+			/** If the filter was not executed yet */
+			bool IsFirstTime() const;
+
+			/** Get the previous filtered value */
+			FVector GetPrevious() const;
+
+		private:
+
+			/** The previous filtered value */
+			FVector Previous;
+
+			/** If this is the first time doing a filter */
+			bool bFirstTime;
+		};
+
+	public:
+
+		/** Default constructor */
+		FOneEuroFilter();
+
+		FOneEuroFilter(const double InMinCutoff, const double InCutoffSlope, const double InDeltaCutoff);
+
+		/** Smooth vector */
+		FVector Filter(const FVector& InRaw, const double InDeltaTime);
+
+		/** Set the minimum cutoff */
+		void SetMinCutoff(const double InMinCutoff);
+
+		/** Set the cutoff slope */
+		void SetCutoffSlope(const double InCutoffSlope);
+
+		/** Set the delta slope */
+		void SetDeltaCutoff(const double InDeltaCutoff);
+
+	private:
+
+		const FVector CalculateCutoff(const FVector& InValue);
+		const FVector CalculateAlpha(const FVector& InCutoff, const double InDeltaTime) const;
+		const float CalculateAlpha(const float InCutoff, const double InDeltaTime) const;
+
+		double MinCutoff;
+		double CutoffSlope;
+		double DeltaCutoff;
+		FLowpassFilter RawFilter;
+		FLowpassFilter DeltaFilter;
+	};
+}
\ No newline at end of file
diff --git a/Engine/Source/Editor/ViewportInteraction/Public/ViewportInteractor.h b/Engine/Source/Editor/ViewportInteraction/Public/ViewportInteractor.h
--- a/Engine/Source/Editor/ViewportInteraction/Public/ViewportInteractor.h
+++ b/Engine/Source/Editor/ViewportInteraction/Public/ViewportInteractor.h
@@ -9,6 +9,7 @@
 #include "Engine/EngineTypes.h"
 #include "ViewportInteractionTypes.h"
 #include "Editor/ViewportInteraction/ViewportInteractorData.h"
+#include "ViewportInteractionUtils.h"
 #include "ViewportInteractor.generated.h"
 
 class AActor;
@@ -41,17 +42,14 @@ public:
 	/** Gets the paired interactor assigned by the world interaction, can return null when there is no other interactor */
 	UViewportInteractor* GetOtherInteractor() const;
 
-	/** Updates the hover state of this interactor */
-	void UpdateHoverResult( FHitResult& HitResult );
-
 	/** Whenever the ViewportWorldInteraction is shut down, the interacts will shut down as well */
 	virtual void Shutdown();
 
 	/** Update for this interactor called by the ViewportWorldInteraction */
-	virtual void Tick( const float DeltaTime ) {};
+	virtual void Tick( const float DeltaTime );
 
-	/** Gets the hover component from the interactor data */
-	virtual class UActorComponent* GetHoverComponent();
+	/** Gets the last component hovered on by the interactor laser. */
+	virtual class UActorComponent* GetLastHoverComponent();
 
 	/** Adds a new action to the KeyToActionMap */
 	void AddKeyAction( const FKey& Key, const FViewportActionKeyInput& Action );
@@ -68,6 +66,12 @@ public:
 	/** Handles axis input and translates it to actions */
 	bool HandleInputAxis( class FEditorViewportClient& ViewportClient, const FKey Key, const float Delta, const float DeltaTime );
 
+	/** @return	Returns true if this interactor's designated 'modifier button' is currently held down.  Some interactors may not support this */
+	virtual bool IsModifierPressed() const
+	{
+		return false;
+	}
+
 	/** Gets the world transform of this interactor */
 	FTransform GetTransform() const;
 
@@ -80,6 +84,9 @@ public:
 	/** Gets the interactor data drag velocity */
 	FVector GetDragTranslationVelocity() const;
 
+	/** Sets the hover location */
+	void SetHoverLocation(const FVector& InHoverLocation);
+
 	/**
 	 * Gets the start and end point of the laser pointer for the specified hand
 	 *
@@ -91,7 +98,7 @@ public:
 	 *
 	 * @return	True if we have motion controller data for this hand and could return a valid result
 	 */
-	bool GetLaserPointer( FVector& LaserPointerStart, FVector& LaserPointerEnd, const bool bEvenIfBlocked = false, const float LaserLengthOverride = 0.0f ) const;
+	bool GetLaserPointer( FVector& LaserPointerStart, FVector& LaserPointerEnd, const bool bEvenIfBlocked = false, const float LaserLengthOverride = 0.0f );
 
 	/**
 	 * Gets a sphere on this interactor that can be used to interact with objects in close proximity
@@ -101,7 +108,7 @@ public:
 	 *
 	 * @return	True if the sphere is available and valid, or false is the interactor was busy or we could not determine a valid position for the interactor
 	 */
-	bool GetGrabberSphere( FSphere& OutGrabberSphere, const bool bEvenIfBlocked = false ) const;
+	bool GetGrabberSphere( FSphere& OutGrabberSphere, const bool bEvenIfBlocked = false );
 
 	/** Gets the maximum length of a laser pointer */
 	float GetLaserPointerMaxLength() const;
@@ -123,10 +130,10 @@ public:
 		TArray<UClass*>* ObjectsInFrontOfGizmo = nullptr, const bool bEvenIfBlocked = false, const float LaserLengthOverride = 0.0f );
 
 	/** Reset the values before checking the hover actions */
-	virtual void ResetHoverState( const float DeltaTime ) {};
+	virtual void ResetHoverState();
 
 	/** Needs to be implemented by the base to calculate drag ray length and the velocity for the ray */
-	virtual void CalculateDragRay() {};
+	virtual void CalculateDragRay( float& InOutDragRayLength, float& InOutDragRayVelocity ) {};
 
 	/**
 	 * Creates a hand transform and forward vector for a laser pointer for a given hand
@@ -140,10 +147,10 @@ public:
 	virtual bool GetTransformAndForwardVector( FTransform& OutHandTransform, FVector& OutForwardVector ) const;
 
 	/** Called by StartDragging in world interaction to give the interactor a chance of acting upon starting a drag operation */
-	virtual void OnStartDragging( const FVector& HitLocation, const bool bIsPlacingNewObjects ) {};
+	virtual void OnStartDragging( const FVector& HitLocation, const bool bIsPlacingNewObjects );
 
 	/** Gets the interactor laser hover location */
-	FVector GetHoverLocation() const;
+	FVector GetHoverLocation();
 
 	/** If the interactor laser is currently hovering */
 	bool IsHovering() const;
@@ -164,42 +171,32 @@ public:
 	/** Gets the drag haptic feedback strength console variable */
 	float GetDragHapticFeedbackStrength() const;
 
-	/**
-	 * Enables or disable "light press" locking for this interactor.  This is an advanced feature.
-	 * It should be true if we allow locking of the lightly pressed state for the current event.  This can 
-	 * be useful to turn off in cases where you always want a full press to override the light press, even 
-	 * if it was a very slow press
-	 *
-	 * @param	bAllow		Enables or disables light press locking for this hand
-	 */
-	void SetAllowTriggerLightPressLocking( const bool bInAllow );
+	/** If this interactor is hovering over a type that has priority from GetHitResultFromLaserPointer */
+	bool IsHoveringOverPriorityType() const;
 
-	/** @return Check if this interactor allows the light press to be locked */
-	bool AllowTriggerLightPressLocking() const;
+	/** Returns true if currently hovering over selected actor */
+	bool IsHoveringOverSelectedActor() const;
 
-	/**
-	 * When responding to a "light press" event, you can call this and pass false to to prevent a full press event from
-	 * being possible at all.  Useful when you're not planning to use FullyPressed for anything, but you don't want a 
-	 * full press to cancel your light press.
-	 *
-	 * @param	bInAllow	Whether to allow a full press to interrupt the light press.  Defaults to true, and will reset to true with every re-press of the trigger.
-	 */
-	void SetAllowTriggerFullPress( const bool bInAllow );
+	/** Reset the stored laser end location at the end of tick */
+	void ResetLaserEnd();
 
-	/** @return Returns whether a full press is allowed to interrupt a light press. */
-	bool AllowTriggerFullPress() const;
+protected:
 
-	/** If this interactor is hovering over a type that has priority from GetHitResultFromLaserPointer */
-	bool IsHoveringOverPriorityType() const;
+	/** To be overridden. Called by HandleInputKey before delegates and default input implementation */
+	virtual void PreviewInputKey( class FEditorViewportClient& ViewportClient, FViewportActionKeyInput& Action, const FKey Key, const EInputEvent Event, bool& bOutWasHandled ) {};
 
-protected:
+	/** To be overridden. Called by HandleInputAxis before delegates and default input implementation */
+	virtual void PreviewInputAxis( class FEditorViewportClient& ViewportClient, FViewportActionKeyInput& Action, const FKey Key, const float Delta, const float DeltaTime, bool& bOutWasHandled ) {};
 
-	/** To be overridden by base class. Called by HandleInputKey before delegates and default input implementation */
+	/** To be overridden. Called by HandleInputKey before delegates and default input implementation */
 	virtual void HandleInputKey( class FEditorViewportClient& ViewportClient, FViewportActionKeyInput& Action, const FKey Key, const EInputEvent Event, bool& bOutWasHandled ) {};
 	
-	/** To be overridden by base class. Called by HandleInputAxis before delegates and default input implementation */
+	/** To be overridden. Called by HandleInputAxis before delegates and default input implementation */
 	virtual void HandleInputAxis( class FEditorViewportClient& ViewportClient, FViewportActionKeyInput& Action, const FKey Key, const float Delta, const float DeltaTime, bool& bOutWasHandled ) {};
 
+	/** If this interactors allows to smooth the laser. Default is true. */
+	virtual bool AllowLaserSmoothing() const;
+
 protected:
 
 	/** All the private data for the interactor */
@@ -218,4 +215,12 @@ protected:
 
 	/** True if this interactor supports 'grabber sphere' interaction.  Usually disabled for mouse cursors */
 	bool bAllowGrabberSphere;
+
+	/** Store end of the laser pointer. This will be returned when calling GetLaserPointer multiple times a tick */
+	TOptional<FVector> SavedLaserPointerEnd;
+
+private:
+
+	/** Smoothing filter for laser */
+	ViewportInteractionUtils::FOneEuroFilter SmoothingOneEuroFilter;
 };
\ No newline at end of file
diff --git a/Engine/Source/Editor/ViewportInteraction/Public/ViewportTransformable.h b/Engine/Source/Editor/ViewportInteraction/Public/ViewportTransformable.h
--- a/Engine/Source/Editor/ViewportInteraction/Public/ViewportTransformable.h
+++ b/Engine/Source/Editor/ViewportInteraction/Public/ViewportTransformable.h
@@ -13,11 +13,7 @@ public:
 
 	/** Sets up safe defaults */
 	FViewportTransformable()
-		: StartTransform( FTransform::Identity ),
-		  LastTransform( FTransform::Identity ),
-		  TargetTransform( FTransform::Identity ),
-		  UnsnappedTargetTransform( FTransform::Identity ),
-		  InterpolationSnapshotTransform( FTransform::Identity )
+		: StartTransform( FTransform::Identity )
 	{
 	}
 
@@ -27,8 +23,8 @@ public:
 	/** Updates the transform of the actual object */
 	virtual void ApplyTransform( const FTransform& NewTransform, const bool bSweep ) = 0;
 
-	/** Returns the bounding box of this transformable in local space */
-	virtual FBox GetLocalSpaceBoundingBox() const = 0;
+	/** Returns the bounding box of this transformable, built in the specified coordinate system */
+	virtual FBox BuildBoundingBox( const FTransform& BoundingBoxToWorld ) const = 0;
 
 	/** Returns true if this transformable is a single, unoriented point in space, thus never
 	    supports being rotated or scaled when only a single transformable is selected. */
@@ -48,26 +44,19 @@ public:
 	{
 	}
 
+	/** Get the velocity of the object */
+	virtual FVector GetLinearVelocity() const
+	{
+		return FVector::ZeroVector;
+	}
+
 	/** For actor transformables, this will add it's actor to the incoming list */
 	virtual void UpdateIgnoredActorList( TArray<class AActor*>& IgnoredActors )
 	{
 	}
 
 
-	/** The object's position when we started the action */
+	/** The object's world space transform when we started the action */
 	FTransform StartTransform;
-
-	/** The last position we were at.  Used for smooth interpolation when snapping is enabled */
-	FTransform LastTransform;
-
-	/** The target position we want to be at.  Used for smooth interpolation when snapping is enabled */
-	FTransform TargetTransform;
-
-	/** The target position we'd be at if snapping wasn't enabled.  This is used for 'elastic' smoothing when
-	snapping is turned on.  Basically it allows us to 'reach' toward the desired position */
-	FTransform UnsnappedTargetTransform;
-
-	/** A "snapshot" transform, used for certain interpolation effects */
-	FTransform InterpolationSnapshotTransform;
 };
 
diff --git a/Engine/Source/Editor/ViewportInteraction/Public/ViewportTransformer.h b/Engine/Source/Editor/ViewportInteraction/Public/ViewportTransformer.h
--- a/Engine/Source/Editor/ViewportInteraction/Public/ViewportTransformer.h
+++ b/Engine/Source/Editor/ViewportInteraction/Public/ViewportTransformer.h
@@ -18,6 +18,27 @@ public:
 	UFUNCTION()
 	virtual void Shutdown();
 
+	/** @return If this transformer can be used to align its transformable to actors in the scene */
+	UFUNCTION()
+	virtual bool CanAlignToActors() const
+	{
+		return false;
+	};
+
+	/** @return True if the transform gizmo should be aligned to the center of the bounds of all selected objects with more than one is selected.  Otherwise it will be at the pivot of the last transformable, similar to legacl editor actor selection */
+	UFUNCTION()
+	virtual bool ShouldCenterTransformGizmoPivot() const
+	{
+		return false;
+	}
+
+	/** When starting to drag */
+	UFUNCTION()
+	virtual void OnStartDragging(class UViewportInteractor* Interactor) {};
+
+	/** When ending drag */
+	UFUNCTION()
+	virtual void OnStopDragging(class UViewportInteractor* Interactor) {};
 
 protected:
 
