Manual copy from 3340948 & 3341239.
- Added OnSwitchBeginPIEAndSIE event.
- Disabled SIE to PIE for Oculus.
- Moved bIsActive / bActive from VREditorMode and ViewportWorldInteraction to base class EditorWorldExtension.
- Made bActive private in EditorWorldExtension.h.
- Fix merge compiler issue in VREditorMode.cpp, added include StaticMeshComponents.

diff --git a/Engine/Source/Editor/UnrealEd/Private/Editor.cpp b/Engine/Source/Editor/UnrealEd/Private/Editor.cpp
--- a/Engine/Source/Editor/UnrealEd/Private/Editor.cpp
+++ b/Engine/Source/Editor/UnrealEd/Private/Editor.cpp
@@ -95,6 +95,7 @@ FEditorDelegates::FOnPIEEvent							FEditorDelegates::EndPIE;
 FEditorDelegates::FOnPIEEvent							FEditorDelegates::PausePIE;
 FEditorDelegates::FOnPIEEvent							FEditorDelegates::ResumePIE;
 FEditorDelegates::FOnPIEEvent							FEditorDelegates::SingleStepPIE;
+FEditorDelegates::FOnPIEEvent							FEditorDelegates::OnPreSwitchBeginPIEAndSIE;
 FEditorDelegates::FOnPIEEvent							FEditorDelegates::OnSwitchBeginPIEAndSIE;
 FSimpleMulticastDelegate								FEditorDelegates::PropertySelectionChange;
 FSimpleMulticastDelegate								FEditorDelegates::PostLandscapeLayerUpdated;
diff --git a/Engine/Source/Editor/UnrealEd/Private/EditorWorldExtension.cpp b/Engine/Source/Editor/UnrealEd/Private/EditorWorldExtension.cpp
--- a/Engine/Source/Editor/UnrealEd/Private/EditorWorldExtension.cpp
+++ b/Engine/Source/Editor/UnrealEd/Private/EditorWorldExtension.cpp
@@ -11,7 +11,8 @@
 /************************************************************************/
 UEditorWorldExtension::UEditorWorldExtension() :
 	Super(),
-	OwningExtensionsCollection( nullptr )
+	OwningExtensionsCollection(nullptr),
+	bActive(true)
 {
 
 }
@@ -97,6 +98,16 @@ void UEditorWorldExtension::DestroyTransientActor(AActor* Actor)
 	}
 }
 
+void UEditorWorldExtension::SetActive(const bool bInActive)
+{
+	bActive = bInActive;
+}
+
+bool UEditorWorldExtension::IsActive() const
+{
+	return bActive;
+}
+
 UEditorWorldExtensionCollection* UEditorWorldExtension::GetOwningCollection()
 {
 	return OwningExtensionsCollection;
@@ -179,6 +190,7 @@ UEditorWorldExtensionCollection::UEditorWorldExtensionCollection() :
 	if( !IsTemplate() )
 	{
 		FEditorDelegates::PostPIEStarted.AddUObject( this, &UEditorWorldExtensionCollection::PostPIEStarted );
+		FEditorDelegates::PrePIEEnded.AddUObject( this, &UEditorWorldExtensionCollection::OnPreEndPIE );
 		FEditorDelegates::EndPIE.AddUObject( this, &UEditorWorldExtensionCollection::OnEndPIE );
 	}
 }
@@ -186,6 +198,7 @@ UEditorWorldExtensionCollection::UEditorWorldExtensionCollection() :
 UEditorWorldExtensionCollection::~UEditorWorldExtensionCollection()
 {
 	FEditorDelegates::PostPIEStarted.RemoveAll( this );
+	FEditorDelegates::PrePIEEnded.RemoveAll( this );
 	FEditorDelegates::EndPIE.RemoveAll( this );
 
 	EditorExtensions.Empty();
@@ -281,7 +294,10 @@ void UEditorWorldExtensionCollection::Tick( float DeltaSeconds )
 	for (FEditorExtensionTuple& EditorExtensionTuple : EditorExtensions)
 	{
 		UEditorWorldExtension* EditorExtension = EditorExtensionTuple.Get<0>();
-		EditorExtension->Tick(DeltaSeconds);
+		if (EditorExtension->IsActive())
+		{
+			EditorExtension->Tick(DeltaSeconds);
+		}
 	}
 }
 
@@ -311,6 +327,19 @@ bool UEditorWorldExtensionCollection::InputAxis( FEditorViewportClient* InViewpo
 }
 
 
+void UEditorWorldExtensionCollection::ShowAllActors(const bool bShow)
+{
+	for (FEditorExtensionTuple& EditorExtensionTuple : EditorExtensions)
+	{
+		UEditorWorldExtension* EditorExtension = EditorExtensionTuple.Get<0>();
+		for (AActor* Actor : EditorExtension->ExtensionActors)
+		{
+			Actor->SetActorHiddenInGame(!bShow);
+			Actor->SetActorEnableCollision(bShow);
+		}
+	}
+}
+
 void UEditorWorldExtensionCollection::PostPIEStarted( bool bIsSimulatingInEditor )
 {
 	if( bIsSimulatingInEditor && GEditor->EditorWorld != nullptr && GEditor->EditorWorld == WorldContext->World() )
@@ -327,6 +356,19 @@ void UEditorWorldExtensionCollection::PostPIEStarted( bool bIsSimulatingInEditor
 }
 
 
+void UEditorWorldExtensionCollection::OnPreEndPIE(bool bWasSimulatingInEditor)
+{
+	if (!bWasSimulatingInEditor && EditorWorldOnSimulate != nullptr && EditorWorldOnSimulate->World() == GEditor->EditorWorld)
+	{
+		if (!GIsRequestingExit)
+		{
+			// Revert back to the editor world before closing the play world, otherwise actors and objects will be destroyed.
+			SetWorldContext(EditorWorldOnSimulate);
+			EditorWorldOnSimulate = nullptr;
+		}
+	}
+}
+
 void UEditorWorldExtensionCollection::OnEndPIE( bool bWasSimulatingInEditor )
 {
 	if( bWasSimulatingInEditor && EditorWorldOnSimulate != nullptr && EditorWorldOnSimulate->World() == GEditor->EditorWorld )
diff --git a/Engine/Source/Editor/UnrealEd/Private/PlayLevel.cpp b/Engine/Source/Editor/UnrealEd/Private/PlayLevel.cpp
--- a/Engine/Source/Editor/UnrealEd/Private/PlayLevel.cpp
+++ b/Engine/Source/Editor/UnrealEd/Private/PlayLevel.cpp
@@ -3515,6 +3515,8 @@ void UEditorEngine::ToggleBetweenPIEandSIE( bool bNewSession )
 {
 	bIsToggleBetweenPIEandSIEQueued = false;
 
+	FEditorDelegates::OnPreSwitchBeginPIEAndSIE.Broadcast(bIsSimulatingInEditor);
+
 	// The first PIE world context is the one that can toggle between PIE and SIE
 	// Network PIE/SIE toggling is not really meant to be supported.
 	FSlatePlayInEditorInfo * SlateInfoPtr = nullptr;
diff --git a/Engine/Source/Editor/UnrealEd/Public/Editor.h b/Engine/Source/Editor/UnrealEd/Public/Editor.h
--- a/Engine/Source/Editor/UnrealEd/Public/Editor.h
+++ b/Engine/Source/Editor/UnrealEd/Public/Editor.h
@@ -146,6 +146,8 @@ struct UNREALED_API FEditorDelegates
 	static FOnPIEEvent ResumePIE;
 	/** Sent when a PIE session is single-stepped */
 	static FOnPIEEvent SingleStepPIE;
+	/** Sent just before the user switches between from PIE to SIE, or vice-versa.  Passes in whether we are currently in SIE */
+	static FOnPIEEvent OnPreSwitchBeginPIEAndSIE;
 	/** Sent after the user switches between from PIE to SIE, or vice-versa.  Passes in whether we are currently in SIE */
 	static FOnPIEEvent OnSwitchBeginPIEAndSIE;
 	/** Within a property window, the currently selected item was changed.*/
diff --git a/Engine/Source/Editor/UnrealEd/Public/EditorWorldExtension.h b/Engine/Source/Editor/UnrealEd/Public/EditorWorldExtension.h
--- a/Engine/Source/Editor/UnrealEd/Public/EditorWorldExtension.h
+++ b/Engine/Source/Editor/UnrealEd/Public/EditorWorldExtension.h
@@ -53,6 +53,12 @@ public:
 	/** Destroys a transient actor we created earlier */
 	void DestroyTransientActor(AActor* Actor);
 
+	/** Sets if this extension should be ticked. */
+	void SetActive(const bool bInActive);
+
+	/** If this extension is currently being ticked. */
+	bool IsActive() const;
+
 	/** Get the owning collection of extensions */
 	UEditorWorldExtensionCollection* GetOwningCollection();
 
@@ -64,10 +70,10 @@ protected:
 	/** Reparent actors to a new world */
 	virtual void TransitionWorld(UWorld* NewWorld);
 
-	/** Give base class a chance to act on entering simulate mode */
+	/** Give child class a chance to act on entering simulate mode */
 	virtual void EnteredSimulateInEditor() {};
 	
-	/** Give base class a chance to act on leaving simulate mode */
+	/** Give child class a chance to act on leaving simulate mode */
 	virtual void LeftSimulateInEditor() {};
 
 	/** The collection of extensions that is owning this extension */
@@ -75,6 +81,7 @@ protected:
 
 private:
 
+	/** Reparent the actors to a new world. */
 	void ReparentActor(AActor* Actor, UWorld* NewWorld);
 
 	/** Let the FEditorWorldExtensionCollection set the world of this extension before init */
@@ -82,6 +89,9 @@ private:
 
 	UPROPERTY()
 	TArray<AActor*> ExtensionActors;
+
+	/** If this extension is currently being ticked */
+	bool bActive;
 };
 
 /**
@@ -142,6 +152,9 @@ public:
 	/** Notifies all extensions of axis movement */
 	bool InputAxis( FEditorViewportClient* InViewportClient, FViewport* Viewport, int32 ControllerId, FKey Key, float Delta, float DeltaTime);
 	
+	/** Show or hide all the actors of extensions that belong to this collection. */
+	void ShowAllActors(const bool bShow);
+
 private:
 
 	/** Sets the world for this collection and gives every extension an opportunity to transition */
@@ -150,6 +163,9 @@ private:
 	/** Called by the editor after PIE or Simulate is started */
 	void PostPIEStarted( bool bIsSimulatingInEditor );
 
+	/** Called just before PIE or Simulate ends */
+	void OnPreEndPIE(bool bWasSimulatingInEditor);
+
 	/** Called when PIE or Simulate ends */
 	void OnEndPIE( bool bWasSimulatingInEditor );
 
diff --git a/Engine/Source/Editor/VREditor/Public/VREditorMode.h b/Engine/Source/Editor/VREditor/Public/VREditorMode.h
--- a/Engine/Source/Editor/VREditor/Public/VREditorMode.h
+++ b/Engine/Source/Editor/VREditor/Public/VREditorMode.h
@@ -120,9 +120,6 @@ public:
 	/** If the mode was completely initialized */
 	bool IsFullyInitialized() const;
 
-	/** If the mode is currently running */
-	bool IsActive() const;
-	
 	/** * Gets the tick handle to give external systems the change to be ticked right after the ViewportWorldInteraction is ticked */
 	DECLARE_EVENT_OneParam(UVREditorMode, FOnVRTickHandle, const float /* DeltaTime */);
 	FOnVRTickHandle& OnTickHandle()
@@ -314,6 +311,8 @@ private:
 	void PostPIEStarted( bool bIsSimulatingInEditor );
 	void PrePIEEnded( bool bWasSimulatingInEditor );
 	void OnEndPIE( bool bWasSimulatingInEditor );
+	void OnPreSwitchPIEAndSIE(bool bIsSimulatingInEditor);
+	void OnSwitchPIEAndSIE(bool bIsSimulatingInEditor);
 
 	/** Start using the viewport passed */
 	void StartViewport( TSharedPtr<SLevelViewport> Viewport );
@@ -321,6 +320,9 @@ private:
 	/** Close the current viewport */
 	void CloseViewport( const bool bShouldDisableStereo );
 
+	/** Resets all the settings when exiting PIE to VR Editor. */
+	void RestoreFromPIE();
+
 	/** Restore the world to meters to the saved one when entering VR Editor */
 	void RestoreWorldToMeters();
 
@@ -470,9 +472,6 @@ private:
 	/** If this is the first tick or before */
 	bool bFirstTick;
 
-	/** If this current mode is running */
-	bool bIsActive;
-
 	/** Pointer to the current Sequencer */
 	class ISequencer* CurrentSequencer;
 
@@ -481,6 +480,9 @@ private:
 
 	/** If we started play in editor from the VR Editor*/
 	bool bStartedPlayFromVREditor;	
+
+	/** If we started play in editor from the VR Editor while in simulate. */
+	bool bStartedPlayFromVREditorSimulate;
 	
 	/** Container of assets */
 	UPROPERTY()
diff --git a/Engine/Source/Editor/VREditor/UI/VRRadialMenuHandler.cpp b/Engine/Source/Editor/VREditor/UI/VRRadialMenuHandler.cpp
--- a/Engine/Source/Editor/VREditor/UI/VRRadialMenuHandler.cpp
+++ b/Engine/Source/Editor/VREditor/UI/VRRadialMenuHandler.cpp
@@ -679,7 +679,7 @@ void UVRRadialMenuHandler::ToolsMenuGenerator(FMenuBuilder MenuBuilder, TSharedP
 		FUIAction
 		(
 			FExecuteAction::CreateStatic(&FVREditorActionCallbacks::OnPlayButtonClicked, VRMode),
-			FCanExecuteAction::CreateStatic(&FVREditorActionCallbacks::CanPlay)
+			FCanExecuteAction::CreateStatic(&FVREditorActionCallbacks::CanPlay, VRMode)
 			),
 		NAME_None,
 		EUserInterfaceActionType::CollapsedButton
diff --git a/Engine/Source/Editor/VREditor/VREditorActions.cpp b/Engine/Source/Editor/VREditor/VREditorActions.cpp
--- a/Engine/Source/Editor/VREditor/VREditorActions.cpp
+++ b/Engine/Source/Editor/VREditor/VREditorActions.cpp
@@ -296,9 +296,10 @@ void FVREditorActionCallbacks::OnPlayButtonClicked(UVREditorMode* InVRMode)
 	InVRMode->TogglePIEAndVREditor();
 }
 
-bool FVREditorActionCallbacks::CanPlay()
+bool FVREditorActionCallbacks::CanPlay(UVREditorMode* InVRMode)
 {
-	return FLevelEditorActionCallbacks::DefaultCanExecuteAction() && VREd::AllowPlay->GetInt() == 1;
+	return FLevelEditorActionCallbacks::DefaultCanExecuteAction() && VREd::AllowPlay->GetInt() == 1 &&
+		(InVRMode->GetHMDDeviceType() != EHMDDeviceType::DT_OculusRift || (InVRMode->GetHMDDeviceType() == EHMDDeviceType::DT_OculusRift && GEditor != nullptr && !GEditor->bIsSimulatingInEditor));
 }
 
 void FVREditorActionCallbacks::OnSimulateButtonClicked(UVREditorMode* InVRMode)
diff --git a/Engine/Source/Editor/VREditor/VREditorActions.h b/Engine/Source/Editor/VREditor/VREditorActions.h
--- a/Engine/Source/Editor/VREditor/VREditorActions.h
+++ b/Engine/Source/Editor/VREditor/VREditorActions.h
@@ -124,7 +124,7 @@ public:
 	/**
 	 * If we are allowed to enter play.
 	 */
-	static bool CanPlay();
+	static bool CanPlay(UVREditorMode* InVRMode);
 
 	/**
 	 * Enters Simulate mode for physics and animation playback
diff --git a/Engine/Source/Editor/VREditor/VREditorMode.cpp b/Engine/Source/Editor/VREditor/VREditorMode.cpp
--- a/Engine/Source/Editor/VREditor/VREditorMode.cpp
+++ b/Engine/Source/Editor/VREditor/VREditorMode.cpp
@@ -38,6 +38,7 @@
 #include "SequencerSettings.h"
 #include "Kismet/GameplayStatics.h"
 #include "Components/AudioComponent.h"
+#include "Components/StaticMeshComponent.h"
 #include "VREditorActions.h"
 #include "EditorModes.h"
 #include "VRModeSettings.h"
@@ -81,9 +82,9 @@ UVREditorMode::UVREditorMode() :
 	LeftHandInteractor( nullptr ),
 	RightHandInteractor( nullptr ),
 	bFirstTick( true ),
-	bIsActive( false ),
 	SavedWorldToMetersScaleForPIE(0.0f),
 	bStartedPlayFromVREditor(false),
+	bStartedPlayFromVREditorSimulate(false),
 	AssetContainer(nullptr)
 {
 }
@@ -170,6 +171,9 @@ void UVREditorMode::Enter()
 	FEditorDelegates::PostPIEStarted.AddUObject( this, &UVREditorMode::PostPIEStarted );
 	FEditorDelegates::PrePIEEnded.AddUObject( this, &UVREditorMode::PrePIEEnded );
 	FEditorDelegates::EndPIE.AddUObject(this, &UVREditorMode::OnEndPIE);
+	FEditorDelegates::OnPreSwitchBeginPIEAndSIE.AddUObject(this, &UVREditorMode::OnPreSwitchPIEAndSIE);
+	FEditorDelegates::OnSwitchBeginPIEAndSIE.AddUObject(this, &UVREditorMode::OnSwitchPIEAndSIE);
+
 
 	// @todo vreditor: We need to make sure the user can never switch to orthographic mode, or activate settings that
 	// would disrupt the user's ability to view the VR scene.
@@ -284,7 +288,9 @@ void UVREditorMode::Enter()
 	GEditor->OnEditorClose().AddUObject( this, &UVREditorMode::OnEditorClosed );
 
 	bFirstTick = true;
-	bIsActive = true;
+	SetActive(true);
+	bStartedPlayFromVREditor = false;
+	bStartedPlayFromVREditorSimulate = false;
 }
 
 void UVREditorMode::Exit(const bool bShouldDisableStereo)
@@ -389,17 +395,19 @@ void UVREditorMode::Exit(const bool bShouldDisableStereo)
 	FEditorDelegates::PostPIEStarted.RemoveAll( this );
 	FEditorDelegates::PrePIEEnded.RemoveAll( this );
 	FEditorDelegates::EndPIE.RemoveAll( this );
+	FEditorDelegates::OnPreSwitchBeginPIEAndSIE.RemoveAll(this);
+	FEditorDelegates::OnSwitchBeginPIEAndSIE.RemoveAll(this);
 
 	GEditor->OnEditorClose().RemoveAll( this );
 
 	bWantsToExitMode = false;
-	bIsActive = false;
+	SetActive(false);
 	bFirstTick = false;
 }
 
 void UVREditorMode::OnEditorClosed()
 {
-	if( bIsActive )
+	if(IsActive())
 	{
 		Exit( false );
 		Shutdown();
@@ -418,7 +426,7 @@ void UVREditorMode::OnVREditorWindowClosed( const TSharedRef<SWindow>& ClosedWin
 
 void UVREditorMode::PreTick( const float DeltaTime )
 {
-	if( !bIsFullyInitialized || !bIsActive || bWantsToExitMode )
+	if( !bIsFullyInitialized || !IsActive() || bWantsToExitMode )
 	{
 		return;
 	}
@@ -447,7 +455,7 @@ void UVREditorMode::PreTick( const float DeltaTime )
 
 void UVREditorMode::PostTick( float DeltaTime )
 {
-	if( !bIsFullyInitialized || !bIsActive || bWantsToExitMode || !VREditorLevelViewportWeakPtr.IsValid() )
+	if( !bIsFullyInitialized || !IsActive() || bWantsToExitMode || !VREditorLevelViewportWeakPtr.IsValid() )
 	{
 		return;
 	}
@@ -574,11 +582,6 @@ bool UVREditorMode::IsFullyInitialized() const
 	return bIsFullyInitialized;
 }
 
-bool UVREditorMode::IsActive() const
-{
-	return bIsActive;
-}
-
 bool UVREditorMode::IsShowingRadialMenu(const UVREditorInteractor* Interactor) const
 {
 	return UISystem->IsShowingRadialMenu(Interactor);
@@ -790,7 +793,40 @@ void UVREditorMode::ToggleSIEAndVREditor()
 
 void UVREditorMode::TogglePIEAndVREditor()
 {
-	if (GEditor->EditorWorld == nullptr && GEditor->PlayWorld == nullptr)
+	bool bRequestedPIE = false;
+	if (GEditor->EditorWorld == nullptr && GEditor->PlayWorld == nullptr && !GEditor->bIsSimulatingInEditor)
+	{
+		const FVector* StartLoc = NULL;
+		const FRotator* StartRot = NULL;
+		const bool bHMDIsReady = (GEngine && GEngine->HMDDevice.IsValid() && GEngine->HMDDevice->IsHMDConnected());
+		GEditor->RequestPlaySession(true, VREditorLevelViewportWeakPtr.Pin(), false /*bSimulateInEditor*/, StartLoc, StartRot, -1, false, bHMDIsReady);
+		bRequestedPIE = true;
+	}
+	else if (GEditor->PlayWorld != nullptr)
+	{
+		// Since we are already in simulate, we want to toggle to PIE.
+		if (GEditor->bIsSimulatingInEditor)
+		{
+			bStartedPlayFromVREditorSimulate = true;
+			bRequestedPIE = true;
+
+			GEditor->RequestToggleBetweenPIEandSIE();
+		}
+		else
+		{
+			// If this play started while in simulate, then toggle back to simulate.
+			if (bStartedPlayFromVREditorSimulate)
+			{
+				GEditor->RequestToggleBetweenPIEandSIE();
+			}
+			else
+			{
+				GEditor->RequestEndPlayMap();
+			}
+		}
+	}
+
+	if (bRequestedPIE)
 	{
 		// Turn off input processing while in PIE.  We don't want any input events until the user comes back to the editor
 		WorldInteraction->SetUseInputPreprocessor(false);
@@ -800,17 +836,9 @@ void UVREditorMode::TogglePIEAndVREditor()
 		// Restore the world to meters before entering play
 		RestoreWorldToMeters();
 
-		bIsActive = false;
+		SetActive(false);
+		WorldInteraction->SetActive(false);
 		bStartedPlayFromVREditor = true;
-
-		const FVector* StartLoc = NULL;
-		const FRotator* StartRot = NULL;
-		const bool bHMDIsReady = (GEngine && GEngine->HMDDevice.IsValid() && GEngine->HMDDevice->IsHMDConnected());
-		GEditor->RequestPlaySession(true, VREditorLevelViewportWeakPtr.Pin(), false /*bSimulateInEditor*/, StartLoc, StartRot, -1, false, bHMDIsReady);
-	}
-	else if (GEditor->PlayWorld != nullptr && !GEditor->bIsSimulatingInEditor)
-	{
-		GEditor->RequestEndPlayMap();
 	}
 }
 
@@ -1048,6 +1076,20 @@ void UVREditorMode::CloseViewport( const bool bShouldDisableStereo )
 	}
 }
 
+void UVREditorMode::RestoreFromPIE()
+{
+	SetActive(true);
+	bStartedPlayFromVREditor = false;
+	bStartedPlayFromVREditorSimulate = false;
+
+	GetWorld()->GetWorldSettings()->WorldToMeters = SavedWorldToMetersScaleForPIE;
+	WorldInteraction->SetWorldToMetersScale(SavedWorldToMetersScaleForPIE);
+
+	// Re-enable input pre-processing
+	WorldInteraction->SetUseInputPreprocessor(true);
+	WorldInteraction->SetActive(true);
+}
+
 void UVREditorMode::RestoreWorldToMeters()
 {
 	const float DefaultWorldToMeters = VREd::DefaultWorldToMeters->GetFloat();
@@ -1123,7 +1165,7 @@ const UVREditorAssetContainer& UVREditorMode::GetAssetContainer() const
 
 void UVREditorMode::PlaySound(USoundBase* SoundBase, const FVector& InWorldLocation, const float InVolume /*= 1.0f*/)
 {
-	if (bIsActive && bIsFullyInitialized && GEditor != nullptr && GEditor->CanPlayEditorSound())
+	if (IsActive() && bIsFullyInitialized && GEditor != nullptr && GEditor->CanPlayEditorSound())
 	{
 		const float Volume = InVolume*VREd::SFXMultiplier->GetFloat();
 		UGameplayStatics::PlaySoundAtLocation(GetWorld(), SoundBase, InWorldLocation, Volume);
@@ -1132,7 +1174,7 @@ void UVREditorMode::PlaySound(USoundBase* SoundBase, const FVector& InWorldLocat
 
 void UVREditorMode::PostPIEStarted( bool bIsSimulatingInEditor )
 {
-	if( !bIsSimulatingInEditor )
+	if (!bIsSimulatingInEditor && bStartedPlayFromVREditor)
 	{
 		GEnableVREditorHacks = false;
 	}
@@ -1141,7 +1183,7 @@ void UVREditorMode::PostPIEStarted( bool bIsSimulatingInEditor )
 
 void UVREditorMode::PrePIEEnded( bool bWasSimulatingInEditor )
 {
-	if( !bWasSimulatingInEditor )
+	if (!bWasSimulatingInEditor && bStartedPlayFromVREditor && !bStartedPlayFromVREditorSimulate)
 	{
 		GEnableVREditorHacks = true;
 	}
@@ -1149,17 +1191,45 @@ void UVREditorMode::PrePIEEnded( bool bWasSimulatingInEditor )
 
 void UVREditorMode::OnEndPIE(bool bWasSimulatingInEditor)
 {
-	// We can not use this for leaving simulate. We need to use UEditorWorldExtension::LeftSimulateInEditor for that to be sure the world has already been transistioned.
-	if (!bWasSimulatingInEditor && bStartedPlayFromVREditor)
+	if (bStartedPlayFromVREditor && !bWasSimulatingInEditor && !bStartedPlayFromVREditorSimulate)
 	{
-		bIsActive = true;
-		bStartedPlayFromVREditor = false;
+		RestoreFromPIE();
+	}
+}
 
-		GetWorld()->GetWorldSettings()->WorldToMeters = SavedWorldToMetersScaleForPIE;
-		WorldInteraction->SetWorldToMetersScale(SavedWorldToMetersScaleForPIE);
+void UVREditorMode::OnPreSwitchPIEAndSIE(bool bIsSimulatingInEditor)
+{
+	if (bStartedPlayFromVREditorSimulate)
+	{
+		if (bIsSimulatingInEditor)
+		{
+			// Pre SIE to PIE
+			GetOwningCollection()->ShowAllActors(false);
+		}
+		else
+		{
+			// Pre PIE to SIE
+			GEnableVREditorHacks = true;
+		}
+	}
+}
 
-		// Re-enable input pre-processing
-		WorldInteraction->SetUseInputPreprocessor(true);
+void UVREditorMode::OnSwitchPIEAndSIE(bool bIsSimulatingInEditor)
+{
+	if (bStartedPlayFromVREditorSimulate)
+	{
+		if (bIsSimulatingInEditor)
+		{
+			// Post PIE to SIE
+			RestoreFromPIE();
+			GetOwningCollection()->ShowAllActors(true);
+		}
+		else
+		{
+			// Post SIE to PIE
+			GEnableVREditorHacks = false;
+			FSlateApplication::Get().SetAllUserFocusToGameViewport();
+		}
 	}
 }
 
diff --git a/Engine/Source/Editor/VREditor/VREditorModeManager.cpp b/Engine/Source/Editor/VREditor/VREditorModeManager.cpp
--- a/Engine/Source/Editor/VREditor/VREditorModeManager.cpp
+++ b/Engine/Source/Editor/VREditor/VREditorModeManager.cpp
@@ -6,6 +6,7 @@
 #include "Modules/ModuleManager.h"
 #include "GameFramework/PlayerController.h"
 #include "GameFramework/WorldSettings.h"
+#include "GameFramework/PlayerInput.h"
 #include "Editor.h"
 
 #include "EngineGlobals.h"
@@ -87,6 +88,9 @@ void FVREditorModeManager::Tick( const float DeltaTime )
 				PlayerController->GetInputKeyTimeDown( EKeys::MotionController_Left_Trigger ) > ShutDownInputKeyTime)
 			{
 				CurrentVREditorMode->TogglePIEAndVREditor();
+				// We need to clear the input of the playercontroller when exiting PIE. 
+				// Otherwise the input will still be pressed down causing toggle between PIE and VR Editor to be called instantly whenever entering PIE a second time.
+				PlayerController->PlayerInput->FlushPressedKeys();
 				break;
 			}
 		}
diff --git a/Engine/Source/Editor/ViewportInteraction/Public/ViewportWorldInteraction.h b/Engine/Source/Editor/ViewportInteraction/Public/ViewportWorldInteraction.h
--- a/Engine/Source/Editor/ViewportInteraction/Public/ViewportWorldInteraction.h
+++ b/Engine/Source/Editor/ViewportInteraction/Public/ViewportWorldInteraction.h
@@ -54,9 +54,6 @@ public:
 	virtual void Shutdown() override;
 	virtual void Tick( float DeltaSeconds ) override;
 	
-	/** Returns true if this worldinteraction is currently running */
-	bool IsActive() const;
-
 	/** Adds interactor to the worldinteraction */
 	void AddInteractor( UViewportInteractor* Interactor );
 
@@ -699,9 +696,6 @@ private:
 	/** Event that is fired when the world scale changes */
 	FOnWorldScaleChanged OnWorldScaleChangedEvent;
 
-	/** If the world interaction is active and running the tick function */
-	bool bActive;
-
 	/** If this world interaction should get input from the input processor */
 	bool bUseInputPreprocessor;
 
diff --git a/Engine/Source/Editor/ViewportInteraction/ViewportWorldInteraction.cpp b/Engine/Source/Editor/ViewportInteraction/ViewportWorldInteraction.cpp
--- a/Engine/Source/Editor/ViewportInteraction/ViewportWorldInteraction.cpp
+++ b/Engine/Source/Editor/ViewportInteraction/ViewportWorldInteraction.cpp
@@ -286,7 +286,6 @@ UViewportWorldInteraction::UViewportWorldInteraction():
 	DefaultMouseCursorInteractor( nullptr ),
 	DefaultMouseCursorInteractorRefCount( 0 ),
 	bIsInVR( false ),
-	bActive( false ),
 	bUseInputPreprocessor( false ),
 	bAllowWorldMovement( true ),
 	CurrentDeltaTime(0.0f),
@@ -336,7 +335,7 @@ void UViewportWorldInteraction::Init()
 	// We need a mouse cursor!
 	this->AddMouseCursorInteractor();
 
-	bActive = true;
+	SetActive(true);
 
 	CandidateActors.Reset();
 
@@ -348,7 +347,7 @@ void UViewportWorldInteraction::Init()
 
 void UViewportWorldInteraction::Shutdown()
 {
-	bActive = false;
+	SetActive(false);
 
 	if( DefaultMouseCursorInteractorRefCount == 1 )
 	{
@@ -421,7 +420,7 @@ void UViewportWorldInteraction::LeftSimulateInEditor()
 
 void UViewportWorldInteraction::OnEditorClosed()
 {
-	if( bActive )
+	if( IsActive() )
 	{
 		Shutdown();
 	}
@@ -533,11 +532,6 @@ void UViewportWorldInteraction::SetDefaultOptionalViewportClient(const TSharedPt
 	DefaultOptionalViewportClient = InEditorViewportClient.Get();
 }
 
-bool UViewportWorldInteraction::IsActive() const
-{
-	return bActive;
-}
-
 void UViewportWorldInteraction::PairInteractors( UViewportInteractor* FirstInteractor, UViewportInteractor* SecondInteractor )
 {
 	FirstInteractor->SetOtherInteractor( SecondInteractor );
@@ -548,19 +542,19 @@ bool UViewportWorldInteraction::InputKey( FEditorViewportClient* InViewportClien
 {
 	bool bWasHandled = false;
 
-	if( bActive )
+	if( IsActive() )
 	{
 		check(InViewportClient != nullptr);
 		OnKeyInputEvent.Broadcast(InViewportClient, Key, Event, bWasHandled);
 
-		if( !bWasHandled || !bActive )
+		if( !bWasHandled || !IsActive() )
 		{
 			for( UViewportInteractor* Interactor : Interactors )
 			{
 				bWasHandled = Interactor->HandleInputKey( *InViewportClient, Key, Event );
 
 				// Stop iterating if the input was handled by an Interactor
-				if( bWasHandled || !bActive )
+				if( bWasHandled || !IsActive() )
 				{
 					break;
 				}
@@ -596,19 +590,19 @@ bool UViewportWorldInteraction::InputAxis( FEditorViewportClient* InViewportClie
 {
 	bool bWasHandled = false;
 
-	if( bActive )
+	if( IsActive() )
 	{
 		check( InViewportClient != nullptr );
 		OnAxisInputEvent.Broadcast( InViewportClient, ControllerId, Key, Delta, DeltaTime, bWasHandled );
 
-		if( !bWasHandled || !bActive )
+		if( !bWasHandled || !IsActive() )
 		{
 			for( UViewportInteractor* Interactor : Interactors )
 			{
 				bWasHandled = Interactor->HandleInputAxis( *InViewportClient, Key, Delta, DeltaTime );
 
 				// Stop iterating if the input was handled by an interactor
-				if( bWasHandled || !bActive )
+				if( bWasHandled || !IsActive() )
 				{
 					break;
 				}
@@ -3084,7 +3078,7 @@ const class UViewportInteractionAssetContainer& UViewportWorldInteraction::LoadA
 
 void UViewportWorldInteraction::PlaySound(USoundBase* SoundBase, const FVector& InWorldLocation, const float InVolume /*= 1.0f*/)
 {
-	if (bActive && GEditor != nullptr && GEditor->CanPlayEditorSound())
+	if (IsActive() && GEditor != nullptr && GEditor->CanPlayEditorSound())
 	{
 		const float Volume = InVolume*VI::SFXMultiplier->GetFloat();
 		UGameplayStatics::PlaySoundAtLocation(GetWorld(), SoundBase, InWorldLocation, FRotator::ZeroRotator, Volume);
