VR Editor: Fix warning FBox: Use ForceInit constructor instead.

diff --git a/Engine/Source/Editor/VREditor/VREditorWorldInteraction.cpp b/Engine/Source/Editor/VREditor/VREditorWorldInteraction.cpp
--- a/Engine/Source/Editor/VREditor/VREditorWorldInteraction.cpp
+++ b/Engine/Source/Editor/VREditor/VREditorWorldInteraction.cpp
@@ -149,7 +149,7 @@ void UVREditorWorldInteraction::StartDraggingMaterialOrTexture( UViewportInterac
 		InteractorData.TransformGizmoInteractionType = ETransformGizmoInteractionType::None;
 		InteractorData.GizmoStartTransform = FTransform::Identity;
 		InteractorData.GizmoLastTransform = InteractorData.GizmoTargetTransform = InteractorData.GizmoUnsnappedTargetTransform = InteractorData.GizmoInterpolationSnapshotTransform = InteractorData.GizmoStartTransform;
-		InteractorData.GizmoStartLocalBounds = FBox( 0 );
+		InteractorData.GizmoStartLocalBounds = FBox(EForceInit::ForceInit);
 
 		InteractorData.GizmoSpaceFirstDragUpdateOffsetAlongAxis = FVector::ZeroVector;	// Will be determined on first update
 		InteractorData.GizmoSpaceDragDeltaFromStartOffset = FVector::ZeroVector;	// Set every frame while dragging
