VR Editor: Fix mac compiler issues

diff --git a/Engine/Source/Editor/FoliageEdit/Private/FoliageEdMode.cpp b/Engine/Source/Editor/FoliageEdit/Private/FoliageEdMode.cpp
--- a/Engine/Source/Editor/FoliageEdit/Private/FoliageEdMode.cpp
+++ b/Engine/Source/Editor/FoliageEdit/Private/FoliageEdMode.cpp
@@ -522,80 +522,75 @@ void FEdModeFoliage::OnVRAction(class FEditorViewportClient& ViewportClient, UVi
 				FVector LaserPointerStart, LaserPointerEnd;
 				if (Interactor->GetLaserPointer( /* Out */ LaserPointerStart, /* Out */ LaserPointerEnd))
 				{
-					// Go ahead and paint immediately
-					FVector LaserPointerStart, LaserPointerEnd;
-					if (Interactor->GetLaserPointer( /* Out */ LaserPointerStart, /* Out */ LaserPointerEnd))
+					const FVector LaserPointerDirection = (LaserPointerEnd - LaserPointerStart).GetSafeNormal();
+					BrushTraceDirection = LaserPointerDirection;
+
+					// Only start painting if we're not dragging a widget handle
+					if (ViewportClient.GetCurrentWidgetAxis() == EAxisList::None)
 					{
-						const FVector LaserPointerDirection = (LaserPointerEnd - LaserPointerStart).GetSafeNormal();
-						BrushTraceDirection = LaserPointerDirection;
+						if (UISettings.GetPaintToolSelected() || UISettings.GetReapplyToolSelected() || UISettings.GetLassoSelectToolSelected())
+						{
+							StartFoliageBrushTrace(&ViewportClient, Interactor);
+							FoliageBrushTrace(&ViewportClient, LaserPointerStart, LaserPointerDirection);
+						}
 
-						// Only start painting if we're not dragging a widget handle
-						if (ViewportClient.GetCurrentWidgetAxis() == EAxisList::None)
+						// Fill a static mesh with foliage brush
+						else if (UISettings.GetPaintBucketToolSelected() || UISettings.GetReapplyPaintBucketToolSelected())
 						{
-							if (UISettings.GetPaintToolSelected() || UISettings.GetReapplyToolSelected() || UISettings.GetLassoSelectToolSelected())
-							{
-								StartFoliageBrushTrace(&ViewportClient, Interactor);
-								FoliageBrushTrace(&ViewportClient, LaserPointerStart, LaserPointerDirection);
-							}
+							FHitResult HitResult = Interactor->GetHitResultFromLaserPointer();
 
-							// Fill a static mesh with foliage brush
-							else if (UISettings.GetPaintBucketToolSelected() || UISettings.GetReapplyPaintBucketToolSelected())
+							if (HitResult.Actor.Get() != nullptr)
 							{
-								FHitResult HitResult = Interactor->GetHitResultFromLaserPointer();
+								GEditor->BeginTransaction(NSLOCTEXT("UnrealEd", "FoliageMode_EditTransaction", "Foliage Editing"));
 
-								if (HitResult.Actor.Get() != nullptr)
+								if (IsModifierButtonPressed(&ViewportClient))
 								{
-									GEditor->BeginTransaction(NSLOCTEXT("UnrealEd", "FoliageMode_EditTransaction", "Foliage Editing"));
-
-									if (IsModifierButtonPressed(&ViewportClient))
-									{
-										ApplyPaintBucket_Remove(HitResult.Actor.Get());
-									}
-									else
-									{
-										ApplyPaintBucket_Add(HitResult.Actor.Get());
-									}
-
-									GEditor->EndTransaction();
+									ApplyPaintBucket_Remove(HitResult.Actor.Get());
 								}
+								else
+								{
+									ApplyPaintBucket_Add(HitResult.Actor.Get());
+								}
+
+								GEditor->EndTransaction();
 							}
-							// Select an instanced foliage
-							else if (UISettings.GetSelectToolSelected())
-							{
-								FHitResult HitResult = Interactor->GetHitResultFromLaserPointer();
+						}
+						// Select an instanced foliage
+						else if (UISettings.GetSelectToolSelected())
+						{
+							FHitResult HitResult = Interactor->GetHitResultFromLaserPointer();
 
-								GEditor->BeginTransaction(NSLOCTEXT("UnrealEd", "FoliageMode_EditTransaction", "Foliage Editing"));
+							GEditor->BeginTransaction(NSLOCTEXT("UnrealEd", "FoliageMode_EditTransaction", "Foliage Editing"));
 
-								if (HitResult.GetActor() != nullptr)
+							if (HitResult.GetActor() != nullptr)
+							{
+								// Clear all currently selected instances
+								SelectInstances(ViewportClient.GetWorld(), false);
+								for (auto& FoliageMeshUI : FoliageMeshList)
 								{
-									// Clear all currently selected instances
-									SelectInstances(ViewportClient.GetWorld(), false);
-									for (auto& FoliageMeshUI : FoliageMeshList)
-									{
-										UFoliageType* Settings = FoliageMeshUI->Settings;
-										SelectInstanceAtLocation(ViewportClient.GetWorld(), Settings, HitResult.ImpactPoint, !IsModifierButtonPressed(&ViewportClient));
-									}
+									UFoliageType* Settings = FoliageMeshUI->Settings;
+									SelectInstanceAtLocation(ViewportClient.GetWorld(), Settings, HitResult.ImpactPoint, !IsModifierButtonPressed(&ViewportClient));
 								}
+							}
 
-								GEditor->EndTransaction();
+							GEditor->EndTransaction();
 
-								// @todo vreditor: we currently don't have a key mapping scheme to snap selected instances to ground 
-								// SnapSelectedInstancesToGround(GetWorld());
+							// @todo vreditor: we currently don't have a key mapping scheme to snap selected instances to ground 
+							// SnapSelectedInstancesToGround(GetWorld());
 
-							}
 						}
 					}
 				}
 			}
 			// Stop current tracking if the user is no longer painting
 			else if (Action.Event == IE_Released && FoliageInteractor && (FoliageInteractor == Interactor))
-				{
-					EndFoliageBrushTrace();
-					FoliageInteractor = nullptr;
+			{
+				EndFoliageBrushTrace();
+				FoliageInteractor = nullptr;
 
-					bWasHandled = true;
-					bOutIsInputCaptured = false;
-				}
+				bWasHandled = true;
+				bOutIsInputCaptured = false;
+			}
 		}
 	}
 }
diff --git a/Engine/Source/Editor/VREditor/UI/VREditorRadialFloatingUI.cpp b/Engine/Source/Editor/VREditor/UI/VREditorRadialFloatingUI.cpp
--- a/Engine/Source/Editor/VREditor/UI/VREditorRadialFloatingUI.cpp
+++ b/Engine/Source/Editor/VREditor/UI/VREditorRadialFloatingUI.cpp
@@ -510,8 +510,8 @@ const void AVREditorRadialFloatingUI::HighlightSlot(const FVector2D& TrackpadPos
 		{
 			if (ButtonCount != Index)
 			{
-				TSharedRef<SWidget> CurrentChild = WidgetComponents[ButtonCount]->GetSlateWidget().ToSharedRef();
-				TestWidget = UVREditorUISystem::FindWidgetOfType(CurrentChild, ButtonTypeOverride);
+				TSharedRef<SWidget> ChildWidget = WidgetComponents[ButtonCount]->GetSlateWidget().ToSharedRef();
+				TestWidget = UVREditorUISystem::FindWidgetOfType(ChildWidget, ButtonTypeOverride);
 				TSharedRef<SButton> TestButton = StaticCastSharedRef<SButton>(TestWidget);
 				if (TestButton->IsHovered())
 				{
diff --git a/Engine/Source/Editor/VREditor/UI/VREditorUISystem.cpp b/Engine/Source/Editor/VREditor/UI/VREditorUISystem.cpp
--- a/Engine/Source/Editor/VREditor/UI/VREditorUISystem.cpp
+++ b/Engine/Source/Editor/VREditor/UI/VREditorUISystem.cpp
@@ -91,6 +91,7 @@
 #include "LevelEditorActions.h"
 #include "VREditorActions.h"
 #include "GenericCommands.h"
+#include "Components/StaticMeshComponent.h"
 
 #include "ISequencer.h"
 #include "Engine/StaticMesh.h"
@@ -338,226 +339,191 @@ void UVREditorUISystem::OnPreviewInputAction(FEditorViewportClient& ViewportClie
 
 	// Laser Interaction Preview actions
 	if (VREditorInteractor &&
-		VREditorInteractor == LaserInteractor)
+		VREditorInteractor == LaserInteractor && 
+		Action.ActionType == ViewportWorldActionTypes::SelectAndMove &&
+		VREditorInteractor->GetDraggingMode() == EViewportInteractionDraggingMode::Nothing)
 	{
-
-		if (Action.ActionType == ViewportWorldActionTypes::SelectAndMove &&
-			VREditorInteractor->GetDraggingMode() == EViewportInteractionDraggingMode::Nothing)
+		FVector LaserPointerStart, LaserPointerEnd;
+		// If we are clicking on an Actor but not a widget component, send a fake mouse click event to toggle focus
+		if (VREditorInteractor->GetLaserPointer(LaserPointerStart, LaserPointerEnd))
 		{
-			FVector LaserPointerStart, LaserPointerEnd;
-			// If we are clicking on an Actor but not a widget component, send a fake mouse click event to toggle focus
-			if (VREditorInteractor->GetLaserPointer(LaserPointerStart, LaserPointerEnd))
+			FHitResult HitResult = VREditorInteractor->GetHitResultFromLaserPointer();
+			if (HitResult.Actor.IsValid())
 			{
-				FHitResult HitResult = VREditorInteractor->GetHitResultFromLaserPointer();
-				if (HitResult.Actor.IsValid())
+				UWidgetComponent* WidgetComponent = Cast<UWidgetComponent>(HitResult.GetComponent());
+
+				if (WidgetComponent == nullptr)
 				{
-					UWidgetComponent* WidgetComponent = Cast<UWidgetComponent>(HitResult.GetComponent());
+					// If we didn't handle the input in any other way, send an empty mouse down event so Slate focus is handled correctly
+					const bool bIsRightClicking =
+						(Action.Event == IE_Pressed && VREditorInteractor->IsModifierPressed()) ||
+						(Action.Event == IE_Released && VREditorInteractor->IsRightClickingOnUI());
+					TSet<FKey> PressedButtons;
+					FPointerEvent PointerEvent(
+						1 + (uint8)VREditorInteractor->GetControllerSide(),
+						FVector2D::ZeroVector,
+						FVector2D::ZeroVector,
+						PressedButtons,
+						bIsRightClicking ? EKeys::RightMouseButton : EKeys::LeftMouseButton,
+						0.0f,	// Wheel delta
+						FModifierKeysState());
 
-					if (WidgetComponent == nullptr)
+					FWidgetPath EmptyWidgetPath;
+					FReply Reply = FSlateApplication::Get().RoutePointerDownEvent(EmptyWidgetPath, PointerEvent);
+					if (bRadialMenuIsNumpad)
 					{
-						const bool bIsRightClicking =
-							(Action.Event == IE_Pressed && VREditorInteractor->IsModifierPressed()) ||
-							(Action.Event == IE_Released && VREditorInteractor->IsRightClickingOnUI());
-						TSet<FKey> PressedButtons;
-						FPointerEvent PointerEvent(
-							1 + (uint8)VREditorInteractor->GetControllerSide(),
-							FVector2D::ZeroVector,
-							FVector2D::ZeroVector,
-							PressedButtons,
-							bIsRightClicking ? EKeys::RightMouseButton : EKeys::LeftMouseButton,
-							0.0f,	// Wheel delta
-							FModifierKeysState());
-
-						FWidgetPath EmptyWidgetPath;
-						FReply Reply = FSlateApplication::Get().RoutePointerDownEvent(EmptyWidgetPath, PointerEvent);
-						if (bRadialMenuIsNumpad)
+						// If clicking somewhere outside UI so the widget loses focus
+						SwapRadialMenu();
+						if (!bRadialMenuVisibleAtSwap)
 						{
-							// If clicking somewhere outside UI so the widget loses focus
-							SwapRadialMenu();
-							if (!bRadialMenuVisibleAtSwap)
-							{
-								HideRadialMenu(VREditorInteractor);
-							}
+							HideRadialMenu(VREditorInteractor);
 						}
 					}
 				}
-			}
-			// Clicking on UI
-			{
-				FVector LaserPointerStart, LaserPointerEnd;
-				if (VREditorInteractor->GetLaserPointer(LaserPointerStart, LaserPointerEnd))
+				else if (WidgetComponent != nullptr)
 				{
-					FHitResult HitResult = VREditorInteractor->GetHitResultFromLaserPointer();
-					if (HitResult.Actor.IsValid())
-					{
-						// Only allow clicks to our own widget components
-						UVREditorWidgetComponent* WidgetComponent = Cast<UVREditorWidgetComponent>(HitResult.GetComponent());
-						if (WidgetComponent != nullptr)
-						{
-							// Always mark the event as handled so that the editor doesn't try to select the widget component
-							bWasHandled = true;
+					// Only allow clicks to our own widget components
+					// Always mark the event as handled so that the editor doesn't try to select the widget component
+					bWasHandled = true;
 
-							if (Action.Event != IE_Repeat)
-							{
-								// If the Modifier button is held down, treat this like a right click instead of a left click
-								const bool bIsRightClicking =
-									(Action.Event == IE_Pressed && VREditorInteractor->IsModifierPressed()) ||
-									(Action.Event == IE_Released && VREditorInteractor->IsRightClickingOnUI());
+					if (Action.Event != IE_Repeat)
+					{
+						// If the Modifier button is held down, treat this like a right click instead of a left click
+						const bool bIsRightClicking =
+							(Action.Event == IE_Pressed && VREditorInteractor->IsModifierPressed()) ||
+							(Action.Event == IE_Released && VREditorInteractor->IsRightClickingOnUI());
 
-								FVector2D LastLocalHitLocation = WidgetComponent->GetLastLocalHitLocation();
+						FVector2D LastLocalHitLocation = WidgetComponent->GetLastLocalHitLocation();
 
-								FVector2D LocalHitLocation;
-								WidgetComponent->GetLocalHitLocation(HitResult.ImpactPoint, LocalHitLocation);
+						FVector2D LocalHitLocation;
+						WidgetComponent->GetLocalHitLocation(HitResult.ImpactPoint, LocalHitLocation);
 
-								// If we weren't already hovering over this widget, then we'll reset the last hit location
-								if (WidgetComponent != VREditorInteractor->GetLastHoveredWidgetComponent())
-								{
-									LastLocalHitLocation = LocalHitLocation;
+						// If we weren't already hovering over this widget, then we'll reset the last hit location
+						if (WidgetComponent != VREditorInteractor->GetLastHoveredWidgetComponent())
+						{
+							LastLocalHitLocation = LocalHitLocation;
 
-									if (UVREditorWidgetComponent* VRWidgetComponent = Cast<UVREditorWidgetComponent>(VREditorInteractor->GetLastHoveredWidgetComponent()))
-									{
-										VRWidgetComponent->SetIsHovering(false);
-										OnHoverEndEffect(VRWidgetComponent);
-									}
-								}
+							if (UVREditorWidgetComponent* VRWidgetComponent = Cast<UVREditorWidgetComponent>(VREditorInteractor->GetLastHoveredWidgetComponent()))
+							{
+								VRWidgetComponent->SetIsHovering(false);
+								OnHoverEndEffect(VRWidgetComponent);
+							}
+						}
 
-								FWidgetPath WidgetPathUnderFinger = FWidgetPath(WidgetComponent->GetHitWidgetPath(HitResult.ImpactPoint, /*bIgnoreEnabledStatus*/ false));
-								if (WidgetPathUnderFinger.IsValid())
-								{
-									TSet<FKey> PressedButtons;
-									if (Action.Event == IE_Pressed)
-									{
-										PressedButtons.Add(bIsRightClicking ? EKeys::RightMouseButton : EKeys::LeftMouseButton);
-									}
+						FWidgetPath WidgetPathUnderFinger = FWidgetPath(WidgetComponent->GetHitWidgetPath(HitResult.ImpactPoint, /*bIgnoreEnabledStatus*/ false));
+						if (WidgetPathUnderFinger.IsValid())
+						{
+							TSet<FKey> PressedButtons;
+							if (Action.Event == IE_Pressed)
+							{
+								PressedButtons.Add(bIsRightClicking ? EKeys::RightMouseButton : EKeys::LeftMouseButton);
+							}
 
-									FPointerEvent PointerEvent(
-										1 + (uint8)VREditorInteractor->GetControllerSide(),
-										LocalHitLocation,
-										LastLocalHitLocation,
-										PressedButtons,
-										bIsRightClicking ? EKeys::RightMouseButton : EKeys::LeftMouseButton,
-										0.0f,	// Wheel delta
-										FModifierKeysState());
+							FPointerEvent PointerEvent(
+								1 + (uint8)VREditorInteractor->GetControllerSide(),
+								LocalHitLocation,
+								LastLocalHitLocation,
+								PressedButtons,
+								bIsRightClicking ? EKeys::RightMouseButton : EKeys::LeftMouseButton,
+								0.0f,	// Wheel delta
+								FModifierKeysState());
 
-									VREditorInteractor->SetLastHoveredWidgetComponent(WidgetComponent);
+							VREditorInteractor->SetLastHoveredWidgetComponent(WidgetComponent);
 
-									if (UVREditorWidgetComponent* VRWidgetComponent = Cast<UVREditorWidgetComponent>(VREditorInteractor->GetLastHoveredWidgetComponent()))
-									{
-										VRWidgetComponent->SetIsHovering(true);
-										OnHoverBeginEffect(VRWidgetComponent);
-									}
+							if (UVREditorWidgetComponent* VRWidgetComponent = Cast<UVREditorWidgetComponent>(VREditorInteractor->GetLastHoveredWidgetComponent()))
+							{
+								VRWidgetComponent->SetIsHovering(true);
+								OnHoverBeginEffect(VRWidgetComponent);
+							}
 
-									FReply Reply = FReply::Unhandled();
-									if (Action.Event == IE_Pressed)
+							FReply Reply = FReply::Unhandled();
+							if (Action.Event == IE_Pressed)
+							{
+								const double CurrentTime = FPlatformTime::Seconds();
+								if (CurrentTime - VREditorInteractor->GetLastUIPressTime() <= GetDefault<UVRModeSettings>()->DoubleClickTime)
+								{
+									// Trigger a double click event!
+									Reply = FSlateApplication::Get().RoutePointerDoubleClickEvent(WidgetPathUnderFinger, PointerEvent);
+								}
+								else
+								{
+									// If we are clicking on an editable text field and the radial menu is not a numpad, show the numpad
+									if (WidgetPathUnderFinger.Widgets.Last().Widget->GetTypeAsString() == TEXT("SEditableText") && (bRadialMenuIsNumpad == false))
 									{
-										const double CurrentTime = FPlatformTime::Seconds();
-										if (CurrentTime - VREditorInteractor->GetLastUIPressTime() <= GetDefault<UVRModeSettings>()->DoubleClickTime)
+										if (!QuickRadialMenu->bHidden)
 										{
-											// Trigger a double click event!
-											Reply = FSlateApplication::Get().RoutePointerDoubleClickEvent(WidgetPathUnderFinger, PointerEvent);
+											bRadialMenuVisibleAtSwap = true;
 										}
 										else
 										{
-											// If we are clicking on an editable text field and the radial menu is not a numpad, show the numpad
-											if (WidgetPathUnderFinger.Widgets.Last().Widget->GetTypeAsString() == TEXT("SEditableText") && (bRadialMenuIsNumpad == false))
-											{
-												if (!QuickRadialMenu->bHidden)
-												{
-													bRadialMenuVisibleAtSwap = true;
-												}
-												else
-												{
-													bRadialMenuVisibleAtSwap = false;
-													// Force the radial menu to spawn even if the laser is over UI
-													const bool bForceRefresh = false;
-													TryToSpawnRadialMenu(UIInteractor, bForceRefresh);
-												}
-												SwapRadialMenu();
-											}
-											Reply = FSlateApplication::Get().RoutePointerDownEvent(WidgetPathUnderFinger, PointerEvent);
+											bRadialMenuVisibleAtSwap = false;
+											// Force the radial menu to spawn even if the laser is over UI
+											const bool bForceRefresh = false;
+											TryToSpawnRadialMenu(UIInteractor, bForceRefresh);
 										}
+										SwapRadialMenu();
+									}
+									Reply = FSlateApplication::Get().RoutePointerDownEvent(WidgetPathUnderFinger, PointerEvent);
+								}
 
-										// In case of selecting a level in the content browser the VREditormode is closed, this makes sure nothing happens after that.
-										if (!IVREditorModule::Get().IsVREditorModeActive())
-										{
-											bWasHandled = true;
-											return;
-										}
+								// In case of selecting a level in the content browser the VREditormode is closed, this makes sure nothing happens after that.
+								if (!IVREditorModule::Get().IsVREditorModeActive())
+								{
+									bWasHandled = true;
+									return;
+								}
 
-										VREditorInteractor->SetIsClickingOnUI(true);
-										VREditorInteractor->SetIsRightClickingOnUI(bIsRightClicking);
-										VREditorInteractor->SetLastUIPressTime(CurrentTime);
-										bOutIsInputCaptured = true;
+								VREditorInteractor->SetIsClickingOnUI(true);
+								VREditorInteractor->SetIsRightClickingOnUI(bIsRightClicking);
+								VREditorInteractor->SetLastUIPressTime(CurrentTime);
+								bOutIsInputCaptured = true;
 
-										// Play a haptic effect on press
-										VREditorInteractor->PlayHapticEffect(VREd::UIPressHapticFeedbackStrength->GetFloat());
-									}
-									else if (Action.Event == IE_Released)
-									{
-										Reply = FSlateApplication::Get().RoutePointerUpEvent(WidgetPathUnderFinger, PointerEvent);
-									}
-								}
+								// Play a haptic effect on press
+								VREditorInteractor->PlayHapticEffect(VREd::UIPressHapticFeedbackStrength->GetFloat());
+							}
+							else if (Action.Event == IE_Released)
+							{
+								Reply = FSlateApplication::Get().RoutePointerUpEvent(WidgetPathUnderFinger, PointerEvent);
 							}
-						}
-						else
-						{
-							// If we didn't handle the input in any other way, send an empty mouse down event so Slate focus is handled correctly
-							const bool bIsRightClicking =
-								(Action.Event == IE_Pressed && VREditorInteractor->IsModifierPressed()) ||
-								(Action.Event == IE_Released && VREditorInteractor->IsRightClickingOnUI());
-							TSet<FKey> PressedButtons;
-							FPointerEvent PointerEvent(
-								1 + (uint8)VREditorInteractor->GetControllerSide(),
-								FVector2D::ZeroVector,
-								FVector2D::ZeroVector,
-								PressedButtons,
-								bIsRightClicking ? EKeys::RightMouseButton : EKeys::LeftMouseButton,
-								0.0f,	// Wheel delta
-								FModifierKeysState());
-
-							FWidgetPath EmptyWidgetPath;
-							FReply Reply = FSlateApplication::Get().RoutePointerDownEvent(EmptyWidgetPath, PointerEvent);
 						}
 					}
 				}
 			}
-
-			if (Action.Event == IE_Released)
+		}
+		if (Action.Event == IE_Released)
+		{
+			bool bWasRightClicking = false;
+			if (VREditorInteractor->IsClickingOnUI())
 			{
-				bool bWasRightClicking = false;
-				if (VREditorInteractor->IsClickingOnUI())
+				if (VREditorInteractor->IsRightClickingOnUI())
 				{
-					if (VREditorInteractor->IsRightClickingOnUI())
-					{
-						bWasRightClicking = true;
-					}
-					VREditorInteractor->SetIsClickingOnUI(false);
-					VREditorInteractor->SetIsRightClickingOnUI(false);
-					bOutIsInputCaptured = false;
+					bWasRightClicking = true;
 				}
+				VREditorInteractor->SetIsClickingOnUI(false);
+				VREditorInteractor->SetIsRightClickingOnUI(false);
+				bOutIsInputCaptured = false;
+			}
 
-				if (!bWasHandled)
-				{
-					TSet<FKey> PressedButtons;
-					FPointerEvent PointerEvent(
-						1 + (uint8)VREditorInteractor->GetControllerSide(),
-						FVector2D::ZeroVector,
-						FVector2D::ZeroVector,
-						PressedButtons,
-						bWasRightClicking ? EKeys::RightMouseButton : EKeys::LeftMouseButton,
-						0.0f,	// Wheel delta
-						FModifierKeysState());
-
-					FWidgetPath EmptyWidgetPath;
-
-					VREditorInteractor->SetIsClickingOnUI(false);
-					VREditorInteractor->SetIsRightClickingOnUI(false);
-					FReply Reply = FSlateApplication::Get().RoutePointerUpEvent(EmptyWidgetPath, PointerEvent);
-				}
+			if (!bWasHandled)
+			{
+				TSet<FKey> PressedButtons;
+				FPointerEvent PointerEvent(
+					1 + (uint8)VREditorInteractor->GetControllerSide(),
+					FVector2D::ZeroVector,
+					FVector2D::ZeroVector,
+					PressedButtons,
+					bWasRightClicking ? EKeys::RightMouseButton : EKeys::LeftMouseButton,
+					0.0f,	// Wheel delta
+					FModifierKeysState());
+
+				FWidgetPath EmptyWidgetPath;
+
+				VREditorInteractor->SetIsClickingOnUI(false);
+				VREditorInteractor->SetIsRightClickingOnUI(false);
+				FReply Reply = FSlateApplication::Get().RoutePointerUpEvent(EmptyWidgetPath, PointerEvent);
 			}
 		}
 	}
-
 }
 
 
@@ -959,7 +925,7 @@ void UVREditorUISystem::CreateUIs()
 			}
 
 			QuickRadialMenu->SetRelativeOffset(RelativeOffset);
-			QuickRadialMenu->ShowUI(false, 0.0f, false);
+			QuickRadialMenu->ShowUI(false, false, 0.0f, false);
 		}
 	}
 	// Make some editor UIs!
@@ -2615,7 +2581,7 @@ UVREditorMotionControllerInteractor* UVREditorUISystem::GetUIInteractor()
 	return UIInteractor;
 }
 
-void UVREditorUISystem::SequencerRadialMenuGenerator(FMenuBuilder MenuBuilder, TSharedPtr<FUICommandList> CommandList, class UVREditorMode* VRMode, float& RadiusOverride)
+void UVREditorUISystem::SequencerRadialMenuGenerator(FMenuBuilder MenuBuilder, TSharedPtr<FUICommandList> CommandList, UVREditorMode* InVRMode, float& RadiusOverride)
 {
 	RadiusOverride = 1.0f;
 	// First menu entry is at 90 degrees 
@@ -2625,7 +2591,7 @@ void UVREditorUISystem::SequencerRadialMenuGenerator(FMenuBuilder MenuBuilder, T
 		FSlateIcon(FVREditorStyle::GetStyleSetName(), "VREditorStyle.SequencerPlay"),
 		FUIAction
 		(
-			FExecuteAction::CreateStatic(&FVREditorActionCallbacks::PlaySequenceAtRate, VRMode, 1.0f),
+			FExecuteAction::CreateStatic(&FVREditorActionCallbacks::PlaySequenceAtRate, InVRMode, 1.0f),
 			FCanExecuteAction::CreateUObject(this->GetRadialMenuHandler(), &UVRRadialMenuHandler::IsActionMenuBound)
 		)
 		);
@@ -2635,7 +2601,7 @@ void UVREditorUISystem::SequencerRadialMenuGenerator(FMenuBuilder MenuBuilder, T
 		FSlateIcon(FVREditorStyle::GetStyleSetName(), "VREditorStyle.SequencerReverse"),
 		FUIAction
 		(
-			FExecuteAction::CreateStatic(&FVREditorActionCallbacks::PlaySequenceAtRate, VRMode, -1.0f),
+			FExecuteAction::CreateStatic(&FVREditorActionCallbacks::PlaySequenceAtRate, InVRMode, -1.0f),
 			FCanExecuteAction::CreateUObject(this->GetRadialMenuHandler(), &UVRRadialMenuHandler::IsActionMenuBound)
 		)
 		);
@@ -2645,7 +2611,7 @@ void UVREditorUISystem::SequencerRadialMenuGenerator(FMenuBuilder MenuBuilder, T
 		FSlateIcon(FVREditorStyle::GetStyleSetName(), "VREditorStyle.SequencerStop"),
 		FUIAction
 		(
-			FExecuteAction::CreateStatic(&FVREditorActionCallbacks::PauseSequencePlayback, VRMode),
+			FExecuteAction::CreateStatic(&FVREditorActionCallbacks::PauseSequencePlayback, InVRMode),
 			FCanExecuteAction::CreateUObject(this->GetRadialMenuHandler(), &UVRRadialMenuHandler::IsActionMenuBound)
 		)
 		);
@@ -2655,7 +2621,7 @@ void UVREditorUISystem::SequencerRadialMenuGenerator(FMenuBuilder MenuBuilder, T
 		FSlateIcon(FVREditorStyle::GetStyleSetName(), "VREditorStyle.PlayFromStart"),
 		FUIAction
 		(
-			FExecuteAction::CreateStatic(&FVREditorActionCallbacks::PlayFromBeginning, VRMode),
+			FExecuteAction::CreateStatic(&FVREditorActionCallbacks::PlayFromBeginning, InVRMode),
 			FCanExecuteAction::CreateUObject(this->GetRadialMenuHandler(), &UVRRadialMenuHandler::IsActionMenuBound)
 			)
 		);
@@ -2665,7 +2631,7 @@ void UVREditorUISystem::SequencerRadialMenuGenerator(FMenuBuilder MenuBuilder, T
 		FSlateIcon(FVREditorStyle::GetStyleSetName(), "VREditorStyle.Scrub"),
 		FUIAction
 		(
-			FExecuteAction::CreateStatic(&FVREditorActionCallbacks::ToggleSequencerScrubbing, VRMode, UIInteractor),
+			FExecuteAction::CreateStatic(&FVREditorActionCallbacks::ToggleSequencerScrubbing, InVRMode, UIInteractor),
 			FCanExecuteAction::CreateUObject(this->GetRadialMenuHandler(), &UVRRadialMenuHandler::IsActionMenuBound),
 			FGetActionCheckState::CreateStatic(&FVREditorActionCallbacks::GetSequencerScrubState, UIInteractor)
 			),
@@ -2678,9 +2644,9 @@ void UVREditorUISystem::SequencerRadialMenuGenerator(FMenuBuilder MenuBuilder, T
 		FSlateIcon(FVREditorStyle::GetStyleSetName(), "VREditorStyle.ToggleLooping"),
 		FUIAction
 		(
-			FExecuteAction::CreateStatic(&FVREditorActionCallbacks::ToggleLooping, VRMode),
+			FExecuteAction::CreateStatic(&FVREditorActionCallbacks::ToggleLooping, InVRMode),
 			FCanExecuteAction::CreateUObject(this->GetRadialMenuHandler(), &UVRRadialMenuHandler::IsActionMenuBound),
-			FGetActionCheckState::CreateStatic(&FVREditorActionCallbacks::IsLoopingChecked, VRMode)
+			FGetActionCheckState::CreateStatic(&FVREditorActionCallbacks::IsLoopingChecked, InVRMode)
 			),
 		NAME_None,
 		EUserInterfaceActionType::ToggleButton
diff --git a/Engine/Source/Editor/VREditor/UI/VREditorUISystem.h b/Engine/Source/Editor/VREditor/UI/VREditorUISystem.h
--- a/Engine/Source/Editor/VREditor/UI/VREditorUISystem.h
+++ b/Engine/Source/Editor/VREditor/UI/VREditorUISystem.h
@@ -63,7 +63,9 @@ struct FVRButton
 	float ScaleRate;
 
 	FVRButton()
-		: AnimationDirection(EVREditorAnimationState::None),
+		: ButtonWidget(nullptr),
+		AnimationDirection(EVREditorAnimationState::None),
+		OriginalRelativeScale(FVector::ZeroVector),
 		CurrentScale(1.0f),
 		MinScale(1.0f),
 		MaxScale(1.10f),
@@ -73,8 +75,8 @@ struct FVRButton
 	FVRButton(class UVREditorWidgetComponent* InButtonWidget, FVector InOriginalScale,
 		EVREditorAnimationState InAnimationDirection = EVREditorAnimationState::None, float InCurrentScale = 1.0f, float InMinScale = 1.0f, float InMaxScale = 1.25f, float InScaleRate = 2.0f)
 		: ButtonWidget(InButtonWidget),
-		OriginalRelativeScale(InOriginalScale),
 		AnimationDirection(InAnimationDirection),
+		OriginalRelativeScale(InOriginalScale),
 		CurrentScale(InCurrentScale),
 		MinScale(InMinScale),
 		MaxScale(InMaxScale),
@@ -301,7 +303,7 @@ protected:
 	void SwapRadialMenu();
 
 	/** Creates the sequencer radial menu to pass to the radial menu generator */
-	void SequencerRadialMenuGenerator(FMenuBuilder MenuBuilder, TSharedPtr<FUICommandList> CommandList, class UVREditorMode* VRMode, float& RadiusOverride);
+	void SequencerRadialMenuGenerator(FMenuBuilder MenuBuilder, TSharedPtr<FUICommandList> CommandList, class UVREditorMode* InVRMode, float& RadiusOverride);
 
 	/**
 	* Handles being notified when any editor mode changes to see if any VR Editor UI needs to change.
diff --git a/Engine/Source/Editor/VREditor/VREditorMotionControllerInteractor.cpp b/Engine/Source/Editor/VREditor/VREditorMotionControllerInteractor.cpp
--- a/Engine/Source/Editor/VREditor/VREditorMotionControllerInteractor.cpp
+++ b/Engine/Source/Editor/VREditor/VREditorMotionControllerInteractor.cpp
@@ -690,8 +690,8 @@ void UVREditorMotionControllerInteractor::PreviewInputKey( FEditorViewportClient
 	{
 		// Are we holding "up" on the trackpad?
 		const bool bIsHoldingUpOnTrackpad =
-			this->bIsTrackpadPositionValid &&
-			this->TrackpadPosition.Y >= VREd::MinTrackpadOffsetBeforeRadialMenu->GetFloat() &&
+			bIsTrackpadPositionValid[0] && bIsTrackpadPositionValid[1] &&
+			TrackpadPosition.Y >= VREd::MinTrackpadOffsetBeforeRadialMenu->GetFloat() &&
 			( GetHMDDeviceType() == EHMDDeviceType::DT_OculusRift || bIsPressingTrackpad );
 
 		if( bIsHoldingUpOnTrackpad && Action.ActionType == ViewportWorldActionTypes::SelectAndMove && Event == IE_Pressed )
diff --git a/Engine/Source/Editor/VREditor/VREditorStyle.cpp b/Engine/Source/Editor/VREditor/VREditorStyle.cpp
--- a/Engine/Source/Editor/VREditor/VREditorStyle.cpp
+++ b/Engine/Source/Editor/VREditor/VREditorStyle.cpp
@@ -82,7 +82,7 @@ TSharedRef< FSlateStyleSet > FVREditorStyle::Create()
 	// Headings will have a font outline
 	FFontOutlineSettings HeadingOutline;
 	HeadingOutline.OutlineColor = FLinearColor(0.2f, 0.2f, 0.2f, 0.5f);
-	HeadingOutline.OutlineSize = 1.5f;
+	HeadingOutline.OutlineSize = 1;
 	FSlateFontInfo HeadlineFont = TTF_CORE_FONT("Fonts/Roboto-Regular", 10);
 	HeadlineFont.OutlineSettings = HeadingOutline;
 
@@ -93,7 +93,7 @@ TSharedRef< FSlateStyleSet > FVREditorStyle::Create()
 	// Headings will have a font outline
 	FFontOutlineSettings HelperOutline;
 	HelperOutline.OutlineColor = FLinearColor( 0.2f, 0.2f, 0.2f, 0.5f );
-	HelperOutline.OutlineSize = 3.0f;
+	HelperOutline.OutlineSize = 3;
 	FSlateFontInfo HelperFont = TTF_CORE_FONT( "Fonts/Roboto-Regular", 24 );
 	HelperFont.OutlineSettings = HelperOutline;
 
@@ -139,7 +139,7 @@ TSharedRef< FSlateStyleSet > FVREditorStyle::Create()
 
 	FFontOutlineSettings RadialOutline;
 	RadialOutline.OutlineColor = FLinearColor(0.1f, 0.1f, 0.1f, 1.0f);
-	RadialOutline.OutlineSize = 1.0f;
+	RadialOutline.OutlineSize = 1;
 
 	FSlateFontInfo RadialFont = TTF_CORE_FONT("Fonts/Roboto-Regular", 8);
 	RadialFont.OutlineSettings = RadialOutline;
@@ -152,7 +152,7 @@ TSharedRef< FSlateStyleSet > FVREditorStyle::Create()
 	FSlateFontInfo ActiveRadialFont = TTF_CORE_FONT("Fonts/Roboto-Regular", 10);
 	FFontOutlineSettings ActiveRadialOutline;
 	ActiveRadialOutline.OutlineColor = FLinearColor::Black;
-	ActiveRadialOutline.OutlineSize = 1.0f;
+	ActiveRadialOutline.OutlineSize = 1;
 	ActiveRadialFont.OutlineSettings = ActiveRadialOutline;
 	Style->Set("VRRadialStyle.ActiveFont", FSlateFontInfo(ActiveRadialFont));
 	
diff --git a/Engine/Source/Editor/VREditor/VREditorWorldInteraction.cpp b/Engine/Source/Editor/VREditor/VREditorWorldInteraction.cpp
--- a/Engine/Source/Editor/VREditor/VREditorWorldInteraction.cpp
+++ b/Engine/Source/Editor/VREditor/VREditorWorldInteraction.cpp
@@ -466,10 +466,12 @@ void UVREditorWorldInteraction::PlaceDraggedMaterialOrTexture( UViewportInteract
 				if( !bPlaced )
 				{
 					const int32 TargetMaterialSlot = -1;	// All materials
-					bool bPlaced = FComponentEditorUtils::AttemptApplyMaterialToComponent( HitComponent, DroppedObjAsMaterial, TargetMaterialSlot );
-
-					const UVREditorAssetContainer& AssetContainer = Owner->GetAssetContainer();
-					Owner->PlaySound(AssetContainer.DropFromContentBrowserSound, Interactor->GetTransform().GetLocation());
+					bool AppliedMaterial = FComponentEditorUtils::AttemptApplyMaterialToComponent( HitComponent, DroppedObjAsMaterial, TargetMaterialSlot );
+					if (AppliedMaterial)
+					{
+						const UVREditorAssetContainer& AssetContainer = Owner->GetAssetContainer();
+						Owner->PlaySound(AssetContainer.DropFromContentBrowserSound, Interactor->GetTransform().GetLocation());
+					}
 				}
 			}
 		}
diff --git a/Engine/Source/Editor/ViewportInteraction/Gizmo/VIGizmoHandle.cpp b/Engine/Source/Editor/ViewportInteraction/Gizmo/VIGizmoHandle.cpp
--- a/Engine/Source/Editor/ViewportInteraction/Gizmo/VIGizmoHandle.cpp
+++ b/Engine/Source/Editor/ViewportInteraction/Gizmo/VIGizmoHandle.cpp
@@ -9,8 +9,11 @@
 
 UGizmoHandleGroup::UGizmoHandleGroup()
 	: Super(),
-	bShowOnUniversalGizmo(true),
-	OwningTransformGizmoActor(nullptr)
+	GizmoMaterial(nullptr),
+	TranslucentGizmoMaterial(nullptr),
+	Handles(),
+	OwningTransformGizmoActor(nullptr),
+	bShowOnUniversalGizmo(true)
 {
 
 }
@@ -18,6 +21,8 @@ UGizmoHandleGroup::UGizmoHandleGroup()
 UGizmoHandleGroup::~UGizmoHandleGroup()
 {
 	OwningTransformGizmoActor = nullptr;
+	GizmoMaterial = nullptr;
+	TranslucentGizmoMaterial = nullptr;
 }
 
 FTransformGizmoHandlePlacement UGizmoHandleGroup::MakeHandlePlacementForIndex( const int32 HandleIndex ) const
diff --git a/Engine/Source/Editor/ViewportInteraction/ViewportWorldInteraction.cpp b/Engine/Source/Editor/ViewportInteraction/ViewportWorldInteraction.cpp
--- a/Engine/Source/Editor/ViewportInteraction/ViewportWorldInteraction.cpp
+++ b/Engine/Source/Editor/ViewportInteraction/ViewportWorldInteraction.cpp
@@ -2582,7 +2582,6 @@ void UViewportWorldInteraction::RefreshTransformGizmo( const bool bNewObjectsSel
 		const bool bPropagateToChildren = true;
 		TransformGizmoActor->GetRootComponent()->SetVisibility( bShouldBeVisible, bPropagateToChildren );
 
-
 		UViewportInteractor* DraggingWithInteractor = nullptr;
 		static TArray< UActorComponent* > HoveringOverHandles;
 		HoveringOverHandles.Reset();
@@ -2691,9 +2690,9 @@ void UViewportWorldInteraction::RefreshTransformGizmo( const bool bNewObjectsSel
 		SpawnGridMeshActor();
 		if ( FSnappingUtils::IsSnapToGridEnabled() )
 		{
-	        const bool bShouldBeVisible = true;
-	        const bool bPropagateToChildren = true;
-	        SnapGridActor->GetRootComponent()->SetVisibility( bShouldBeVisible, bPropagateToChildren );
+	        const bool bShouldGridBeVisible = true;
+	        const bool bPropagateToGridChildren = true;
+	        SnapGridActor->GetRootComponent()->SetVisibility( bShouldGridBeVisible, bPropagateToGridChildren );
 
 			const float GizmoAnimationAlpha = TransformGizmoActor->GetAnimationAlpha();
 
@@ -2797,9 +2796,9 @@ void UViewportWorldInteraction::RefreshTransformGizmo( const bool bNewObjectsSel
 		else
 		{
 			// Grid snap not enabled
-	        const bool bShouldBeVisible = false;
-	        const bool bPropagateToChildren = true;
-	        SnapGridActor->GetRootComponent()->SetVisibility( bShouldBeVisible, bPropagateToChildren );
+	        const bool bShouldGridBeVisible = false;
+	        const bool bPropagateToGridChildren = true;
+	        SnapGridActor->GetRootComponent()->SetVisibility( bShouldGridBeVisible, bPropagateToGridChildren );
 		}
 
 		if (bNewObjectsSelected && bPlayNextRefreshTransformGizmoSound)
