VREditor: Pop-up dialog the first time you enter VR mode or toggle auto-entry on

diff --git a/Engine/Source/Editor/VREditor/Public/VRModeSettings.h b/Engine/Source/Editor/VREditor/Public/VRModeSettings.h
--- a/Engine/Source/Editor/VREditor/Public/VRModeSettings.h
+++ b/Engine/Source/Editor/VREditor/Public/VRModeSettings.h
@@ -66,4 +66,8 @@ public:
 	UPROPERTY(EditAnywhere, config, Category = "Motion Controllers", meta = (DisplayName = "Trigger Pressed Threshold (Oculus Touch)", ClampMin = 0.01, ClampMax = 1.0))
 	float TriggerPressedThreshold_Rift;
 
+private:
+#if WITH_EDITOR
+	void PostEditChangeProperty(struct FPropertyChangedEvent& PropertyChangedEvent);
+#endif
 };
diff --git a/Engine/Source/Editor/VREditor/VREditorModeManager.cpp b/Engine/Source/Editor/VREditor/VREditorModeManager.cpp
--- a/Engine/Source/Editor/VREditor/VREditorModeManager.cpp
+++ b/Engine/Source/Editor/VREditor/VREditorModeManager.cpp
@@ -15,6 +15,9 @@
 #include "EditorWorldExtension.h"
 #include "ViewportWorldInteraction.h"
 #include "VRModeSettings.h"
+#include "Dialogs.h"
+
+#define LOCTEXT_NAMESPACE "VREditor"
 
 FVREditorModeManager::FVREditorModeManager() :
 	CurrentVREditorMode( nullptr ),
@@ -113,7 +116,18 @@ void FVREditorModeManager::EnableVREditor( const bool bEnable, const bool bForce
 	{
 		if( bEnable && ( IsVREditorAvailable() || bForceWithoutHMD ))
 		{
-			StartVREditorMode( bForceWithoutHMD );
+			FSuppressableWarningDialog::FSetupInfo SetupInfo(LOCTEXT("VRModeEntry_Message", "VR Mode enables you to work on your project in virtual reality using motion controllers. This feature is still under development, so you may experience bugs or crashes while using it."),
+				LOCTEXT("VRModeEntry_Title", "Entering VR Mode - Experimental"), "Warning_VRModeEntry", GEditorSettingsIni);
+
+			SetupInfo.ConfirmText = LOCTEXT("VRModeEntry_ConfirmText", "Continue");
+			SetupInfo.CancelText = LOCTEXT("VRModeEntry_CancelText", "Cancel");
+			SetupInfo.bDefaultToSuppressInTheFuture = true;
+			FSuppressableWarningDialog VRModeEntryWarning(SetupInfo);
+
+			if (VRModeEntryWarning.ShowModal() != FSuppressableWarningDialog::Cancel)
+			{
+				StartVREditorMode(bForceWithoutHMD);
+			}
 		}
 		else if( !bEnable )
 		{
@@ -211,3 +225,5 @@ void FVREditorModeManager::OnMapChanged( UWorld* World, EMapChangeType MapChange
 	}
 	CurrentVREditorMode = nullptr;
 }
+
+#undef LOCTEXT_NAMESPACE
diff --git a/Engine/Source/Editor/VREditor/VRModeSettings.cpp b/Engine/Source/Editor/VREditor/VRModeSettings.cpp
--- a/Engine/Source/Editor/VREditor/VRModeSettings.cpp
+++ b/Engine/Source/Editor/VREditor/VRModeSettings.cpp
@@ -1,6 +1,10 @@
 // Copyright 1998-2017 Epic Games, Inc. All Rights Reserved.
 
 #include "VRModeSettings.h"
+#include "Dialogs.h"
+#include "UnrealType.h"
+
+#define LOCTEXT_NAMESPACE "VREditor"
 
 UVRModeSettings::UVRModeSettings(const FObjectInitializer& ObjectInitializer)
 	: Super(ObjectInitializer)
@@ -16,3 +20,24 @@ UVRModeSettings::UVRModeSettings(const FObjectInitializer& ObjectInitializer)
 	TriggerPressedThreshold_Vive = 0.03f;
 	TriggerPressedThreshold_Rift = 0.2f;
 }
+
+#if WITH_EDITOR
+void UVRModeSettings::PostEditChangeProperty(struct FPropertyChangedEvent& PropertyChangedEvent)
+{
+	if (PropertyChangedEvent.Property
+		&& PropertyChangedEvent.Property->GetFName() == GET_MEMBER_NAME_CHECKED(UVRModeSettings, bEnableAutoVREditMode)
+		&& bEnableAutoVREditMode == true)
+	{
+		FSuppressableWarningDialog::FSetupInfo SetupInfo(LOCTEXT("VRModeEntry_Message", "VR Mode enables you to work on your project in virtual reality using motion controllers. This feature is still under development, so you may experience bugs or crashes while using it."),
+			LOCTEXT("VRModeEntry_Title", "Entering VR Mode - Experimental"), "Warning_VRModeEntry", GEditorSettingsIni);
+
+		SetupInfo.ConfirmText = LOCTEXT("VRModeEntry_ConfirmText", "Continue");
+		SetupInfo.CancelText = LOCTEXT("VRModeEntry_CancelText", "Cancel");
+		SetupInfo.bDefaultToSuppressInTheFuture = true;
+		FSuppressableWarningDialog VRModeEntryWarning(SetupInfo);
+		bEnableAutoVREditMode = (VRModeEntryWarning.ShowModal() != FSuppressableWarningDialog::Cancel) ? true : false;
+	}
+}
+#endif
+
+#undef LOCTEXT_NAMESPACE
\ No newline at end of file
