VR Editor: Disable entering VR mode while in simulate.

diff --git a/Engine/Source/Editor/VREditor/VREditorModeManager.cpp b/Engine/Source/Editor/VREditor/VREditorModeManager.cpp
--- a/Engine/Source/Editor/VREditor/VREditorModeManager.cpp
+++ b/Engine/Source/Editor/VREditor/VREditorModeManager.cpp
@@ -147,7 +147,7 @@ bool FVREditorModeManager::IsVREditorActive() const
 bool FVREditorModeManager::IsVREditorAvailable() const
 {
 	const bool bHasHMDDevice = GEngine->HMDDevice.IsValid() && GEngine->HMDDevice->IsHMDEnabled();
-	return bHasHMDDevice;
+	return bHasHMDDevice && !GEditor->bIsSimulatingInEditor;
 }
 
 
