Fix for radial menu after PIE

diff --git a/Engine/Source/Editor/VREditor/UI/VREditorUISystem.cpp b/Engine/Source/Editor/VREditor/UI/VREditorUISystem.cpp
--- a/Engine/Source/Editor/VREditor/UI/VREditorUISystem.cpp
+++ b/Engine/Source/Editor/VREditor/UI/VREditorUISystem.cpp
@@ -288,7 +288,7 @@ void UVREditorUISystem::OnPreviewInputAction(FEditorViewportClient& ViewportClie
 						SwapRadialMenu();
 						if (!bRadialMenuVisibleAtSwap)
 						{
-							HideRadialMenu(VREditorInteractor);
+							HideRadialMenu();
 						}
 						bWasHandled = true;
 					}
@@ -376,7 +376,7 @@ void UVREditorUISystem::OnPreviewInputAction(FEditorViewportClient& ViewportClie
 						SwapRadialMenu();
 						if (!bRadialMenuVisibleAtSwap)
 						{
-							HideRadialMenu(VREditorInteractor);
+							HideRadialMenu();
 						}
 					}
 				}
@@ -744,7 +744,7 @@ void UVREditorUISystem::Tick( FEditorViewportClient* ViewportClient, const float
 	if (!bRadialMenuIsNumpad &&
 		FTimespan::FromSeconds(FPlatformTime::Seconds()) - CompareTime > FTimespan::FromSeconds(1.5f))
 	{
-		HideRadialMenu(UIInteractor);
+		HideRadialMenu();
 	}
 
 	// When dragging a panel with the UI interactor for a while, we want to close the radial menu.
@@ -753,7 +753,7 @@ void UVREditorUISystem::Tick( FEditorViewportClient* ViewportClient, const float
 		DragPanelFromOpenTime += DeltaTime;
 		if (DragPanelFromOpenTime >= 1.0 && UIInteractor != nullptr && DraggingUI != nullptr && InteractorDraggingUI != nullptr && UIInteractor == InteractorDraggingUI)
 		{
-			HideRadialMenu(InteractorDraggingUI, false);
+			HideRadialMenu(false);
 		}
 	}
 
@@ -1371,10 +1371,10 @@ void UVREditorUISystem::TryToSpawnRadialMenu( UVREditorInteractor* Interactor, c
 }
 
 
-void UVREditorUISystem::HideRadialMenu( UVREditorInteractor* Interactor, const bool bPlaySound /*= true*/ )
+void UVREditorUISystem::HideRadialMenu( const bool bPlaySound /*= true */ )
 {
 	// Only hide the radial menu if the passed interactor is actually the interactor with the radial menu
-	if( IsShowingRadialMenu( Interactor ) )
+	if( IsShowingRadialMenu( UIInteractor ) )
 	{
 		QuickRadialMenu->ShowUI( false, true, VREd::RadialMenuFadeDelay->GetFloat(), bPlaySound );
 	}
@@ -2676,7 +2676,7 @@ void UVREditorUISystem::HandleEditorModeChanged(FEdMode* Mode, bool IsEnabled)
 void UVREditorUISystem::ResetAll()
 {
 	RadialMenuHandler->Home();
-	HideRadialMenu(UIInteractor, false);
+	HideRadialMenu(false);
 
 	FVREditorActionCallbacks::DeselectAll();
 
diff --git a/Engine/Source/Editor/VREditor/UI/VREditorUISystem.h b/Engine/Source/Editor/VREditor/UI/VREditorUISystem.h
--- a/Engine/Source/Editor/VREditor/UI/VREditorUISystem.h
+++ b/Engine/Source/Editor/VREditor/UI/VREditorUISystem.h
@@ -162,7 +162,7 @@ public:
 	void TryToSpawnRadialMenu( UVREditorInteractor* Interactor, const bool bForceRefresh, const bool bPlaySound = true );
 
 	/** Hides the radial menu if the specified hand is showing it */
-	void HideRadialMenu( UVREditorInteractor* Interactor, const bool bPlaySound = true );
+	void HideRadialMenu( const bool bPlaySound = true );
 
 	/** Start dragging a dock window on the hand */
 	void StartDraggingDockUI( class AVREditorDockableWindow* InitDraggingDockUI, UVREditorInteractor* Interactor, const float DockSelectDistance, const bool bPlaySound = true );
diff --git a/Engine/Source/Editor/VREditor/VREditorMode.cpp b/Engine/Source/Editor/VREditor/VREditorMode.cpp
--- a/Engine/Source/Editor/VREditor/VREditorMode.cpp
+++ b/Engine/Source/Editor/VREditor/VREditorMode.cpp
@@ -1118,6 +1118,8 @@ void UVREditorMode::RestoreFromPIE()
 	// Re-enable input pre-processing
 	WorldInteraction->SetUseInputPreprocessor(true);
 	WorldInteraction->SetActive(true);
+
+	UISystem->HideRadialMenu(false);
 }
 
 void UVREditorMode::RestoreWorldToMeters()
