Fixing loc errors

diff --git a/Engine/Source/Editor/Kismet/Private/SBlueprintRevisionMenu.cpp b/Engine/Source/Editor/Kismet/Private/SBlueprintRevisionMenu.cpp
--- a/Engine/Source/Editor/Kismet/Private/SBlueprintRevisionMenu.cpp
+++ b/Engine/Source/Editor/Kismet/Private/SBlueprintRevisionMenu.cpp
@@ -196,7 +196,7 @@ void SBlueprintRevisionMenu::OnSourceControlQueryComplete(const FSourceControlOp
 					FText ToolTipText;
 					if (ISourceControlModule::Get().GetProvider().UsesChangelists())
 					{
-						ToolTipText = FText::Format(LOCTEXT("RevisionToolTip", "CL #{CheckInNumber} {UserName} \n{DateTime} \n{ChanglistDescription}"), Args);
+						ToolTipText = FText::Format(LOCTEXT("ChangelistToolTip", "CL #{CheckInNumber} {UserName} \n{DateTime} \n{ChanglistDescription}"), Args);
 					}
 					else
 					{
diff --git a/Engine/Source/Editor/VREditor/UI/VREditorUISystem.cpp b/Engine/Source/Editor/VREditor/UI/VREditorUISystem.cpp
--- a/Engine/Source/Editor/VREditor/UI/VREditorUISystem.cpp
+++ b/Engine/Source/Editor/VREditor/UI/VREditorUISystem.cpp
@@ -2164,7 +2164,7 @@ TSharedRef<SWidget> UVREditorUISystem::BuildQuickMenuWidget()
 
 
 	ExtrasMenuBuilder.AddMenuEntry(
-		LOCTEXT("PauseSimulation", "Pause Simulation"),
+		LOCTEXT("PauseSimulationLegacy", "Pause Simulation"),
 		LOCTEXT("PauseSimulationTooltip", "Pause Simulation"),
 		FSlateIcon(FEditorStyle::GetStyleSetName(), "PlayWorld.Simulate"),
 		FUIAction
@@ -2177,7 +2177,7 @@ TSharedRef<SWidget> UVREditorUISystem::BuildQuickMenuWidget()
 		);
 
 	ExtrasMenuBuilder.AddMenuEntry(
-		LOCTEXT("ResumeSimulation", "Resume Simulation"),
+		LOCTEXT("ResumeSimulationLegacy", "Resume Simulation"),
 		LOCTEXT("ResumeSimulationTooltip", "Resume Simulation"),
 		FSlateIcon(FEditorStyle::GetStyleSetName(), "PlayWorld.Simulate"),
 		FUIAction
