Non-unity fix.

diff --git a/Engine/Source/Editor/VREditor/UI/VREditorUISystem.h b/Engine/Source/Editor/VREditor/UI/VREditorUISystem.h
--- a/Engine/Source/Editor/VREditor/UI/VREditorUISystem.h
+++ b/Engine/Source/Editor/VREditor/UI/VREditorUISystem.h
@@ -25,6 +25,7 @@ class UVREditorInteractor;
 class FMenuBuilder;
 class FUICommandList;
 class UVREditorWidgetComponent;
+class UWidgetComponent;
 
 typedef FName VREditorPanelID;
 
