Fix crash on moving VR Editor windows. Copy of 3630353

diff --git a/Engine/Source/Editor/VREditor/UI/VREditorDockableWindow.h b/Engine/Source/Editor/VREditor/UI/VREditorDockableWindow.h
--- a/Engine/Source/Editor/VREditor/UI/VREditorDockableWindow.h
+++ b/Engine/Source/Editor/VREditor/UI/VREditorDockableWindow.h
@@ -54,6 +54,7 @@ public:
 	virtual void OnHoverLeave( UViewportInteractor* Interactor, const UActorComponent* NewComponent ) override;
 	virtual void OnDragRelease( UViewportInteractor* Interactor ) override;
 	virtual class UViewportDragOperationComponent* GetDragOperationComponent() override;
+	virtual bool CanBeSelected() override { return false; };
 
 protected:
 
