Branch snapshot for CL 3327560

diff --git a/Engine/Source/Editor/UnrealEd/Private/DragTool_FrustumSelect.cpp b/Engine/Source/Editor/UnrealEd/Private/DragTool_FrustumSelect.cpp
--- a/Engine/Source/Editor/UnrealEd/Private/DragTool_FrustumSelect.cpp
+++ b/Engine/Source/Editor/UnrealEd/Private/DragTool_FrustumSelect.cpp
@@ -75,7 +75,7 @@ void FDragTool_ActorFrustumSelect::EndDrag()
 	bool bSelectionChanged = false;
 
 	// Let the editor mode try to handle the selection.
-	const bool bEditorModeHandledSelection = ModeTools->FrustumSelect(Frustum, bLeftMouseButtonDown);
+	const bool bEditorModeHandledSelection = ModeTools->FrustumSelect(Frustum, LevelViewportClient, bLeftMouseButtonDown);
 
 	if( !bEditorModeHandledSelection )
 	{
@@ -118,50 +118,9 @@ void FDragTool_ActorFrustumSelect::EndDrag()
 			// Extend the endpoint of the rect to get the actual line
 			FIntRect BoxRect( FIntPoint( FMath::Max( 0.0f, Start.X ), FMath::Max( 0.0f, Start.Y ) ), FIntPoint( FMath::Min(ViewportSizeX, FMath::TruncToInt(End.X+1)), FMath::Min( ViewportSizeY, FMath::TruncToInt(End.Y+1) ) ) );
 
-			const TArray<FColor>& RawHitProxyData = LevelViewportClient->Viewport->GetRawHitProxyData(BoxRect);
-
 			TSet<AActor*> HitActors;
 			TSet<UModel*> HitModels;
-
-
-			// Lower the resolution with massive box selects
-			int32 Step = (BoxRect.Width() > 500 && BoxRect.Height() > 500) ? 4 : 1;
-
-
-			for (int32 Y = BoxRect.Min.Y; Y < BoxRect.Max.Y; Y = Y < BoxRect.Max.Y - 1 ? FMath::Min(BoxRect.Max.Y-1, Y+Step) : ++Y )
-			{
-				const FColor* SourceData = &RawHitProxyData[Y * ViewportSizeX];
-				for (int32 X = BoxRect.Min.X; X < BoxRect.Max.X; X = X < BoxRect.Max.X-1 ? FMath::Min(BoxRect.Max.X-1, X + Step) : ++X )
-				{
-					FHitProxyId HitProxyId(SourceData[X]);
-					HHitProxy* HitProxy = GetHitProxyById(HitProxyId);
-
-					if (HitProxy)
-					{
-						if( HitProxy->IsA(HActor::StaticGetType()) )
-						{
-							AActor* Actor = ((HActor*)HitProxy)->Actor;
-							if (Actor)
-							{
-								HitActors.Add(Actor);
-							}
-						}
-						else if( HitProxy->IsA(HModel::StaticGetType()) )
-						{
-							HitModels.Add( ((HModel*)HitProxy)->GetModel() );
-						}
-						else if( HitProxy->IsA(HBSPBrushVert::StaticGetType()) )
-						{
-							HBSPBrushVert* HitBSPBrushVert = ((HBSPBrushVert*)HitProxy);
-							if( HitBSPBrushVert->Brush.IsValid() )
-							{
-								HitActors.Add( HitBSPBrushVert->Brush.Get() );
-							}
-						}
-					}
-
-				}
-			}
+			LevelViewportClient->Viewport->GetActorsAndModelsInHitProxy( BoxRect, HitActors, HitModels );
 
 			if (HitModels.Num() > 0)
 			{
@@ -297,10 +256,15 @@ void FDragTool_ActorFrustumSelect::CalculateFrustum( FSceneView* View, FConvexVo
 		FPlane BottomPlane(BoxPoint3, BoxPoint4, CamPoint); // Bottom Plane
 		FPlane LeftPlane(BoxPoint4, BoxPoint1, CamPoint); // Left Plane
 
-		// Try to get all six planes to create a frustum
+		// Try to get all six planes to create a frustum.
+		// The frustum is built with the first four planes corresponding to the sides of the frustum.
 		FPlane NearPlane;
 		FPlane FarPlane;
 		OutFrustum.Planes.Empty();
+		OutFrustum.Planes.Add(TopPlane);
+		OutFrustum.Planes.Add(RightPlane);
+		OutFrustum.Planes.Add(BottomPlane);
+		OutFrustum.Planes.Add(LeftPlane);
 		if ( View->ViewMatrices.GetViewProjectionMatrix().GetFrustumNearPlane(NearPlane) )
 		{
 			OutFrustum.Planes.Add(NearPlane);
@@ -309,10 +273,6 @@ void FDragTool_ActorFrustumSelect::CalculateFrustum( FSceneView* View, FConvexVo
 		{
 			OutFrustum.Planes.Add(FarPlane);
 		}
-		OutFrustum.Planes.Add(TopPlane);
-		OutFrustum.Planes.Add(RightPlane);
-		OutFrustum.Planes.Add(BottomPlane);
-		OutFrustum.Planes.Add(LeftPlane);
 		OutFrustum.Init();
 	}
 	else
diff --git a/Engine/Source/Editor/UnrealEd/Private/EdMode.cpp b/Engine/Source/Editor/UnrealEd/Private/EdMode.cpp
--- a/Engine/Source/Editor/UnrealEd/Private/EdMode.cpp
+++ b/Engine/Source/Editor/UnrealEd/Private/EdMode.cpp
@@ -520,12 +520,12 @@ bool FEdMode::BoxSelect( FBox& InBox, bool InSelect )
 	return bResult;
 }
 
-bool FEdMode::FrustumSelect( const FConvexVolume& InFrustum, bool InSelect )
+bool FEdMode::FrustumSelect( const FConvexVolume& InFrustum, FEditorViewportClient* InViewportClient, bool InSelect )
 {
 	bool bResult = false;
 	if( GetCurrentTool() )
 	{
-		bResult = GetCurrentTool()->FrustumSelect( InFrustum, InSelect );
+		bResult = GetCurrentTool()->FrustumSelect( InFrustum, InViewportClient, InSelect );
 	}
 	return bResult;
 }
diff --git a/Engine/Source/Editor/UnrealEd/Private/Editor.cpp b/Engine/Source/Editor/UnrealEd/Private/Editor.cpp
--- a/Engine/Source/Editor/UnrealEd/Private/Editor.cpp
+++ b/Engine/Source/Editor/UnrealEd/Private/Editor.cpp
@@ -88,6 +88,8 @@ FEditorDelegates::FOnEditorModeTransitioned				FEditorDelegates::EditorModeEnter
 FEditorDelegates::FOnEditorModeTransitioned				FEditorDelegates::EditorModeExit;
 FEditorDelegates::FOnPIEEvent							FEditorDelegates::PreBeginPIE;
 FEditorDelegates::FOnPIEEvent							FEditorDelegates::BeginPIE;
+FEditorDelegates::FOnPIEEvent							FEditorDelegates::PrePIEEnded;
+FEditorDelegates::FOnPIEEvent							FEditorDelegates::PostPIEStarted;
 FEditorDelegates::FOnPIEEvent							FEditorDelegates::EndPIE;
 FEditorDelegates::FOnPIEEvent							FEditorDelegates::PausePIE;
 FEditorDelegates::FOnPIEEvent							FEditorDelegates::ResumePIE;
diff --git a/Engine/Source/Editor/UnrealEd/Private/EditorEngine.cpp b/Engine/Source/Editor/UnrealEd/Private/EditorEngine.cpp
--- a/Engine/Source/Editor/UnrealEd/Private/EditorEngine.cpp
+++ b/Engine/Source/Editor/UnrealEd/Private/EditorEngine.cpp
@@ -134,7 +134,7 @@
 #include "Components/SkeletalMeshComponent.h"
 #include "EngineModule.h"
 
-#include "EditorWorldManager.h"
+#include "EditorWorldExtension.h"
 
 #if PLATFORM_WINDOWS
 	#include "WindowsHWrapper.h"
@@ -317,6 +317,8 @@ UEditorEngine::UEditorEngine(const FObjectInitializer& ObjectInitializer)
 	NumOnlinePIEInstances = 0;
 	DefaultWorldFeatureLevel = GMaxRHIFeatureLevel;
 
+	EditorWorldExtensionsManager = nullptr;
+
 #if !UE_BUILD_SHIPPING
 	if (!AutomationCommon::OnEditorAutomationMapLoadDelegate().IsBound())
 	{
@@ -613,7 +615,7 @@ void UEditorEngine::InitEditor(IEngineLoop* InEngineLoop)
 	TimerManager = MakeShareable(new FTimerManager());
 
 	// create the editor world manager
-	EditorWorldManager = MakeShareable(new FEditorWorldManager());
+	EditorWorldExtensionsManager = NewObject<UEditorWorldExtensionManager>();
 
 	// Settings.
 	FBSPOps::GFastRebuild = 0;
@@ -896,6 +898,7 @@ void UEditorEngine::Init(IEngineLoop* InEngineLoop)
 			TEXT("UndoHistory"),
 			TEXT("DeviceProfileEditor"),
 			TEXT("SourceCodeAccess"),
+			TEXT("MeshEditor"),
 			TEXT("BehaviorTreeEditor"),
 			TEXT("HardwareTargeting"),
 			TEXT("LocalizationDashboard"),
@@ -1699,8 +1702,8 @@ void UEditorEngine::Tick( float DeltaSeconds, bool bIdleMode )
 		}
 	}
 
-	// Updates the ViewportWorldInteraction
-	EditorWorldManager->Tick( DeltaSeconds );
+	// Updates all the extensions for all the editor worlds
+	EditorWorldExtensionsManager->Tick(DeltaSeconds);
 
 	bool bIsMouseOverAnyLevelViewport = false;
 
@@ -7008,4 +7011,4 @@ void UEditorEngine::AutomationLoadMap(const FString& MapName, FString* OutError)
 	return;
 }
 
-#undef LOCTEXT_NAMESPACE 
+#undef LOCTEXT_NAMESPACE 
\ No newline at end of file
diff --git a/Engine/Source/Editor/UnrealEd/Private/EditorModeManager.cpp b/Engine/Source/Editor/UnrealEd/Private/EditorModeManager.cpp
--- a/Engine/Source/Editor/UnrealEd/Private/EditorModeManager.cpp
+++ b/Engine/Source/Editor/UnrealEd/Private/EditorModeManager.cpp
@@ -18,6 +18,9 @@
 #include "Framework/Notifications/NotificationManager.h"
 #include "Widgets/Notifications/SNotificationList.h"
 #include "Engine/LevelStreaming.h"
+#include "ViewportWorldInteraction.h"
+#include "Public/EditorWorldExtension.h"
+#include "Editor/EditorEngine.h"
 
 /*------------------------------------------------------------------------------
 	FEditorModeTools.
@@ -217,7 +220,18 @@ void FEditorModeTools::SetPivotLocation( const FVector& Location, const bool bIn
 
 ECoordSystem FEditorModeTools::GetCoordSystem(bool bGetRawValue)
 {
-	if (!bGetRawValue && (GetWidgetMode() == FWidget::WM_Scale))
+	bool bAligningToActors = false;
+	if (GEditor->GetEditorWorldExtensionsManager() != nullptr)
+	{
+
+		UViewportWorldInteraction* ViewportWorldInteraction = Cast<UViewportWorldInteraction>(GEditor->GetEditorWorldExtensionsManager()->GetEditorWorldExtensions(GetWorld())->FindExtension(UViewportWorldInteraction::StaticClass()));
+		if (ViewportWorldInteraction != nullptr && ViewportWorldInteraction->AreAligningToActors() == true)
+		{
+			bAligningToActors = true;
+		}
+	}
+	if (!bGetRawValue && 
+		((GetWidgetMode() == FWidget::WM_Scale) || bAligningToActors))
 	{
 		return COORD_Local;
 	}
@@ -229,6 +243,20 @@ ECoordSystem FEditorModeTools::GetCoordSystem(bool bGetRawValue)
 
 void FEditorModeTools::SetCoordSystem(ECoordSystem NewCoordSystem)
 {
+	// If we are trying to enter world space but are aligning to actors, turn off aligning to actors
+	if (GEditor->GetEditorWorldExtensionsManager() != nullptr &&
+		NewCoordSystem == COORD_World)
+	{
+		UViewportWorldInteraction* ViewportWorldInteraction = Cast<UViewportWorldInteraction>(GEditor->GetEditorWorldExtensionsManager()->GetEditorWorldExtensions(GetWorld())->FindExtension(UViewportWorldInteraction::StaticClass()));
+		if (ViewportWorldInteraction != nullptr && ViewportWorldInteraction->AreAligningToActors() == true)
+		{
+			if (ViewportWorldInteraction->HasCandidatesSelected())
+			{
+				ViewportWorldInteraction->SetSelectionAsCandidates();
+			}
+			GUnrealEd->Exec(GetWorld(), TEXT("VI.EnableGuides 0"));
+		}
+	}
 	CoordSystem = NewCoordSystem;
 }
 
@@ -569,13 +597,13 @@ bool FEditorModeTools::BoxSelect( FBox& InBox, bool InSelect )
 }
 
 /** Notifies all active modes of frustum selection attempts */
-bool FEditorModeTools::FrustumSelect( const FConvexVolume& InFrustum, bool InSelect )
+bool FEditorModeTools::FrustumSelect( const FConvexVolume& InFrustum, FEditorViewportClient* InViewportClient, bool InSelect )
 {
 	bool bHandled = false;
 	for( int32 ModeIndex = 0; ModeIndex < Modes.Num(); ++ModeIndex )
 	{
 		const TSharedPtr<FEdMode>& Mode = Modes[ ModeIndex ];
-		bHandled |= Mode->FrustumSelect( InFrustum, InSelect );
+		bHandled |= Mode->FrustumSelect( InFrustum, InViewportClient, InSelect );
 	}
 	return bHandled;
 }
@@ -1176,4 +1204,4 @@ bool FEditorModeTools::CanCycleWidgetMode() const
 	}
 
 	return false;
-}
+}
\ No newline at end of file
diff --git a/Engine/Source/Editor/UnrealEd/Private/EditorTransaction.cpp b/Engine/Source/Editor/UnrealEd/Private/EditorTransaction.cpp
--- a/Engine/Source/Editor/UnrealEd/Private/EditorTransaction.cpp
+++ b/Engine/Source/Editor/UnrealEd/Private/EditorTransaction.cpp
@@ -57,8 +57,9 @@ inline UObject* BuildSubobjectKey(UObject* InObj, TArray<FName>& OutHierarchyNam
 	A single transaction.
 -----------------------------------------------------------------------------*/
 
-FTransaction::FObjectRecord::FObjectRecord(FTransaction* Owner, UObject* InObject, FScriptArray* InArray, int32 InIndex, int32 InCount, int32 InOper, int32 InElementSize, STRUCT_DC InDefaultConstructor, STRUCT_AR InSerializer, STRUCT_DTOR InDestructor)
+FTransaction::FObjectRecord::FObjectRecord(FTransaction* Owner, UObject* InObject, TUniquePtr<FChange> InCustomChange, FScriptArray* InArray, int32 InIndex, int32 InCount, int32 InOper, int32 InElementSize, STRUCT_DC InDefaultConstructor, STRUCT_AR InSerializer, STRUCT_DTOR InDestructor)
 	:	Object				( InObject )
+	,	CustomChange		( MoveTemp( InCustomChange ) )
 	,	Array				( InArray )
 	,	Index				( InIndex )
 	,	Count				( InCount )
@@ -76,8 +77,20 @@ FTransaction::FObjectRecord::FObjectRecord(FTransaction* Owner, UObject* InObjec
 		bWantsBinarySerialization = false; 
 	}
 	ObjectAnnotation = Object->GetTransactionAnnotation();
-	FWriter Writer( Data, ReferencedObjects, ReferencedNames, bWantsBinarySerialization );
-	SerializeContents( Writer, Oper );
+
+	// Don't bother saving the object state if we have a custom change which can perform the undo operation
+	if( CustomChange.IsValid() )
+	{
+		// @todo mesheditor debug
+		//GWarn->Logf( TEXT( "------------ Saved Undo Change ------------" ) );
+		//CustomChange->PrintToLog( *GWarn );
+		//GWarn->Logf( TEXT( "-------------------------------------------" ) );
+	}
+	else
+	{
+		FWriter Writer( Data, ReferencedObjects, ReferencedNames, bWantsBinarySerialization );
+		SerializeContents( Writer, Oper );
+	}
 }
 
 void FTransaction::FObjectRecord::SerializeContents( FArchive& Ar, int32 InOper )
@@ -153,6 +166,8 @@ void FTransaction::FObjectRecord::Restore( FTransaction* Owner )
 	{
 		bRestored = true;
 		check(!Owner->bFlip);
+		check(!CustomChange.IsValid());
+
 		FTransaction::FObjectRecord::FReader Reader( Owner, Data, ReferencedObjects, ReferencedNames, bWantsBinarySerialization );
 		SerializeContents( Reader, Oper );
 	}
@@ -160,6 +175,12 @@ void FTransaction::FObjectRecord::Restore( FTransaction* Owner )
 
 void FTransaction::FObjectRecord::Save(FTransaction* Owner)
 {
+	// if record has a custom change, no need to do anything here
+	if( CustomChange.IsValid() )
+	{
+		return;
+	}
+
 	// common undo/redo path, before applying undo/redo buffer we save current state:
 	check(Owner->bFlip);
 	if (!bRestored)
@@ -182,12 +203,34 @@ void FTransaction::FObjectRecord::Load(FTransaction* Owner)
 	if (!bRestored)
 	{
 		bRestored = true;
-		FTransaction::FObjectRecord::FReader Reader(Owner, Data, ReferencedObjects, ReferencedNames, bWantsBinarySerialization);
-		SerializeContents(Reader, Oper);
-		Exchange(ObjectAnnotation, FlipObjectAnnotation);
-		Exchange(Data, FlipData);
-		Exchange(ReferencedObjects, FlipReferencedObjects);
-		Exchange(ReferencedNames, FlipReferencedNames);
+
+		if( CustomChange.IsValid() )
+		{
+			// @todo mesheditor debug
+			//GWarn->Logf( TEXT( "---------- Undoing Custom Change ----------" ) );
+			//CustomChange->PrintToLog( *GWarn );
+			
+			TUniquePtr<FChange> InvertedChange = CustomChange->Execute( Object.Get() );
+
+			// @todo mesheditor debug
+			//GWarn->Logf( TEXT( "-----(Here's what the Redo looks like)-----" ) );
+			//if( InvertedChange.IsValid() )
+			//{
+			//	InvertedChange->PrintToLog( *GWarn );
+			//}
+			//GWarn->Logf( TEXT( "-------------------------------------------" ) );
+
+			this->CustomChange = MoveTemp( InvertedChange );
+		}
+		else
+		{
+			FTransaction::FObjectRecord::FReader Reader(Owner, Data, ReferencedObjects, ReferencedNames, bWantsBinarySerialization);
+			SerializeContents(Reader, Oper);
+			Exchange(ObjectAnnotation, FlipObjectAnnotation);
+			Exchange(Data, FlipData);
+			Exchange(ReferencedObjects, FlipReferencedObjects);
+			Exchange(ReferencedNames, FlipReferencedNames);
+		}
 		Oper *= -1;
 	}
 }
@@ -365,7 +408,6 @@ void FTransaction::AddReferencedObjects( FReferenceCollector& Collector )
 	Collector.AddReferencedObjects(ObjectMap);
 }
 
-// FTransactionBase interface.
 void FTransaction::SaveObject( UObject* Object )
 {
 	check(Object);
@@ -376,7 +418,7 @@ void FTransaction::SaveObject( UObject* Object )
 	{
 		ObjectMap.Add(Object,1);
 		// Save the object.
-		new( Records )FObjectRecord( this, Object, NULL, 0, 0, 0, 0, NULL, NULL, NULL );
+		new( Records )FObjectRecord( this, Object, nullptr, NULL, 0, 0, 0, 0, NULL, NULL, NULL );
 	}
 	else
 	{
@@ -402,10 +444,25 @@ void FTransaction::SaveArray( UObject* Object, FScriptArray* Array, int32 Index,
 	if( Object->HasAnyFlags(RF_Transactional) && !Object->GetOutermost()->HasAnyPackageFlags(PKG_PlayInEditor))
 	{
 		// Save the array.
-		new( Records )FObjectRecord( this, Object, Array, Index, Count, Oper, ElementSize, DefaultConstructor, Serializer, Destructor );
+		new( Records )FObjectRecord( this, Object, nullptr, Array, Index, Count, Oper, ElementSize, DefaultConstructor, Serializer, Destructor );
 	}
 }
 
+void FTransaction::StoreUndo( UObject* Object, TUniquePtr<FChange> UndoChange )
+{
+	check( Object );
+	Object->CheckDefaultSubobjects();
+
+	int32* SaveCount = ObjectMap.Find( Object );
+	if( !SaveCount )
+	{
+		ObjectMap.Add( Object, 0 );
+	}
+
+	// Save the undo record
+	new( Records )FObjectRecord( this, Object, MoveTemp( UndoChange ), NULL, 0, 0, 0, 0, NULL, NULL, NULL );
+}
+
 void FTransaction::SetPrimaryObject(UObject* InObject)
 {
 	if (PrimaryObject == NULL)
@@ -471,7 +528,6 @@ void FTransaction::Apply()
 	});
 
 	TArray<ULevel*> LevelsToCommitModelSurface;
-	NumModelsModified = 0;		// Count the number of UModels that were changed.
 	for (auto ChangedObjectIt : ChangedObjects)
 	{
 		UObject* ChangedObject = ChangedObjectIt.Key;
@@ -479,7 +535,6 @@ void FTransaction::Apply()
 		if (Model && Model->Nodes.Num())
 		{
 			FBSPOps::bspBuildBounds(Model);
-			++NumModelsModified;
 		}
 		
 		if (UModelComponent* ModelComponent = Cast<UModelComponent>(ChangedObject))
@@ -696,7 +751,7 @@ void UTransBuffer::Cancel( int32 StartIndex /*=0*/ )
 			GUndo = NULL;
 			
 			// remove the currently active transaction from the buffer
-			UndoBuffer.Pop();
+			UndoBuffer.RemoveAt( UndoBuffer.Num() - 1 );
 		}
 		else
 		{
diff --git a/Engine/Source/Editor/UnrealEd/Private/EditorViewportClient.cpp b/Engine/Source/Editor/UnrealEd/Private/EditorViewportClient.cpp
--- a/Engine/Source/Editor/UnrealEd/Private/EditorViewportClient.cpp
+++ b/Engine/Source/Editor/UnrealEd/Private/EditorViewportClient.cpp
@@ -43,6 +43,8 @@
 #include "ComponentRecreateRenderStateContext.h"
 #include "EditorBuildUtils.h"
 #include "AudioDevice.h"
+#include "EditorWorldExtension.h"
+#include "ViewportWorldInteraction.h"
 
 #define LOCTEXT_NAMESPACE "EditorViewportClient"
 
@@ -1417,8 +1419,16 @@ EMouseCursor::Type FEditorViewportClient::GetCursor(FViewport* InViewport,int32
 			}
 		}
 	}
-
-
+	
+	// Allow the viewport interaction to override any previously set mouse cursor
+	UViewportWorldInteraction* WorldInteraction = Cast<UViewportWorldInteraction>(GEditor->GetEditorWorldExtensionsManager()->GetEditorWorldExtensions(GetWorld())->FindExtension(UViewportWorldInteraction::StaticClass()));
+	if (WorldInteraction != nullptr && WorldInteraction->ShouldSuppressExistingCursor())
+	{
+			MouseCursor = EMouseCursor::None;
+			RequiredCursorVisibiltyAndAppearance.bHardwareCursorVisible = false;
+			RequiredCursorVisibiltyAndAppearance.bSoftwareCursorVisible = false;
+			UpdateRequiredCursorVisibility();
+	}
 
 	CachedMouseX = X;
 	CachedMouseY = Y;
@@ -2270,6 +2280,12 @@ bool FEditorViewportClient::InputKey(FViewport* InViewport, int32 ControllerId,
 		return true;
 	}
 
+	UEditorWorldExtensionCollection& EditorWorldExtensionCollection = *GEditor->GetEditorWorldExtensionsManager()->GetEditorWorldExtensions( GetWorld() );
+	if( EditorWorldExtensionCollection.InputKey(this, Viewport, Key, Event))
+	{
+		return true;
+	}
+
 	FInputEventState InputState(InViewport, Key, Event);
 
 	bool bHandled = false;
@@ -4268,7 +4284,7 @@ void FEditorViewportClient::UpdateRequiredCursorVisibility()
 			return;
 		}
 
-		if( MouseDeltaTracker->UsingDragTool() )
+		if (MouseDeltaTracker->UsingDragTool() )
 		{
 			RequiredCursorVisibiltyAndAppearance.bOverrideAppearance = false;
 		}
diff --git a/Engine/Source/Editor/UnrealEd/Private/PlayLevel.cpp b/Engine/Source/Editor/UnrealEd/Private/PlayLevel.cpp
--- a/Engine/Source/Editor/UnrealEd/Private/PlayLevel.cpp
+++ b/Engine/Source/Editor/UnrealEd/Private/PlayLevel.cpp
@@ -156,6 +156,8 @@ void UEditorEngine::EndPlayMap()
 
 	TGuardValue<bool> GuardIsEndingPlay(bIsEndingPlay, true);
 
+	FEditorDelegates::PrePIEEnded.Broadcast( bIsSimulatingInEditor );
+
 	FlushAsyncLoading();
 
 	// Monitoring when PIE corrupts references between the World and the PIE generated World for UE-20486
@@ -185,7 +187,7 @@ void UEditorEngine::EndPlayMap()
 		}
 	}
 
-	if (GEngine->HMDDevice.IsValid())
+	if (GEngine->HMDDevice.IsValid() && !bIsSimulatingInEditor)
 	{
 		GEngine->HMDDevice->OnEndPlay(*GEngine->GetWorldContextFromWorld(PlayWorld));
 	}
@@ -588,8 +590,8 @@ void UEditorEngine::TeardownPlaySession(FWorldContext& PieWorldContext)
 
 				if( !bIsSimulatingInEditor)
 				{
-					// Set the editor viewport location to match that of Play in Viewport if we aren't simulating in the editor, we have a valid player to get the location from 
-					if (bLastViewAndLocationValid == true)
+					// Set the editor viewport location to match that of Play in Viewport if we aren't simulating in the editor, we have a valid player to get the location from (unless we're going back to VR Editor, in which case we won't teleport the user.)
+					if (bLastViewAndLocationValid == true && !GEngine->IsStereoscopic3D( Viewport->GetActiveViewport() ) )
 					{
 						bLastViewAndLocationValid = false;
 						Viewport->GetLevelViewportClient().SetViewLocation( LastViewLocation );
@@ -2420,7 +2422,7 @@ void UEditorEngine::PlayInEditor( UWorld* InWorld, bool bInSimulateInEditor )
 		OutputLogErrorsToMessageLogProxyPtr = MakeShareable(new FOutputLogErrorsToMessageLogProxy());
 	}
 
-	if (GEngine->HMDDevice.IsValid())
+	if (GEngine->HMDDevice.IsValid() && !bInSimulateInEditor)
 	{
 		GEngine->HMDDevice->OnBeginPlay(*GEngine->GetWorldContextFromWorld(InWorld));
 	}
@@ -2572,6 +2574,13 @@ void UEditorEngine::PlayInEditor( UWorld* InWorld, bool bInSimulateInEditor )
 		}
 	}
 
+	// Make sure to focus the game viewport.
+	if (!bInSimulateInEditor)
+	{
+		FSlateApplication::Get().SetAllUserFocusToGameViewport();
+	}
+
+	FEditorDelegates::PostPIEStarted.Broadcast( bInSimulateInEditor );
 }
 
 void UEditorEngine::SpawnIntraProcessPIEWorlds(bool bAnyBlueprintErrors, bool bStartInSpectatorMode)
@@ -3144,7 +3153,7 @@ UGameInstance* UEditorEngine::CreatePIEGameInstance(int32 InPIEInstance, bool bI
 	UGameViewportClient* ViewportClient = NULL;
 	ULocalPlayer *NewLocalPlayer = NULL;
 	
-	if (GEngine->HMDDevice.IsValid())
+	if (GEngine->HMDDevice.IsValid() && !bInSimulateInEditor )
 	{
 		GEngine->HMDDevice->OnBeginPlay(*PieWorldContext);
 	}
diff --git a/Engine/Source/Editor/UnrealEd/Private/SnappingUtils.cpp b/Engine/Source/Editor/UnrealEd/Private/SnappingUtils.cpp
--- a/Engine/Source/Editor/UnrealEd/Private/SnappingUtils.cpp
+++ b/Engine/Source/Editor/UnrealEd/Private/SnappingUtils.cpp
@@ -103,6 +103,7 @@ private:
 	FVertexSnappingImpl VertexSnappingImpl;
 };
 
+
 //////////////////////////////////////////////////////////////////////////
 // FEditorViewportSnapping
 
@@ -287,7 +288,14 @@ bool FEditorViewportSnapping::SnapActorsToNearestActor( FVector& Drag, FLevelEdi
 
 void FEditorViewportSnapping::SnapPointToGrid(FVector& Point, const FVector& GridBase)
 {
-	if( IsSnapToGridEnabled() )
+	/**
+	bool bAlignedSnapUsed = false;
+	if (!LastBestAlignedLocation.IsNearlyZero(EditorViewportSnapping::SnapTolerance->GetFloat()))
+	{
+		Point = LastBestAlignedLocation;
+		bAlignedSnapUsed = true;
+	}**/
+	if(IsSnapToGridEnabled())
 	{
 		Point = (Point - GridBase).GridSnap( GEditor->GetGridSize() ) + GridBase;
 	}
@@ -366,14 +374,16 @@ bool FEditorViewportSnapping::SnapToBSPVertex(FVector& Location, FVector GridBas
 
 void FEditorViewportSnapping::ClearSnappingHelpers( bool bClearImmediately )
 {
-	VertexSnappingImpl.ClearSnappingHelpers( bClearImmediately );
+	VertexSnappingImpl.ClearSnappingHelpers(bClearImmediately);
 }
 
-void FEditorViewportSnapping::DrawSnappingHelpers(const FSceneView* View,FPrimitiveDrawInterface* PDI)
+void FEditorViewportSnapping::DrawSnappingHelpers(const FSceneView* View, FPrimitiveDrawInterface* PDI)
 {
 	VertexSnappingImpl.DrawSnappingHelpers( View, PDI );
 }
 
+
+
 bool FEditorViewportSnapping::SnapLocationToNearestVertex( FVector& Location, const FVector2D& MouseLocation, FLevelEditorViewportClient* ViewportClient, FVector& OutVertexNormal, bool bDrawVertHelpers )
 {
 	bool bSnapped = false;
@@ -409,7 +419,6 @@ bool FEditorViewportSnapping::SnapDragLocationToNearestVertex( const FVector& Ba
 	}
 	return bSnapped;
 }
-
 //////////////////////////////////////////////////////////////////////////
 // FSnappingUtils
 
diff --git a/Engine/Source/Editor/UnrealEd/Public/EdMode.h b/Engine/Source/Editor/UnrealEd/Public/EdMode.h
--- a/Engine/Source/Editor/UnrealEd/Public/EdMode.h
+++ b/Engine/Source/Editor/UnrealEd/Public/EdMode.h
@@ -183,7 +183,7 @@ public:
 	 * @param	InBox	The selection box to use, in worldspace coordinates.
 	 * @return		true if something was selected/deselected, false otherwise.
 	 */
-	bool BoxSelect( FBox& InBox, bool InSelect = true );
+	virtual bool BoxSelect( FBox& InBox, bool InSelect = true );
 
 	/**
 	 * Lets each mode/tool handle frustum selection in its own way.
@@ -191,7 +191,7 @@ public:
 	 * @param	InFrustum	The selection box to use, in worldspace coordinates.
 	 * @return	true if something was selected/deselected, false otherwise.
 	 */
-	bool FrustumSelect( const FConvexVolume& InFrustum, bool InSelect = true );
+	virtual bool FrustumSelect( const FConvexVolume& InFrustum, FEditorViewportClient* InViewportClient, bool InSelect = true );
 
 	void SelectNone();
 	virtual void SelectionChanged() {}
diff --git a/Engine/Source/Editor/UnrealEd/Public/Editor.h b/Engine/Source/Editor/UnrealEd/Public/Editor.h
--- a/Engine/Source/Editor/UnrealEd/Public/Editor.h
+++ b/Engine/Source/Editor/UnrealEd/Public/Editor.h
@@ -132,8 +132,12 @@ struct UNREALED_API FEditorDelegates
 	static FOnEditorModeTransitioned EditorModeExit;
 	/** Sent when a PIE session is beginning (before we decide if PIE can run - allows clients to avoid blocking PIE) */
 	static FOnPIEEvent PreBeginPIE;
-	/** Sent when a PIE session is beginning */
+	/** Sent when a PIE session is beginning (but hasn't actually started yet) */
 	static FOnPIEEvent BeginPIE;
+	/** Sent when a PIE session has fully started and after BeginPlay() has been called */
+	static FOnPIEEvent PostPIEStarted;
+	/** Sent when a PIE session is ending, before anything else happens */
+	static FOnPIEEvent PrePIEEnded;
 	/** Sent when a PIE session is ending */
 	static FOnPIEEvent EndPIE;
 	/** Sent when a PIE session is paused */
diff --git a/Engine/Source/Editor/UnrealEd/Public/EditorModeManager.h b/Engine/Source/Editor/UnrealEd/Public/EditorModeManager.h
--- a/Engine/Source/Editor/UnrealEd/Public/EditorModeManager.h
+++ b/Engine/Source/Editor/UnrealEd/Public/EditorModeManager.h
@@ -180,7 +180,7 @@ public:
 	bool BoxSelect( FBox& InBox, bool InSelect );
 
 	/** Notifies all active modes of frustum selection attempts */
-	bool FrustumSelect( const FConvexVolume& InFrustum, bool InSelect );
+	bool FrustumSelect( const FConvexVolume& InFrustum, FEditorViewportClient* InViewportClient, bool InSelect );
 
 	/** true if any active mode uses a transform widget */
 	bool UsesTransformWidget() const;
diff --git a/Engine/Source/Editor/UnrealEd/Public/EditorModeTools.h b/Engine/Source/Editor/UnrealEd/Public/EditorModeTools.h
--- a/Engine/Source/Editor/UnrealEd/Public/EditorModeTools.h
+++ b/Engine/Source/Editor/UnrealEd/Public/EditorModeTools.h
@@ -127,7 +127,7 @@ public:
 	}
 	//@}
 
-	virtual bool FrustumSelect( const FConvexVolume& InFrustum, bool InSelect = true )
+	virtual bool FrustumSelect( const FConvexVolume& InFrustum, FEditorViewportClient* InViewportClient, bool InSelect = true )
 	{
 		return false;
 	}
diff --git a/Engine/Source/Editor/UnrealEd/Public/EditorWorldExtension.h b/Engine/Source/Editor/UnrealEd/Public/EditorWorldExtension.h
new file mode 100644
--- /dev/null
+++ b/Engine/Source/Editor/UnrealEd/Public/EditorWorldExtension.h
@@ -0,0 +1,201 @@
+// Copyright 1998-2017 Epic Games, Inc. All Rights Reserved.
+
+#pragma once
+
+#include "CoreMinimal.h"
+#include "UObject/GCObject.h"
+#include "EditorWorldExtension.generated.h"
+
+struct FWorldContext;
+
+UCLASS()
+class UNREALED_API UEditorWorldExtension : public UObject
+{
+	GENERATED_BODY()
+
+	friend class UEditorWorldExtensionCollection;
+public:
+	
+	/** Default constructor */
+	UEditorWorldExtension();
+
+	/** Default destructor */
+	virtual ~UEditorWorldExtension();
+
+	/** Initialize extension */
+	virtual void Init() {};
+
+	/** Shut down extension when world is destroyed */
+	virtual void Shutdown() {};
+
+	/** Give base class the chance to tick */
+	virtual void Tick( float DeltaSeconds ) {};
+
+	virtual bool InputKey( FEditorViewportClient* InViewportClient, FViewport* Viewport, FKey Key, EInputEvent Event);
+	virtual bool InputAxis( FEditorViewportClient* InViewportClient, FViewport* Viewport, int32 ControllerId, FKey Key, float Delta, float DeltaTime);
+
+	/** Gets the world owning this extension */
+	virtual UWorld* GetWorld() const override;
+
+	/**  Spawns a transient actor that we can use in the current world of this extension (templated for convenience) */
+	template<class T>
+	inline T* SpawnTransientSceneActor(const FString& ActorName, const bool bWithSceneComponent = false, const EObjectFlags InObjectFlags = EObjectFlags::RF_Transient | EObjectFlags::RF_DuplicateTransient )
+	{
+		return CastChecked<T>(SpawnTransientSceneActor(T::StaticClass(), ActorName, bWithSceneComponent, InObjectFlags));
+	}
+
+	/** Spawns a transient actor that we can use in the current world of this extension */
+	AActor* SpawnTransientSceneActor(TSubclassOf<AActor> ActorClass, const FString& ActorName, const bool bWithSceneComponent = false, const EObjectFlags InObjectFlags = EObjectFlags::RF_Transient | EObjectFlags::RF_DuplicateTransient );
+
+	/** Destroys a transient actor we created earlier */
+	void DestroyTransientActor(AActor* Actor);
+
+	/** Get the owning collection of extensions */
+	UEditorWorldExtensionCollection* GetOwningCollection();
+
+	/** Executes command */
+	bool ExecCommand(const FString& InCommand);
+
+protected:
+	
+	/** Reparent actors to a new world */
+	virtual void TransitionWorld(UWorld* NewWorld);
+
+	/** Give base class a chance to act on entering simulate mode */
+	virtual void EnteredSimulateInEditor() {};
+	
+	/** Give base class a chance to act on leaving simulate mode */
+	virtual void LeftSimulateInEditor() {};
+
+	/** The collection of extensions that is owning this extension */
+	UEditorWorldExtensionCollection* OwningExtensionsCollection;
+
+private:
+
+	void ReparentActor(AActor* Actor, UWorld* NewWorld);
+
+	/** Let the FEditorWorldExtensionCollection set the world of this extension before init */
+	void InitInternal(UEditorWorldExtensionCollection* InOwningExtensionsCollection);
+
+	UPROPERTY()
+	TArray<AActor*> ExtensionActors;
+};
+
+/**
+ * Holds a collection of UEditorExtension
+ */
+UCLASS()
+class UNREALED_API UEditorWorldExtensionCollection : public UObject
+{
+	GENERATED_BODY()
+
+	friend class UEditorWorldExtensionManager;
+public:
+
+	/** Default constructor */
+	UEditorWorldExtensionCollection();
+
+	/** Default destructor */
+	virtual ~UEditorWorldExtensionCollection();
+
+	/** Gets the world from the world context */
+	virtual UWorld* GetWorld() const override;
+
+	/** Gets the world context */
+	FWorldContext* GetWorldContext() const;
+
+	/** 
+	 * Adds an extension to the collection
+	 * @param	EditorExtensionClass	The subclass of UEditorExtension that will be created, initialized and added to the collection.
+	 * @return							The created UEditorExtension.
+	 */
+	UEditorWorldExtension* AddExtension( TSubclassOf<UEditorWorldExtension> EditorExtensionClass );
+
+	/** 
+	 * Adds an extension to the collection
+	 * @param	EditorExtension			The UEditorExtension that will be created, initialized and added to the collection.
+	 */
+	void AddExtension( UEditorWorldExtension* EditorExtension );
+	
+	/** 
+	 * Removes an extension from the collection and calls Shutdown() on the extension
+	 * @param	EditorExtension			The UEditorExtension to remove.  It must already have been added.
+	 */
+	void RemoveExtension( UEditorWorldExtension* EditorExtension );
+
+	/**
+	 * Find an extension based on the class
+	 * @param	EditorExtensionClass	The class to find an extension with
+	 * @return							The first extension that is found based on class
+	 */
+	UEditorWorldExtension* FindExtension( TSubclassOf<UEditorWorldExtension> EditorExtensionClass );
+
+	/** Ticks all extensions */
+	void Tick( float DeltaSeconds );
+
+	/** Notifies all extensions of keyboard input */
+	bool InputKey( FEditorViewportClient* InViewportClient, FViewport* Viewport, FKey Key, EInputEvent Event);
+
+	/** Notifies all extensions of axis movement */
+	bool InputAxis( FEditorViewportClient* InViewportClient, FViewport* Viewport, int32 ControllerId, FKey Key, float Delta, float DeltaTime);
+	
+private:
+
+	/** Sets the world for this collection and gives every extension an opportunity to transition */
+	void SetWorldContext(FWorldContext* InWorldContext);
+
+	/** Called by the editor after PIE or Simulate is started */
+	void PostPIEStarted( bool bIsSimulatingInEditor );
+
+	/** Called when PIE or Simulate ends */
+	void OnEndPIE( bool bWasSimulatingInEditor );
+
+	/** World context */
+	FWorldContext* WorldContext;
+
+	/** After entering Simulate, this stores the counterpart editor world to the Simulate world, so that we
+        know this collection needs to transition back to editor world after Simulate finishes */
+	FWorldContext* EditorWorldOnSimulate;
+
+	/** List of extensions along with their reference count.  Extensions will only be truly removed and Shutdown() after their
+	    reference count drops to zero. */
+	typedef TTuple<UEditorWorldExtension*, int32> FEditorExtensionTuple;
+	TArray<FEditorExtensionTuple> EditorExtensions;
+};
+
+/**
+ * Holds a map of extension collections paired with worlds
+ */
+UCLASS()
+class UNREALED_API UEditorWorldExtensionManager	: public UObject
+{
+	GENERATED_BODY()
+
+public:
+	/** Default constructor */
+	UEditorWorldExtensionManager();
+
+	/** Default destructor */
+	virtual ~UEditorWorldExtensionManager();
+
+	/** Gets the editor world wrapper that is found with the world passed.
+	 * Adds one for this world if there was non found. */
+	UEditorWorldExtensionCollection* GetEditorWorldExtensions(UWorld* InWorld);
+
+	/** Ticks all the collections */
+	void Tick( float DeltaSeconds );
+
+private:
+
+	/** Adds a new editor world wrapper when a new world context was created */
+	UEditorWorldExtensionCollection* OnWorldContextAdd(FWorldContext* InWorldContext);
+
+	/** Adds a new editor world wrapper when a new world context was created */
+	void OnWorldContextRemove(FWorldContext& InWorldContext);
+
+	UEditorWorldExtensionCollection* FindExtensionCollection(const UWorld* InWorld);
+
+	/** Map of all the editor world maps */
+	UPROPERTY()
+	TArray<UEditorWorldExtensionCollection*> EditorWorldExtensionCollection;
+};
\ No newline at end of file
diff --git a/Engine/Source/Editor/UnrealEd/UnrealEd.Build.cs b/Engine/Source/Editor/UnrealEd/UnrealEd.Build.cs
--- a/Engine/Source/Editor/UnrealEd/UnrealEd.Build.cs
+++ b/Engine/Source/Editor/UnrealEd/UnrealEd.Build.cs
@@ -63,8 +63,6 @@ public class UnrealEd : ModuleRules
                 "PortalRpc",
                 "PortalServices",
                 "BlueprintNativeCodeGen",
-                "ViewportInteraction",
-                "VREditor",
                 "Persona",
             }
 		);
@@ -151,8 +149,7 @@ public class UnrealEd : ModuleRules
                 "PixelInspectorModule",
 				"MovieScene",
 				"MovieSceneTracks",
-                "ViewportInteraction",
-                "VREditor"
+                "ViewportInteraction"
             }
 		);
 
@@ -218,7 +215,8 @@ public class UnrealEd : ModuleRules
                 "IOSPlatformEditor",
 				"HTML5PlatformEditor",
 				"SizeMap",
-                "PortalProxies",
+				"MeshEditor",
+				"PortalProxies",
                 "PortalServices",
                 "GeometryCacheEd",
                 "BlueprintNativeCodeGen",
@@ -235,7 +233,8 @@ public class UnrealEd : ModuleRules
 			{
                 "GraphEditor",
 				"Kismet",
-				"AudioEditor"
+				"AudioEditor",
+                "ViewportInteraction"
             }
 		); 
 
@@ -299,4 +298,4 @@ public class UnrealEd : ModuleRules
 			Definitions.Add( "WITH_RECAST=0" );
 		}
 	}
-}
+}
\ No newline at end of file
diff --git a/Engine/Source/Editor/VREditor/Public/VREditorAssetContainer.h b/Engine/Source/Editor/VREditor/Public/VREditorAssetContainer.h
new file mode 100644
--- /dev/null
+++ b/Engine/Source/Editor/VREditor/Public/VREditorAssetContainer.h
@@ -0,0 +1,82 @@
+// Copyright 1998-2017 Epic Games, Inc. All Rights Reserved.
+
+#pragma once
+
+#include "SlateWidgetStyleContainerBase.h"
+#include "SlateWidgetStyle.h"
+#include "Sound/SoundBase.h"
+#include "VREditorAssetContainer.generated.h"
+
+/**
+ * Represents the common menu sounds used in viewport interaction
+ */
+USTRUCT()
+struct FVREditorAssetContainer : public FSlateWidgetStyle
+{
+	GENERATED_USTRUCT_BODY()
+
+	FVREditorAssetContainer();
+	virtual ~FVREditorAssetContainer();
+
+	// FSlateWidgetStyle
+	virtual void GetResources(TArray<const FSlateBrush*>& OutBrushes) const override;
+	static const FName TypeName;
+	virtual const FName GetTypeName() const override { return TypeName; };
+	static const FVREditorAssetContainer& GetDefault();
+
+	/** The sound that should play when starting the game */
+	UPROPERTY(EditAnywhere, Category = Sound)
+	USoundBase* DockableWindowCloseSound;
+	FVREditorAssetContainer& SetDockableWindowCloseSound(USoundBase* InSound) { DockableWindowCloseSound = InSound; return *this; }
+
+	UPROPERTY(EditAnywhere, Category = Sound)
+	USoundBase* DockableWindowOpenSound;
+	FVREditorAssetContainer& SetDockableWindowOpenSound(USoundBase* InSound) { DockableWindowOpenSound = InSound; return *this; }
+
+	UPROPERTY(EditAnywhere, Category = Sound)
+	USoundBase* DockableWindowDropSound;
+	FVREditorAssetContainer& SetDockableWindowDropSound(USoundBase* InSound) { DockableWindowDropSound = InSound; return *this; }
+
+	UPROPERTY(EditAnywhere, Category = Sound)
+	USoundBase* DockableWindowDragSound;
+	FVREditorAssetContainer& SetDockableWindowDragSound(USoundBase* InSound) { DockableWindowDragSound = InSound; return *this; }
+
+	UPROPERTY(EditAnywhere, Category = Sound)
+	USoundBase* DropFromContentBrowserSound;
+	FVREditorAssetContainer& SetDropFromContentBrowserSound(USoundBase* InSound) { DropFromContentBrowserSound = InSound; return *this; }
+
+	UPROPERTY(EditAnywhere, Category = Sound)
+	USoundBase* RadialMenuOpenSound;
+	FVREditorAssetContainer& SetRadialMenuOpenSound(USoundBase* InSound) { RadialMenuOpenSound = InSound; return *this; }
+
+	UPROPERTY(EditAnywhere, Category = Sound)
+	USoundBase* RadialMenuCloseSound;
+	FVREditorAssetContainer& SetRadialMenuCloseSound(USoundBase* InSound) { RadialMenuCloseSound = InSound; return *this; }
+
+	UPROPERTY(EditAnywhere, Category = Sound)
+	USoundBase* TeleportSound;
+	FVREditorAssetContainer& SetTeleportSound(USoundBase* InSound) { TeleportSound = InSound; return *this; }
+
+	UPROPERTY(EditAnywhere, Category = Sound)
+	USoundCue* ButtonPressSound;
+	FVREditorAssetContainer& SetTeleportSound(USoundCue* InSound) { ButtonPressSound = InSound; return *this; }
+};
+
+
+/**
+ */
+UCLASS(hidecategories=Object, MinimalAPI)
+class UVREditorStyleContainer : public USlateWidgetStyleContainerBase
+{
+	GENERATED_UCLASS_BODY()
+
+public:
+	/** The actual data describing the sounds */
+	UPROPERTY(Category=Appearance, EditAnywhere, meta=(ShowOnlyInnerProperties))
+	FVREditorAssetContainer SoundsStyle;
+
+	virtual const struct FSlateWidgetStyle* const GetStyle() const override
+	{
+		return static_cast< const struct FSlateWidgetStyle* >( &SoundsStyle );
+	}
+};
diff --git a/Engine/Source/Editor/VREditor/Public/VREditorInteractor.h b/Engine/Source/Editor/VREditor/Public/VREditorInteractor.h
--- a/Engine/Source/Editor/VREditor/Public/VREditorInteractor.h
+++ b/Engine/Source/Editor/VREditor/Public/VREditorInteractor.h
@@ -12,6 +12,15 @@
 class AActor;
 class UWidgetComponent;
 
+enum class EControllerType : uint8
+{
+	Laser,
+	AssistingLaser,
+	UI,
+	Unknown
+};
+
+
 /**
  * VREditor default interactor
  */
@@ -44,8 +53,9 @@ public:
 	virtual void Tick( const float DeltaTime ) override;
 	virtual FHitResult GetHitResultFromLaserPointer( TArray<AActor*>* OptionalListOfIgnoredActors = nullptr, const bool bIgnoreGizmos = false,
 		TArray<UClass*>* ObjectsInFrontOfGizmo = nullptr, const bool bEvenIfBlocked = false, const float LaserLengthOverride = 0.0f ) override;
-	virtual void ResetHoverState( const float DeltaTime ) override;
-	virtual void OnStartDragging( UActorComponent* ClickedComponent, const FVector& HitLocation, const bool bIsPlacingActors ) override;
+	virtual void ResetHoverState() override;
+	virtual bool IsModifierPressed() const override;
+	virtual void PreviewInputKey( class FEditorViewportClient& ViewportClient, FViewportActionKeyInput& Action, const FKey Key, const EInputEvent Event, bool& bOutWasHandled ) override;
 
 	/** Returns the slide delta for pushing and pulling objects. Needs to be implemented by derived classes (e.g. touchpad for vive controller or scrollweel for mouse ) */
 	virtual float GetSlideDelta();
@@ -54,14 +64,17 @@ public:
 	// Getters and setters
 	//
 
+	/** Returns what controller type this is for asymmetric control schemes */
+	EControllerType GetControllerType() const
+	{
+		return ControllerType;
+	};
+
+	/** Set what controller type this is for asymmetric control schemes */
+	void SetControllerType( const EControllerType& InControllerType );
+
 	/** Gets if this interactor is hovering over UI */
 	bool IsHoveringOverUI() const;
-	
-	/** Sets if there is a docked window on this interactor */
-	void SetHasUIInFront( const bool bInHasUIInFront );
-	
-	/** Check if there is any docked window on this interactor */
-	bool HasUIInFront() const;
 
 	/** Sets if the quick menu is on this interactor */
 	void SetHasUIOnForearm( const bool bInHasUIOnForearm );
@@ -70,13 +83,10 @@ public:
 	bool HasUIOnForearm() const;
 
 	/** Gets the current hovered widget component if any */
-	UWidgetComponent* GetHoveringOverWidgetComponent() const;
+	UWidgetComponent* GetLastHoveredWidgetComponent() const;
 
 	/** Sets the current hovered widget component */
-	void SetHoveringOverWidgetComponent( UWidgetComponent* NewHoveringOverWidgetComponent );
-
-	/** If the modifier key is currently pressed */
-	bool IsModifierPressed() const;
+	void SetLastHoveredWidgetComponent( UWidgetComponent* NewHoveringOverWidgetComponent );
 
 	/** Sets if the interactor is clicking on any UI */
 	void SetIsClickingOnUI( const bool bInIsClickingOnUI );
@@ -93,11 +103,11 @@ public:
 	/** Gets if the interactor is right clicking on UI */
 	bool IsRightClickingOnUI() const;
 
-	/** Gets if the interactor is right clicking on UI */
-	void SetLastClickReleaseTime( const double InLastClickReleaseTime );
+	/** Sets the time the interactor last pressed on UI */
+	void SetLastUIPressTime( const double InLastUIPressTime );
 
-	/** Gets last time the interactor released */
-	double GetLastClickReleaseTime() const;
+	/** Gets last time the interactor pressed on UI */
+	double GetLastUIPressTime() const;
 
 	/** Sets the UI scroll velocity */
 	void SetUIScrollVelocity( const float InUIScrollVelocity );
@@ -109,7 +119,7 @@ public:
 	float GetSelectAndMoveTriggerValue() const;
 	
 	/* ViewportInteractor overrides, checks if the laser is blocked by UI */
-	virtual bool GetIsLaserBlocked() override;
+	virtual bool GetIsLaserBlocked() const override;
 	
 protected:
 
@@ -121,6 +131,9 @@ protected:
 	// General input @todo: VREditor: Should this be general (non-UI) in interactordata ?
 	//
 
+	/** For asymmetrical systems - what type of controller this is */
+	EControllerType ControllerType;
+
 	/** Is the Modifier button held down? */
 	bool bIsModifierPressed;
 
@@ -152,8 +165,8 @@ private:
 	/** Inertial scrolling -- how fast to scroll the mousewheel over UI */
 	float UIScrollVelocity;	
 
-	/** Last real time that we released the 'SelectAndMove' button on UI.  This is used to detect double-clicks. */
-	double LastClickReleaseTime;
+	/** Last real time that we pressed the 'SelectAndMove' button on UI.  This is used to detect double-clicks. */
+	double LastUIPressTime;
 
 protected:
 
@@ -169,4 +182,4 @@ protected:
 
 	/** Time that we either started showing or hiding help labels (for fade transitions) */
 	FTimespan HelpLabelShowOrHideStartTime;
-};
+};
\ No newline at end of file
diff --git a/Engine/Source/Editor/VREditor/Public/VREditorMode.h b/Engine/Source/Editor/VREditor/Public/VREditorMode.h
--- a/Engine/Source/Editor/VREditor/Public/VREditorMode.h
+++ b/Engine/Source/Editor/VREditor/Public/VREditorMode.h
@@ -8,32 +8,26 @@
 #include "InputCoreTypes.h"
 #include "Templates/SubclassOf.h"
 #include "Editor/UnrealEdTypes.h"
+#include "EditorWorldExtension.h"
 #include "ShowFlags.h"
 #include "Misc/App.h"
 #include "Widgets/SWindow.h"
 #include "HeadMountedDisplayTypes.h"
 #include "Kismet/HeadMountedDisplayFunctionLibrary.h"
+#include "UI/VRRadialMenuHandler.h"
 #include "VREditorMode.generated.h"
 
 class AActor;
 class FEditorViewportClient;
 class SLevelViewport;
-enum class EGizmoHandleTypes : uint8;
+enum class EAutoKeyMode : uint8;
 
 // Forward declare the GizmoHandleTypes that is defined in VIBaseTransformGizmo.h
 enum class EGizmoHandleTypes : uint8;
 
-/** Reason for exiting the VR editor */
-enum EVREditorExitType
-{
-	Normal = 0,
-	PIE_VR = 1,		//Play in editor with VR
-	SIE_VR = 2		//Simulate in editor with VR
-};
-
 /**
-* Types of actions that can be performed with VR controller devices
-*/
+ * Types of actions that can be performed with VR controller devices
+ */
 namespace VRActionTypes
 {
 	static const FName Touch( "Touch" );
@@ -46,19 +40,39 @@ namespace VRActionTypes
 }
 
 /**
- * VR Editor Mode.  Extends editor viewports with functionality for VR controls and object manipulation
+ * VR Editor Mode. Extends editor viewports with functionality for VR controls and object manipulation
  */
 UCLASS()
-class VREDITOR_API UVREditorMode : public UObject
+class VREDITOR_API UVREditorMode : public UEditorWorldExtension
 {
-
-	GENERATED_UCLASS_BODY()
+	GENERATED_BODY()
 
 public:
 
+	/** Default constructor */
+	UVREditorMode();
+
 	/** Cleans up this mode, called when the editor is shutting down */
 	virtual ~UVREditorMode();
 
+	/** Initialize the VREditor */
+	virtual void Init() override;
+
+	/** Shutdown the VREditor */
+	virtual void Shutdown() override;	
+	
+	/** When the user actually enters the VR Editor mode */
+	void Enter();
+
+	/** When the user leaves the VR Editor mode */
+	void Exit( const bool bShouldDisableStereo );
+
+	/** Tick before the ViewportWorldInteraction is ticked */
+	void PreTick( const float DeltaTime );
+
+	/** Tick after the ViewportWorldInteraction is ticked */
+	void PostTick( const float DeltaTime );
+
 	/** Static: Sets whether we should actually use an HMD.  Call this before activating VR mode */
 	void SetActuallyUsingVR( const bool bShouldActuallyUseVR )
 	{
@@ -78,47 +92,8 @@ public:
 	}
 
 	/** Call this to start exiting VR mode */
-	void StartExitingVRMode( const EVREditorExitType InExitType = EVREditorExitType::Normal );
-
-	/** Gets the reason for exiting the mode */
-	EVREditorExitType GetExitType() const
-	{
-		return ExitType;
-	}
-
-	/** EditorWorldManager initializes the start of the editor */
-	void Init( class UViewportWorldInteraction* InViewportWorldInteraction );
-
-	/** EditorWorldManager shuts down the VREditor when closing the editor */
-	void Shutdown();
-
-	/** When the user actually enters the VR Editor mode */
-	void Enter(const bool bReenteringVREditing);
-
-	/** When the user leaves the VR Editor mode */
-	void Exit(const bool bHMDShouldExitStereo);
-
-	/** Tick before the ViewportWorldInteraction is ticked */
-	void PreTick( const float DeltaTime );
-
-	/** Tick after the ViewportWorldInteraction is ticked */
-	void Tick( const float DeltaTime );
+	void StartExitingVRMode();
 	
-	//bool InputKey( FEditorViewportClient* ViewportClient, FViewport* Viewport, FKey Key, EInputEvent Event );
-	//bool InputAxis( FEditorViewportClient* ViewportClient, FViewport* Viewport, int32 ControllerId, FKey Key, float Delta, float DeltaTime );
-	//bool IsCompatibleWith( FEditorModeID OtherModeID ) const;
-	//void Render( const FSceneView* SceneView, FViewport* Viewport, FPrimitiveDrawInterface* PDI );	
-
-	/**
-	 * Gets our avatar's mesh actor
-	 *
-	 * @return	The mesh actor
-	 */
-	AActor* GetAvatarMeshActor();
-
-	/** Gets the World from the ViewportWorldInteraction used by this mode */
-	UWorld* GetWorld() const;
-
 	/** Gets the world space transform of the calibrated VR room origin.  When using a seated VR device, this will feel like the
 	camera's world transform (before any HMD positional or rotation adjustments are applied.) */
 	FTransform GetRoomTransform() const;
@@ -182,19 +157,6 @@ public:
 		return 1.0f - ( ( Alpha * ( ( OvershootAmount + 1 ) * Alpha + OvershootAmount ) + 1 ) - 1.0f );
 	}
 
-	/** Spawns a transient actor that we can use in the editor world (templated for convenience) */
-	template<class T>
-	T* SpawnTransientSceneActor( const FString& ActorName, const bool bWithSceneComponent ) const
-	{
-		return CastChecked<T>( SpawnTransientSceneActor( T::StaticClass(), ActorName, bWithSceneComponent ) );
-	}
-
-	/** Spawns a transient actor that we can use in the editor world */
-	AActor* SpawnTransientSceneActor( TSubclassOf<class AActor> ActorClass, const FString& ActorName, const bool bWithSceneComponent ) const;
-
-	/** Destroys a transient actor we created earlier and nulls out the pointer */
-	void DestroyTransientActor( AActor* Actor ) const;
-
 	/** Gets access to the VR UI system (const) */
 	const class UVREditorUISystem& GetUISystem() const
 	{
@@ -221,9 +183,6 @@ public:
 	/** Gets the world scale factor, which can be multiplied by a scale vector to convert to room space */
 	float GetWorldScaleFactor() const;
 
-	/** Sets up our avatar mesh, if not already spawned */
-	void SpawnAvatarMeshActor();
-
 	/** Called internally when the user changes maps, enters/exits PIE or SIE, or switched between PIE/SIE */
 	void CleanUpActorsBeforeMapChangeOrSimulate();
 
@@ -245,9 +204,6 @@ public:
 	/** Gets the hand interactor  */
 	class UVREditorInteractor* GetHandInteractor( const EControllerHand ControllerHand ) const;
 
-	/** Stops the haptic feedback for the left and right hand interactors for the oculus */
-	void StopOldHapticEffects();
-
 	/** Snaps the current selected actor to the ground */
 	void SnapSelectedActorsToGround();
 	
@@ -265,23 +221,31 @@ public:
 		bool bRealTime;
 		float DragTriggerDistance;
 		bool bOnScreenMessages;
+		float TransformGizmoScale;
 		EHMDTrackingOrigin::Type TrackingOrigin;
 		float WorldToMetersScale;
+		bool bCinematicPreviewViewport;
+		bool bKeyAllEnabled;
+		EAutoKeyMode AutoKeyMode;
 
 		FSavedEditorState()
-			: ViewportType( LVT_Perspective ),
-			  ViewLocation( FVector::ZeroVector ),
-			  ViewRotation( FRotator::ZeroRotator ),
-			  ShowFlags( ESFIM_Editor ),
-			  bLockedPitch( false ),
-			  bGameView( false ),
-			  bAlwaysShowModeWidgetAfterSelectionChanges( false ),
-			  NearClipPlane( 0.0f ),
-			  bRealTime( false ),
-			  DragTriggerDistance( 0.0f ),
-			  bOnScreenMessages( false ),
-			  TrackingOrigin( EHMDTrackingOrigin::Eye ),
-			  WorldToMetersScale( 100.0f )
+			: ViewportType(LVT_Perspective),
+			  ViewLocation(FVector::ZeroVector),
+			  ViewRotation(FRotator::ZeroRotator),
+ 			  ShowFlags(ESFIM_Editor),
+			  bLockedPitch(false),
+			  bGameView(false),
+			  bAlwaysShowModeWidgetAfterSelectionChanges(false),
+			  NearClipPlane(0.0f),
+			  bRealTime(false),
+			  DragTriggerDistance(0.0f),
+			  bOnScreenMessages(false),
+			  TransformGizmoScale( 1.0f ),
+			  TrackingOrigin(EHMDTrackingOrigin::Eye),
+			  WorldToMetersScale(100.0f),
+			  bCinematicPreviewViewport(false),
+			  bKeyAllEnabled(false),
+			  AutoKeyMode()
 		{
 		}
 	};
@@ -289,21 +253,75 @@ public:
 	/** Gets the saved editor state from entering the mode */
 	const FSavedEditorState& GetSavedEditorState() const;
 
+	DECLARE_DELEGATE(FOnVREditingModeExit);
+	/** Used to override dockable area restoration behavior */
+	FOnVREditingModeExit OnVREditingModeExit_Handler;
+
+	void SaveSequencerSettings(bool bInKeyAllEnabled, EAutoKeyMode InAutoKeyMode, const class USequencerSettings& InSequencerSettings);
+
+	/** Start or stop simulate-in-editor mode */
+	void ToggleSIEAndVREditor();
+
+	/** Start or stop play-in-editor mode */
+	void TogglePIEAndVREditor();
+
+	/** Create a static motion controller mesh for the current HMD platform */
+	UStaticMeshComponent* CreateMotionControllerMesh( AActor* OwningActor, USceneComponent* AttachmentToComponent );
+
+	/** Helper function to create a static mesh */
+	UStaticMeshComponent* CreateMesh( AActor* OwningActor, const FString& MeshName, USceneComponent* AttachmentToComponent /*= nullptr */ );
+
+	/** Sets a delegate for the context-specific actions menu */
+	void SetActionsMenuGenerator(const FOnRadialMenuGenerated NewMenuGenerator, const FText NewLabel);
+
+	/** Resets the delegate and button for the context-specific actions menu */
+	void ResetActionsMenuGenerator();
+
+	/** Gets access to VREditorWorldInteraction */
+	class UVREditorWorldInteraction* GetVREditorWorldInteraction()
+	{
+		return VRWorldInteractionExtension;
+	}
+
+	/** Returns true if we started the play in editor session from this VR Editor */
+	bool GetStartedPlayFromVREditor() const;
+	
+	/** Gets the container for all the assets of VREditor. */
+	const struct FVREditorAssetContainer& GetAssetContainer() const;
+	
+	/** Plays sound at location. */
+	void PlaySound(USoundBase* SoundBase, const FVector& InWorldLocation, const float InVolume = 1.0f);
+
+	/** Delegate to be called when a material is placed **/
+	DECLARE_EVENT_ThreeParams( UVREditorWorldInteraction, FOnPlaceDraggedMaterial, UPrimitiveComponent*, UMaterialInterface*, bool& );
+	FOnPlaceDraggedMaterial& OnPlaceDraggedMaterial() { return OnPlaceDraggedMaterialEvent; };
+
+protected:
+
+	virtual void TransitionWorld(UWorld* NewWorld) override;
+	virtual void LeftSimulateInEditor() override;
+
 private:
 
 	/** Called when the editor is closed */
 	void OnEditorClosed();
 
-	//Handles closing the VR mode by escape key
-	void InputKey(const FEditorViewportClient& InViewportClient, const FKey InKey, const EInputEvent InEvent, bool& bOutWasHandled);
-
 	/** Called when someone closes a standalone VR Editor window */
 	void OnVREditorWindowClosed( const TSharedRef<SWindow>& ClosedWindow );
 
 	/** FEditorDelegates callbacks */
-	void OnBeginPIE( const bool bIsSimulatingInEditor );
-	void OnEndPIE( const bool bIsSimulatingInEditor );
-	void OnSwitchBetweenPIEAndSIE( const bool bIsSimulatingInEditor );
+	void PostPIEStarted( bool bIsSimulatingInEditor );
+	void PrePIEEnded( bool bWasSimulatingInEditor );
+	void OnEndPIE( bool bWasSimulatingInEditor );
+
+	/** Start using the viewport passed */
+	void StartViewport( TSharedPtr<SLevelViewport> Viewport );
+	
+	/** Close the current viewport */
+	void CloseViewport( const bool bShouldDisableStereo );
+
+	/** Restore the world to meters to the saved one when entering VR Editor */
+	void RestoreWorldToMeters();
 
 protected:
 
@@ -326,9 +344,6 @@ protected:
 	/** True if we currently want to exit VR mode.  This is used to defer exiting until it is safe to do that */
 	bool bWantsToExitMode;
 
-	/** The reason for exiting the mode, so the module can close the mode and take further actions on what should happen next */
-	EVREditorExitType ExitType;
-
 	/** True if VR mode is fully initialized and ready to render */
 	bool bIsFullyInitialized;
 
@@ -368,6 +383,9 @@ protected:
 
 	FOnVRTickHandle TickHandle;
 
+	/** Event broadcast when a material is placed */
+	FOnPlaceDraggedMaterial OnPlaceDraggedMaterialEvent;
+
 	//
 	// Subsystems
 	//
@@ -378,7 +396,7 @@ protected:
 
 	/** Teleporter system */
 	UPROPERTY()
-	class UVREditorTeleporter* TeleporterSystem;
+	class AVREditorTeleporter* TeleportActor;
 
 	/** Automatic scale system */
 	UPROPERTY()
@@ -402,10 +420,6 @@ protected:
 	// Interactors
 	//
 
-	/** The mouse cursor interactor (currently only available when not in VR) */
-	UPROPERTY()
-	class UMouseCursorInteractor* MouseCursorInteractor;
-
 	/** The right motion controller */
 	UPROPERTY()
 	class UVREditorMotionControllerInteractor* LeftHandInteractor; //@todo vreditor: Hardcoded interactors
@@ -423,15 +437,8 @@ public:
 	{
 		DefaultColor,
 		SelectionColor,
-		TeleportColor,
-		WorldDraggingColor_OneHand,
-		WorldDraggingColor_TwoHands,
-		RedGizmoColor,
-		GreenGizmoColor,
-		BlueGizmoColor,
-		WhiteGizmoColor,
-		HoverGizmoColor,
-		DraggingGizmoColor,
+		WorldDraggingColor,
+		UIColor,
 		UISelectionBarColor,
 		UISelectionBarHoverColor,
 		UICloseButtonColor,
@@ -445,6 +452,12 @@ public:
 	// Get the default near clipping plane for VR editing
 	float GetDefaultVRNearClipPlane() const;
 
+	/** Runtime and plugin modules can force VR Editor to refresh using this function */
+	void RefreshVREditorSequencer(class ISequencer* InCurrentSequencer);
+
+	/** Returns the currently active sequencer */
+	class ISequencer* GetCurrentSequencer();
+
 private:
 
 	// All the colors for this mode
@@ -453,9 +466,18 @@ private:
 	/** If this is the first tick or before */
 	bool bFirstTick;
 
-	/** If the coordinate system was in world space before entering (local) scale mode */
-	bool bWasInWorldSpaceBeforeScaleMode;
-
 	/** If this current mode is running */
 	bool bIsActive;
+
+	/** Pointer to the current Sequencer */
+	class ISequencer* CurrentSequencer;
+
+	/** The world to meters scale when leaving PIE simulate to restore when back in the editor world. */
+	float SavedWorldToMetersScaleForPIE;
+
+	/** If we started play in editor from the VR Editor*/
+	bool bStartedPlayFromVREditor;	
+	
+	/** Container of assets */
+	TSharedPtr<const struct FVREditorAssetContainer> AssetContainer;
 };
diff --git a/Engine/Source/Editor/VREditor/Teleporter/VREditorAutoScaler.cpp b/Engine/Source/Editor/VREditor/Teleporter/VREditorAutoScaler.cpp
--- a/Engine/Source/Editor/VREditor/Teleporter/VREditorAutoScaler.cpp
+++ b/Engine/Source/Editor/VREditor/Teleporter/VREditorAutoScaler.cpp
@@ -9,6 +9,11 @@
 #include "GameFramework/WorldSettings.h"
 #include "Sound/SoundCue.h"
 
+namespace VREd
+{
+	static FAutoConsoleVariable AllowResetScale(TEXT("VREd.AllowResetScale"), 1, TEXT("Allowed to reset world to meters to default world to meters"));
+}
+
 UVREditorAutoScaler::UVREditorAutoScaler( const FObjectInitializer& ObjectInitializer ) :
 	Super( ObjectInitializer ),
 	VRMode( nullptr )
@@ -35,55 +40,44 @@ void UVREditorAutoScaler::Shutdown()
 
 void UVREditorAutoScaler::Scale( const float NewWorldToMetersScale )
 {
-	UViewportWorldInteraction* WorldInteraction = &VRMode->GetWorldInteraction();
-
-	// Set the new world to meters scale.
-	const float OldWorldToMetersScale = WorldInteraction->GetViewportWorld()->GetWorldSettings()->WorldToMeters;
-	WorldInteraction->SetWorldToMetersScale( NewWorldToMetersScale );
-	WorldInteraction->SkipInteractiveWorldMovementThisFrame();
-
-	FTransform RoomTransform = WorldInteraction->GetRoomTransform();
-
-	// Use the HMD location as pivot to scale of, this will cause the view to be unchanged. 
-	const FVector RoomSpacePivotLocation = WorldInteraction->GetRoomSpaceHeadTransform().GetLocation();
-	const FVector WorldSpacePivotLocation = WorldInteraction->GetHeadTransform().GetLocation();// *( OldWorldToMetersScale / 100.0f );
-		
-	//The rest of this calculations is a duplicate from UViewportWorldInteraction.
-	const FVector NewRoomSpacePivotLocation = ( RoomSpacePivotLocation / OldWorldToMetersScale ) * NewWorldToMetersScale;
-	const FVector NewWorldSpacePivotLocation = RoomTransform.TransformPosition( NewRoomSpacePivotLocation );
-	const FVector WorldSpacePivotDelta = NewWorldSpacePivotLocation - WorldSpacePivotLocation;
-	const FVector NewWorldSpaceRoomLocation = RoomTransform.GetLocation() - WorldSpacePivotDelta;
-	RoomTransform.SetLocation( NewWorldSpaceRoomLocation );
-
-	WorldInteraction->SetRoomTransform( RoomTransform );
+	if (VREd::AllowResetScale->GetInt() != 0)
+	{
+		UViewportWorldInteraction* WorldInteraction = &VRMode->GetWorldInteraction();
 
-	UGameplayStatics::PlaySound2D( VRMode->GetWorld(), ScaleSound );
+		// Set the new world to meters scale.
+		WorldInteraction->SetWorldToMetersScale( NewWorldToMetersScale, true );
+		WorldInteraction->SkipInteractiveWorldMovementThisFrame();
+		UGameplayStatics::PlaySound2D( VRMode->GetWorld(), ScaleSound );
+	}
 }
 
 
 void UVREditorAutoScaler::OnInteractorAction( class FEditorViewportClient& ViewportClient, UViewportInteractor* Interactor, const struct FViewportActionKeyInput& Action, bool& bOutIsInputCaptured, bool& bWasHandled )
 {
-	if( Action.ActionType == VRActionTypes::ConfirmRadialSelection )
+	if (VREd::AllowResetScale->GetInt() != 0)
 	{
-		if( Action.Event == IE_Pressed )
+		if( Action.ActionType == VRActionTypes::ConfirmRadialSelection )
 		{
-			const EViewportInteractionDraggingMode DraggingMode = Interactor->GetDraggingMode();
-			UViewportInteractor* OtherInteractor = Interactor->GetOtherInteractor();
-
-			// Start dragging at laser impact when already dragging actors freely
-			if( DraggingMode == EViewportInteractionDraggingMode::World ||
-				( DraggingMode == EViewportInteractionDraggingMode::AssistingDrag && OtherInteractor != nullptr && OtherInteractor->GetDraggingMode() == EViewportInteractionDraggingMode::World ) )
+			if( Action.Event == IE_Pressed )
 			{
-				const float DefaultWorldToMetersScale = VRMode->GetSavedEditorState().WorldToMetersScale;
-				Scale( DefaultWorldToMetersScale );
-
-				bWasHandled = true;
-				bOutIsInputCaptured = true;
+				const EViewportInteractionDraggingMode DraggingMode = Interactor->GetDraggingMode();
+				UViewportInteractor* OtherInteractor = Interactor->GetOtherInteractor();
+
+				// Start dragging at laser impact when already dragging actors freely
+				if( DraggingMode == EViewportInteractionDraggingMode::World ||
+					( DraggingMode == EViewportInteractionDraggingMode::AssistingDrag && OtherInteractor != nullptr && OtherInteractor->GetDraggingMode() == EViewportInteractionDraggingMode::World ) )
+				{
+					const float DefaultWorldToMetersScale = VRMode->GetSavedEditorState().WorldToMetersScale;
+					Scale( DefaultWorldToMetersScale );
+
+					bWasHandled = true;
+					bOutIsInputCaptured = true;
+				}
+			}
+			else if( Action.Event == IE_Released )
+			{
+				bOutIsInputCaptured = false;
 			}
-		}
-		else if( Action.Event == IE_Released )
-		{
-			bOutIsInputCaptured = false;
 		}
 	}
 }
diff --git a/Engine/Source/Editor/VREditor/Teleporter/VREditorTeleporter.cpp b/Engine/Source/Editor/VREditor/Teleporter/VREditorTeleporter.cpp
--- a/Engine/Source/Editor/VREditor/Teleporter/VREditorTeleporter.cpp
+++ b/Engine/Source/Editor/VREditor/Teleporter/VREditorTeleporter.cpp
@@ -1,157 +1,450 @@
 // Copyright 1998-2017 Epic Games, Inc. All Rights Reserved.
 
 #include "VREditorTeleporter.h"
-#include "HAL/IConsoleManager.h"
-#include "Engine/EngineTypes.h"
-#include "Kismet/GameplayStatics.h"
-#include "Sound/SoundCue.h"
 #include "VREditorMode.h"
 #include "ViewportInteractionTypes.h"
 #include "ViewportWorldInteraction.h"
 #include "VREditorInteractor.h"
+#include "VREditorMotionControllerInteractor.h"
+#include "VREditorAssetContainer.h"
 
+#include "Engine/EngineTypes.h"
+#include "HeadMountedDisplayTypes.h"
+#include "Sound/SoundCue.h"
+#include "Materials/MaterialInstance.h"
+#include "Materials/MaterialInstanceDynamic.h"
+#include "HAL/IConsoleManager.h"
 
 namespace VREd
 {
-	static FAutoConsoleVariable TeleportLerpTime( TEXT( "VREd.TeleportLerpTime" ), 0.1f, TEXT( "The lerp time to teleport" ) );
-	static FAutoConsoleVariable TeleportOffset( TEXT( "VREd.TeleportOffset" ), 100.0f, TEXT( "The offset from the hitresult towards the controller" ) );
-	static FAutoConsoleVariable TeleportLaserPointerLength( TEXT( "VREd.TeleportLaserPointerLength" ), 500000.0f, TEXT( "Distance of the LaserPointer for teleporting" ) );
-	static FAutoConsoleVariable TeleportDistance( TEXT( "VREd.TeleportDistance" ), 500.0f, TEXT( "Default distance for teleporting when not hitting anything" ) );
+	static FAutoConsoleVariable TeleportLerpTime(TEXT("VREd.TeleportLerpTime"), 0.1f, TEXT("The lerp time to teleport"));
+	static FAutoConsoleVariable TeleportOffset(TEXT("VREd.TeleportOffset"), 100.0f, TEXT("The offset from the hitresult towards the controller"));
+	static FAutoConsoleVariable TeleportLaserPointerLength(TEXT("VREd.TeleportLaserPointerLength"), 500000.0f, TEXT("Distance of the LaserPointer for teleporting"));
+	static FAutoConsoleVariable TeleportDistance(TEXT("VREd.TeleportDistance"), 500.0f, TEXT("Default distance for teleporting when not hitting anything"));
+	static FAutoConsoleVariable TeleportScaleSensitivity(TEXT("VREd.TeleportScaleSensitivity"), 0.05f, TEXT("Teleport world to meters scale touchpad sensitivity"));
+	static FAutoConsoleVariable TeleportOffsetMultiplier(TEXT("VREd.TeleportOffsetMultiplier"), 0.3f, TEXT("Teleport offset multiplier"));
+	static FAutoConsoleVariable TeleportEnableChangeScale(TEXT("VREd.TeleportEnableChangeScale"), 0, TEXT("Ability to change the world to meters scale while teleporting"));
+	static FAutoConsoleVariable TeleportFadeInAnimateSpeed(TEXT("VREd.TeleportAnimateSpeed"), 3.0f, TEXT("How fast the teleporter should fade in"));
+	static FAutoConsoleVariable TeleportDragSpeed(TEXT("VREd.TeleportDragSpeed"), 0.3f, TEXT("How fast the teleporter should drag behind the laser aiming location"));
+	static FAutoConsoleVariable TeleportAllowScaleBackToDefault(TEXT("VREd.TeleportAllowScaleBackToDefault"), 1, TEXT("Scale back to default world to meters scale"));
+	static FAutoConsoleVariable TeleportAllowPushPull(TEXT("VREd.TeleportAllowPushPull"), 0, TEXT("Allow being able to push and pull the teleporter along the laser."));
+}
+
+AVREditorTeleporter::AVREditorTeleporter():
+	Super(),
+	VRMode(nullptr),
+	TeleportingState(EState::None),
+	TeleportLerpAlpha(0),
+	TeleportStartLocation(FVector::ZeroVector),
+	TeleportGoalLocation(FVector::ZeroVector),
+	TeleportDirectionMeshComponent(nullptr),
+	HMDMeshComponent(nullptr),
+	LeftMotionControllerMeshComponent(nullptr),
+	RightMotionControllerMeshComponent(nullptr),
+	TeleportMID(nullptr),
+	InteractorTryingTeleport(nullptr),
+	OffsetDistance(FVector::ZeroVector),
+	TeleportGoalScale(0),
+	DragRayLength(0),
+	DragRayLengthVelocity(0),
+	bPushedFromEndOfLaser(false),
+	bInitialTeleportAim(true),
+	FadeAlpha(1.0f),
+	bShouldBeVisible(),
+	TeleportTickDelay(0)
+{
 }
 
-UVREditorTeleporter::UVREditorTeleporter( const FObjectInitializer& Initializer ) :
-	Super( Initializer ),
-	VRMode( nullptr ),
-	bIsTeleporting( false ),
-	TeleportLerpAlpha( 0 ),
-	TeleportStartLocation( FVector::ZeroVector ),
-	TeleportGoalLocation( FVector::ZeroVector ),
-	bJustTeleported( false )
+AVREditorTeleporter::~AVREditorTeleporter()
 {
-	// Load sounds
-	TeleportSound = LoadObject<USoundCue>( nullptr, TEXT( "/Engine/VREditor/Sounds/VR_teleport_Cue" ) );
-	check( TeleportSound != nullptr );
+	TeleportDirectionMeshComponent = nullptr;
+	HMDMeshComponent = nullptr;
+	LeftMotionControllerMeshComponent = nullptr;
+	RightMotionControllerMeshComponent = nullptr;
+	TeleportMID = nullptr;
+	InteractorTryingTeleport = nullptr;
+	VRMode = nullptr;
 }
 
-void UVREditorTeleporter::Init( UVREditorMode* InMode )
+void AVREditorTeleporter::Init(UVREditorMode* InMode)
 {
 	VRMode = InMode;
-	VRMode->OnTickHandle().AddUObject( this, &UVREditorTeleporter::Tick );
-	VRMode->GetWorldInteraction().OnViewportInteractionInputAction().AddUObject( this, &UVREditorTeleporter::OnInteractorAction );
+
+	// Register for event for ticking and input
+	VRMode->OnTickHandle().AddUObject(this, &AVREditorTeleporter::Tick);
+	VRMode->GetWorldInteraction().OnPreviewInputAction().AddUObject(this, &AVREditorTeleporter::OnPreviewInputAction);
+
+	SetActorEnableCollision(false);
+
+	{
+		UMaterialInstance* RoomSpaceMaterialInstance = LoadObject<UMaterialInstance>(nullptr, TEXT("/Engine/VREditor/Teleport/TeleportDirectionMaterial_Inst"));
+		check(RoomSpaceMaterialInstance != nullptr);
+		TeleportMID = UMaterialInstanceDynamic::Create(RoomSpaceMaterialInstance, GetTransientPackage());
+		check(TeleportMID != nullptr);
+	}
+
+	{
+		TeleportDirectionMeshComponent = VRMode->CreateMesh(this, FString("/Engine/VREditor/Teleport/TeleportDirectionMesh"), GetRootComponent());
+		TeleportDirectionMeshComponent->SetWorldScale3D(FVector(1, 1, 1));
+		TeleportDirectionMeshComponent->SetMaterial(0, TeleportMID);
+		TeleportDirectionMeshComponent->SetCollisionEnabled(ECollisionEnabled::NoCollision);
+		TeleportDirectionMeshComponent->SetCollisionResponseToAllChannels(ECollisionResponse::ECR_Ignore);
+		TeleportDirectionMeshComponent->SetCollisionResponseToChannel(COLLISION_GIZMO, ECollisionResponse::ECR_Ignore);
+		TeleportDirectionMeshComponent->SetCastShadow(false);
+
+		HMDMeshComponent = VRMode->CreateMesh(this, FString("/Engine/VREditor/Devices/Generic/GenericHMD"), GetRootComponent());
+		HMDMeshComponent->SetMaterial(0, TeleportMID);
+		HMDMeshComponent->SetCastShadow(false);
+
+		LeftMotionControllerMeshComponent = VRMode->CreateMotionControllerMesh(this, GetRootComponent());
+		check(LeftMotionControllerMeshComponent != nullptr);
+		LeftMotionControllerMeshComponent->SetCastShadow(false);
+
+		RightMotionControllerMeshComponent = VRMode->CreateMotionControllerMesh(this, GetRootComponent());
+		check(RightMotionControllerMeshComponent != nullptr);
+		RightMotionControllerMeshComponent->SetCastShadow(false);
+
+		LeftMotionControllerMeshComponent->SetMaterial(0, TeleportMID);
+		RightMotionControllerMeshComponent->SetMaterial(0, TeleportMID);
+	}
+
+	// Set the initial teleport scale to the current world to meters
+	TeleportGoalScale = VRMode->GetSavedEditorState().WorldToMetersScale;
+
+	Show(false);
 }
 
-void UVREditorTeleporter::Shutdown()
+void AVREditorTeleporter::Shutdown()
 {
-	VRMode->OnTickHandle().RemoveAll( this );
-	VRMode->GetWorldInteraction().OnViewportInteractionInputAction().RemoveAll( this );
-	TeleportSound = nullptr;
+	VRMode->OnTickHandle().RemoveAll(this);
+	VRMode->GetWorldInteraction().OnViewportInteractionInputAction().RemoveAll(this);
 	VRMode = nullptr;
 }
 
-void UVREditorTeleporter::Tick( const float DeltaTime )
+void AVREditorTeleporter::Tick(const float DeltaTime)
 {
-	if ( bIsTeleporting )
+	UpdateFadingState(DeltaTime);
+	if (TeleportingState == EState::Aiming && InteractorTryingTeleport != nullptr)
 	{
-		Teleport( DeltaTime );
+		{
+			UVREditorMotionControllerInteractor* VREditorInteractor = Cast<UVREditorMotionControllerInteractor>(InteractorTryingTeleport);
+			if (VREditorInteractor != nullptr && (VREditorInteractor->GetDraggingMode() == EViewportInteractionDraggingMode::World ||
+				(VREditorInteractor->GetOtherInteractor() != nullptr && VREditorInteractor->GetOtherInteractor()->GetInteractorData().DraggingMode == EViewportInteractionDraggingMode::World &&
+					VREditorInteractor->GetDraggingMode() == EViewportInteractionDraggingMode::AssistingDrag)) &&
+				!FMath::IsNearlyZero(VREditorInteractor->GetSelectAndMoveTriggerValue(), KINDA_SMALL_NUMBER))
+			{
+				VREditorInteractor->SetForceShowLaser(true);
+			}
+		}
+
+		UpdateTeleportAim(DeltaTime); 
+	}
+	else if (TeleportingState == EState::Teleporting)
+	{
+		static const uint32 StaticTeleportDelay = 2;
+		if (TeleportTickDelay > StaticTeleportDelay)
+		{
+			Teleport(DeltaTime);
+		}
+		TeleportTickDelay++;
 	}
 }
 
-void UVREditorTeleporter::StartTeleport( UViewportInteractor* Interactor )
+void AVREditorTeleporter::StartTeleport(UViewportInteractor* Interactor)
 {
-	UVREditorInteractor* VREditorInteractor = Cast<UVREditorInteractor>( Interactor );
-	if ( VREditorInteractor && !VREditorInteractor->IsHoveringOverUI() )
+
+	if (VREd::TeleportEnableChangeScale->GetInt() != 0 || VREd::TeleportAllowScaleBackToDefault->GetInt() != 0)
 	{
-		FVector LaserPointerStart, LaserPointerEnd;
-		if ( Interactor->GetLaserPointer( LaserPointerStart, LaserPointerEnd ) )
+		// Set the world to meters scale @todo vreditor: This causes some strange movement 
+		VRMode->GetWorldInteraction().SetWorldToMetersScale(TeleportGoalScale * 100, true);
+	}
+
+	TeleportingState = EState::Teleporting;
+	TeleportLerpAlpha = 0.0f;
+
+	const FVREditorAssetContainer& AssetContainer = VRMode->GetAssetContainer();
+	VRMode->PlaySound(AssetContainer.TeleportSound, TeleportGoalLocation);
+	Interactor->PlayHapticEffect(1.0f);
+}
+
+void AVREditorTeleporter::OnPreviewInputAction(class FEditorViewportClient& ViewportClient, UViewportInteractor* Interactor, const struct FViewportActionKeyInput& Action, bool& bOutIsInputCaptured, bool& bWasHandled)
+{
+	UVREditorMotionControllerInteractor* VRMotionControllerInteractor = Cast<UVREditorMotionControllerInteractor>(Interactor);
+	if (VRMotionControllerInteractor != nullptr && VRMotionControllerInteractor->GetControllerType() == EControllerType::Laser)
+	{
+		const bool bDraggingWorldForTeleport =
+			Interactor->GetDraggingMode() == EViewportInteractionDraggingMode::World &&
+			(Interactor->GetOtherInteractor() == nullptr || Interactor->GetOtherInteractor()->GetDraggingMode() != EViewportInteractionDraggingMode::AssistingDrag);
+		if (bDraggingWorldForTeleport && 
+			Action.Event == IE_Pressed && 
+			Action.ActionType == ViewportWorldActionTypes::SelectAndMove)
 		{
-			FVector EndLocation;
-			FHitResult HitResult = Interactor->GetHitResultFromLaserPointer( nullptr, true, nullptr, false, VREd::TeleportLaserPointerLength->GetFloat() );
-			if( HitResult.bBlockingHit )
-			{
-				EndLocation = HitResult.Location;
-			}
-			else
+			if (InteractorTryingTeleport == nullptr && TeleportingState == EState::None)
 			{
-				FVector Direction = LaserPointerEnd - LaserPointerStart;
-				Direction.Normalize();
-				EndLocation = LaserPointerStart + ( Direction * ( VREd::TeleportDistance->GetFloat() * VRMode->GetWorldScaleFactor() ) );
+				TeleportGoalScale = VRMode->GetWorldScaleFactor();
+
+				// Checking teleport
+				SetVisibility(true);
+				InteractorTryingTeleport = Interactor;
+
+				TeleportingState = EState::Aiming;
+				SetColor(VRMode->GetColor(UVREditorMode::EColors::WorldDraggingColor));
+				VRMode->GetWorldInteraction().AllowWorldMovement(false);
 			}
 
-			const FVector HeadLocation = VRMode->GetHeadTransform().GetLocation();
-			const FVector RoomLocation = VRMode->GetRoomTransform().GetLocation();
+			bWasHandled = true;
+		}
 
-			FVector NewHandLocation = EndLocation;
-			if( HitResult.bBlockingHit )
+		if (Action.Event == IE_Released && 
+			Action.ActionType == ViewportWorldActionTypes::SelectAndMove &&
+			InteractorTryingTeleport != nullptr && InteractorTryingTeleport == Interactor)
+		{
+			UVREditorMotionControllerInteractor* VREditorInteractor = Cast<UVREditorMotionControllerInteractor>(InteractorTryingTeleport);
+			if (VREditorInteractor && VREditorInteractor == InteractorTryingTeleport && !(VREditorInteractor->IsHoveringOverGizmo() || VREditorInteractor->IsHoveringOverUI()))
 			{
-				const FVector Offset = ( LaserPointerStart - EndLocation ).GetSafeNormal() *
-					VREd::TeleportOffset->GetFloat() * VRMode->GetWorldInteraction().GetWorldScaleFactor(); // Offset from the hitresult, TeleportOffset is in centimeter
-				NewHandLocation = EndLocation + Offset;
-			}
+				// Confirm teleport
+				StartTeleport(Interactor);
+
+				// Clean everything up
+				bPushedFromEndOfLaser = false;
+				bInitialTeleportAim = true;
+				bShouldBeVisible = false;
+				SetVisibility(false);
 
-			const FVector HandHeadOffset = HeadLocation - Interactor->GetTransform().GetLocation();		// The offset between the hand and the head
-			const FVector NewHeadLocation = NewHandLocation + HandHeadOffset;							// The location in world space where the head is going to be
-																										// The new roomspace location in world space
-			const FVector HeadRoomOffset = HeadLocation - RoomLocation;									// Where the head is going to be in world space
-			const FVector NewRoomSpaceLocation = NewHeadLocation - HeadRoomOffset;						// Result, where the roomspace is going to be in world space
-																										// Set values to start teleporting
-			bIsTeleporting = true;
-			TeleportLerpAlpha = 0.0f;
-			TeleportStartLocation = RoomLocation;
-			TeleportGoalLocation = NewRoomSpaceLocation;
-
-			UGameplayStatics::PlaySound2D( VRMode->GetWorld(), TeleportSound );
+				VREditorInteractor->SetForceShowLaser(false);
+				InteractorTryingTeleport = nullptr;
+
+				bWasHandled = true;
+			}
 		}
 	}
 }
 
-void UVREditorTeleporter::OnInteractorAction( class FEditorViewportClient& ViewportClient, UViewportInteractor* Interactor, const struct FViewportActionKeyInput& Action, bool& bOutIsInputCaptured, bool& bWasHandled )
+void AVREditorTeleporter::Teleport(const float DeltaTime)
 {
-	if ( ( Action.ActionType == ViewportWorldActionTypes::SelectAndMove ||
-		Action.ActionType == ViewportWorldActionTypes::SelectAndMove_LightlyPressed ) )
+	if (TeleportLerpAlpha == 0.0f)
+	{
+		TeleportStartLocation = VRMode->GetRoomTransform().GetLocation();
+	}
+	
+	FTransform RoomTransform = VRMode->GetRoomTransform();
+	TeleportLerpAlpha += DeltaTime;
+	if (TeleportLerpAlpha > VREd::TeleportLerpTime->GetFloat())
 	{
-		if ( Action.Event == IE_Pressed )
+		// Teleporting is finished
+		TeleportLerpAlpha = VREd::TeleportLerpTime->GetFloat();
+		TeleportingState = EState::None;
+		VRMode->GetWorldInteraction().AllowWorldMovement(true);
+		TeleportTickDelay = 0;
+	}
+
+	// Calculate the new position of the roomspace
+	FVector NewLocation = FMath::Lerp(TeleportStartLocation, TeleportGoalLocation, TeleportLerpAlpha / VREd::TeleportLerpTime->GetFloat());
+	RoomTransform.SetLocation(NewLocation);
+	VRMode->SetRoomTransform(RoomTransform);
+}
+
+void AVREditorTeleporter::UpdateTeleportAim(const float DeltaTime)
+{
+	UVREditorMotionControllerInteractor* VREditorInteractor = Cast<UVREditorMotionControllerInteractor>(InteractorTryingTeleport);
+	if (VREditorInteractor)
+	{
+		Show(!(VREditorInteractor->IsHoveringOverGizmo() || VREditorInteractor->IsHoveringOverUI()));
+
+		FVector LaserPointerStart, LaserPointerEnd, EndLocation;
+		if (VREditorInteractor->GetLaserPointer(LaserPointerStart, LaserPointerEnd) && TeleportingState == EState::Aiming)
 		{
-			if ( Action.ActionType == ViewportWorldActionTypes::SelectAndMove_LightlyPressed )
+			const bool bAllowPushPull = VREd::TeleportAllowPushPull->GetInt() == 0 ? false : true;
+
+			if (VREd::TeleportEnableChangeScale->GetInt() != 0)
 			{
-				// Check to teleport before checking on actor selection, otherwise objects will be selected when teleporting
-				if	( Interactor->GetDraggingMode() == EViewportInteractionDraggingMode::World &&
-					( Interactor->GetOtherInteractor() == nullptr || Interactor->GetOtherInteractor()->GetDraggingMode() != EViewportInteractionDraggingMode::AssistingDrag ) )
+				// Calaclate the new goal scale with the trackpad delta X axis 
+				TeleportGoalScale += VREditorInteractor->GetTrackpadSlideDelta(0) * (TeleportGoalScale * VREd::TeleportScaleSensitivity->GetFloat());
+
+				const float MinScale = VRMode->GetWorldInteraction().GetMinScale() * 0.01f;
+				const float MaxScale = VRMode->GetWorldInteraction().GetMaxScale() * 0.01f;
+
+				// The scale must be in the limits of the worldinteraction minimum and maximum scale
+				if (TeleportGoalScale > MaxScale)
+				{
+					TeleportGoalScale = MaxScale;
+				}
+				else if (TeleportGoalScale < MinScale)
 				{
-					StartTeleport( Interactor );
-					bJustTeleported = true;
-					bWasHandled = true;
+					TeleportGoalScale = MinScale;
 				}
 			}
-			else if ( Action.ActionType == ViewportWorldActionTypes::SelectAndMove && bJustTeleported )
+			else if (VREd::TeleportAllowScaleBackToDefault->GetInt() != 0)
 			{
-				//Block the input
-				bJustTeleported = false;
-				bWasHandled = true;
+				TeleportGoalScale = VRMode->GetSavedEditorState().WorldToMetersScale / 100.f;
+			}
+			else
+			{
+				TeleportGoalScale = VRMode->GetWorldScaleFactor();
 			}
 
-			// Nothing should happen when the hand has a UI in front, otherwise it will deselect everything
-			if ( Interactor->GetIsLaserBlocked() )
+			// If the laser is hitting something the teleport will go their with an appropriate offset. 
+			FHitResult HitResult = VREditorInteractor->GetHitResultFromLaserPointer(nullptr, true, nullptr, false, VREd::TeleportLaserPointerLength->GetFloat());
+			if (HitResult.bBlockingHit && !bPushedFromEndOfLaser)
 			{
-				bWasHandled = true;
+				// Calculate an offset with the impact normal, so the teleporter will show up on-top, underneath or next to where laser is aiming at.
+				OffsetDistance = HitResult.ImpactNormal * (50 * TeleportGoalScale);
+				float OffsetFromImpactNormalZ = (HitResult.ImpactNormal.Z - 1);
+				OffsetFromImpactNormalZ -= (OffsetFromImpactNormalZ * VREd::TeleportOffsetMultiplier->GetFloat());
+				OffsetDistance.Z = OffsetFromImpactNormalZ * (((VRMode->GetHeadTransform().GetLocation().Z - VRMode->GetRoomTransform().GetLocation().Z) / VRMode->GetWorldScaleFactor())  * TeleportGoalScale);
+
+				// Set the final location based on the hit location and the offset
+				EndLocation = HitResult.Location + OffsetDistance;
+
+				// Update the raylength to the current length, so if we have to pull or push the teleporter next frame it will not jump
+				DragRayLength = FVector::Dist(LaserPointerStart, HitResult.ImpactPoint);
+			}
+			// If the laser is not hitting anything or has been pushed away already, the user can push or pull the teleporter along the laser
+			else
+			{
+				EndLocation = UpdatePushPullTeleporter(VREditorInteractor, LaserPointerStart, LaserPointerEnd, bAllowPushPull);
+			}
+
+			// The trackpad has been used while aiming for teleporting, so the teleporter won't go to the end of the laser after this
+			if (!FMath::IsNearlyZero(VREditorInteractor->GetTrackpadSlideDelta(1)) && !bPushedFromEndOfLaser && bAllowPushPull)
+			{
+				bPushedFromEndOfLaser = true;
+			} 
+				
+			// Smooth the final location
+			if (!bInitialTeleportAim)
+			{
+				FVector TeleporterAndAimLocationOffset = EndLocation - GetActorLocation();
+				if (TeleporterAndAimLocationOffset.Size() > (VRMode->GetWorldScaleFactor() * 1.0f)) //@todo vreditor: Tweak this value for less jitter
+				{
+					EndLocation = GetActorLocation() + (TeleporterAndAimLocationOffset * VREd::TeleportDragSpeed->GetFloat());
+				}
 			}
+
+			UpdateVisuals(EndLocation);
+			bInitialTeleportAim = false;
 		}
 	}
 }
 
-void UVREditorTeleporter::Teleport( const float DeltaTime )
+FVector AVREditorTeleporter::UpdatePushPullTeleporter(UVREditorMotionControllerInteractor* VREditorInteractor, const FVector& LaserPointerStart, const FVector& LaserPointerEnd, const bool bEnablePushPull /* = true */)
 {
-	FTransform RoomTransform = VRMode->GetRoomTransform();
-	TeleportLerpAlpha += DeltaTime;
+	if (bEnablePushPull)
+	{
+		VREditorInteractor->CalculateDragRay(DragRayLength, DragRayLengthVelocity);
+	}
+
+	FVector Direction = LaserPointerEnd - LaserPointerStart;
+	Direction.Normalize();
+
+	return (LaserPointerStart + (Direction * DragRayLength)) + OffsetDistance;
+}
+
+void AVREditorTeleporter::SetVisibility(const bool bVisible)
+{
+	TeleportDirectionMeshComponent->SetVisibility(bVisible);
+	HMDMeshComponent->SetVisibility(bVisible);
+	LeftMotionControllerMeshComponent->SetVisibility(bVisible);
+	RightMotionControllerMeshComponent->SetVisibility(bVisible);
+}
+
+void AVREditorTeleporter::SetColor(const FLinearColor& Color)
+{
+	static FName StaticRoomspaceMeshColorName("Color");
+	TeleportMID->SetVectorParameterValue(StaticRoomspaceMeshColorName, Color);
+}
+
+void AVREditorTeleporter::UpdateVisuals(const FVector& NewLocation)
+{
+	const float WorldScale = VRMode->GetWorldScaleFactor();
+	const float CalculatedScale = CalculateAnimatedScaleFactor();
+	const FVector AnimatedScale = FVector(TeleportGoalScale * CalculatedScale);
+
+	SetActorLocation(NewLocation);
+	FTransform HMDTransform = VRMode->GetHeadTransform();
 
-	if ( TeleportLerpAlpha > VREd::TeleportLerpTime->GetFloat() )
+	// Update the teleport direction indicator
 	{
-		// Teleporting is finished
-		TeleportLerpAlpha = VREd::TeleportLerpTime->GetFloat();
-		bIsTeleporting = false;
+		const FRotator TeleportDirection = FRotator(0.f, HMDTransform.GetRotation().Rotator().Yaw, 0.f);
+		TeleportDirectionMeshComponent->SetWorldRotation(TeleportDirection);
+		TeleportDirectionMeshComponent->SetWorldScale3D(AnimatedScale);
+		TeleportDirectionMeshComponent->SetRelativeLocation(FVector(0.0f, 0.0f, AnimatedScale.Z * 0.5f));
 	}
 
-	// Calculate the new position of the roomspace
-	FVector NewLocation = FMath::Lerp( TeleportStartLocation, TeleportGoalLocation, TeleportLerpAlpha / VREd::TeleportLerpTime->GetFloat() );
-	RoomTransform.SetLocation( NewLocation );
-	VRMode->SetRoomTransform( RoomTransform );
+	HMDTransform.SetLocation(FVector(NewLocation.X, NewLocation.Y, NewLocation.Z + ((VRMode->GetRoomSpaceHeadTransform().GetLocation().Z / WorldScale) * TeleportGoalScale)));
+	HMDTransform.SetScale3D(AnimatedScale);
+	HMDMeshComponent->SetWorldTransform(HMDTransform);
+
+	//Calculate the teleported room transform
+	FTransform HeadToWorld = VRMode->GetHeadTransform();
+	HeadToWorld.SetRotation(FQuat::Identity);
+	const FTransform RoomToWorld = VRMode->GetRoomTransform();
+	FTransform RoomToHeadInWorld = RoomToWorld.GetRelativeTransform(HeadToWorld);
+	RoomToHeadInWorld.SetLocation(FVector((RoomToHeadInWorld.GetLocation() / WorldScale) * TeleportGoalScale));
+
+	HMDTransform.SetRotation(FQuat::Identity);
+	FTransform TeleportRoomInWorld = HMDTransform + RoomToHeadInWorld;
+	TeleportGoalLocation = TeleportRoomInWorld.GetLocation();
+
+	// Calculate the teleported motion controllers
+	for (int i = 0; i < 2; i++)
+	{
+		UVREditorInteractor* Interactor = i == 0 ? VRMode->GetHandInteractor(EControllerHand::Left) : VRMode->GetHandInteractor(EControllerHand::Right);
+		UStaticMeshComponent* MotionControllerMeshComponent = i == 0 ? LeftMotionControllerMeshComponent : RightMotionControllerMeshComponent;
+
+		FTransform MotionControllerToWorld = Interactor->GetTransform();
+		FTransform MotionControllerToHeadInWorld = MotionControllerToWorld.GetRelativeTransform(HeadToWorld);
+		MotionControllerToHeadInWorld.SetLocation(FVector((MotionControllerToHeadInWorld.GetLocation() / WorldScale) * TeleportGoalScale));
+		FTransform TeleportedMotionControllerToWorld = HMDTransform + MotionControllerToHeadInWorld;
+		TeleportedMotionControllerToWorld.SetRotation(MotionControllerToWorld.GetRotation());
+		TeleportedMotionControllerToWorld.SetScale3D(AnimatedScale);
+
+		MotionControllerMeshComponent->SetWorldTransform(TeleportedMotionControllerToWorld);
+	}	
+}
+
+void AVREditorTeleporter::Show(const bool bShow)
+{
+	// Only start the fading in or out if it is the first time requesting it. Only reset animation if it wasn't started yet.
+	if (bShow)
+	{
+		if (!bShouldBeVisible.IsSet() || bShouldBeVisible.GetValue() != bShow)
+		{
+			bShouldBeVisible = true;
+			FadeAlpha = 0.0f;
+		}
+	}
+	else
+	{
+		bShouldBeVisible = false;
+		SetVisibility(false);
+	}
+}
+
+void AVREditorTeleporter::UpdateFadingState(const float DeltaTime)
+{
+	if (bShouldBeVisible.IsSet() && bShouldBeVisible.GetValue())
+	{
+		FadeAlpha += VREd::TeleportFadeInAnimateSpeed->GetFloat() * DeltaTime;
+		FadeAlpha = FMath::Clamp(FadeAlpha, 0.0f, 1.0f);
+
+		if (FadeAlpha > 0.0f + KINDA_SMALL_NUMBER)
+		{
+			// At least a little bit visible
+			if (!TeleportDirectionMeshComponent->IsVisible())
+			{
+				SetVisibility(true);
+			}
+		}
+	}
+}
+
+float AVREditorTeleporter::CalculateAnimatedScaleFactor() const
+{
+	const float AnimationOvershootAmount = 0.7f;	// @todo vreditor tweak
+	const float EasedAlpha = UVREditorMode::OvershootEaseOut(FadeAlpha, AnimationOvershootAmount);
+
+	// Animate vertically more than horizontally; just looks a little better
+	const float Scale = FMath::Max(0.001f, EasedAlpha);
+	return Scale * Scale;
 }
diff --git a/Engine/Source/Editor/VREditor/Teleporter/VREditorTeleporter.h b/Engine/Source/Editor/VREditor/Teleporter/VREditorTeleporter.h
--- a/Engine/Source/Editor/VREditor/Teleporter/VREditorTeleporter.h
+++ b/Engine/Source/Editor/VREditor/Teleporter/VREditorTeleporter.h
@@ -11,51 +11,85 @@ class UViewportInteractor;
 class UVREditorMode;
 
 /**
-* VR Editor teleport manager
-*/
+ * VR Editor teleport manager and the visual representation of the teleport
+ */
 UCLASS()
-class UVREditorTeleporter : public UObject
+class AVREditorTeleporter: public AActor
 {
-	GENERATED_UCLASS_BODY()
+	GENERATED_BODY()
 
 public:
 
-	/** Initializes the teleport system */
-	void Init( class UVREditorMode* InMode );
+	/** Default constructor that sets up CDO properties */
+	AVREditorTeleporter();
 
-	/** Shuts down the teleport system */
+	/** Deconstructor */
+	virtual ~AVREditorTeleporter();
+
+	/** Initializes the teleporter */
+	void Init(class UVREditorMode* InMode);
+
+	/** Shuts down the teleporter */
 	void Shutdown();
-	
-	/** Returns true if we're actively teleporting */
-	bool IsTeleporting() const
+
+private:
+
+	//~ Begin AActor interface
+	virtual void Tick(const float DeltaTime) override;
+	virtual bool IsEditorOnly() const final
 	{
-		return bIsTeleporting;
+		return true;
 	}
+	//~ End AActor interface
 
-private:
-	/** Ticks the current teleport */
-	void Tick( const float DeltaTime );
-	
 	/** Start teleporting, does a ray trace with the hand passed and calculates the locations for lerp movement in Teleport */
-	void StartTeleport( class UViewportInteractor* Interactor );
+	void StartTeleport(class UViewportInteractor* Interactor);
 
 	/** Called when the user presses a button on their motion controller device */
-	void OnInteractorAction( class FEditorViewportClient& ViewportClient, UViewportInteractor* Interactor,
-		const struct FViewportActionKeyInput& Action, bool& bOutIsInputCaptured, bool& bWasHandled );
+	void OnPreviewInputAction(class FEditorViewportClient& ViewportClient, UViewportInteractor* Interactor,
+		const struct FViewportActionKeyInput& Action, bool& bOutIsInputCaptured, bool& bWasHandled);
+
+	/** Move the roomspace using lerp towards the new location */
+	void Teleport(const float DeltaTime);
 
-	/** Actually teleport to a position */
-	void Teleport( const float DeltaTime );
+	/** Updating aiming with teleport to end of laser and pulling and pushing */
+	void UpdateTeleportAim(const float DeltaTime);
 
-	//
-	// Fields
-	//
+	/** Helper function to push and pull the teleporter along the laser */
+	FVector UpdatePushPullTeleporter(class UVREditorMotionControllerInteractor* VREditorInteractor, const FVector& LaserPointerStart, const FVector& LaserPointerEnd, const bool bEnablePushPull = true);
 
-	/** Owning mode */
+	/** Hide or show the teleporter visuals */
+	void SetVisibility(const bool bVisible);
+
+	/** Sets the color for the teleporter visuals */
+	void SetColor(const FLinearColor& Color);
+
+	/** Figures out the new transforms for all the visuals based on the new location and the transforms of the HMD and motioncontrollers */
+	void UpdateVisuals(const FVector& NewLocation);
+
+	/** If we want to start showing or hiding the meshes */
+	void Show(const bool bShow);
+
+	void UpdateFadingState(const float DeltaTime);
+
+	/** Calculated the scale for animation */
+	float CalculateAnimatedScaleFactor() const;
+
+	void ResetTeleporter();
+	
+	/** The owning VR mode */
 	UPROPERTY()
 	UVREditorMode* VRMode;
 
-	/** If we are currently teleporting */
-	bool bIsTeleporting;
+	enum EState
+	{
+		None = 0,				// When not aiming for a teleport
+		Aiming = 1,	// Aiming at the end of the laser, user can also scale using touchpad/analog stick
+		Teleporting = 3			// Currently teleporting from one location to another with lerp
+	};
+
+	/** The current teleport state */
+	EState TeleportingState;
 
 	/** The current lerp of the teleport between the TeleportStartLocation and the TeleportGoalLocation */
 	float TeleportLerpAlpha;
@@ -66,11 +100,54 @@ private:
 	/** The calculated goal location in StartTeleport to move the Roomspace to */
 	FVector TeleportGoalLocation;
 
-	/** So we can teleport on a light press and make sure a second teleport by a fullpress doesn't happen */
-	bool bJustTeleported;
+	/** Visuals for the feet location of the teleporter with the same direction of the HMD yaw */
+	UPROPERTY()
+	UStaticMeshComponent* TeleportDirectionMeshComponent;
+
+	/** Visuals for teleport HMD */
+	UPROPERTY()
+	UStaticMeshComponent* HMDMeshComponent;
+
+	/** Visuals for teleport left motion controller */
+	UPROPERTY()
+	UStaticMeshComponent* LeftMotionControllerMeshComponent;
+
+	/** Visuals for teleport right motion controller */
+	UPROPERTY()
+	UStaticMeshComponent* RightMotionControllerMeshComponent;
+	
+	/** Dynamic material for teleport visuals */
+	UPROPERTY()
+	UMaterialInstanceDynamic* TeleportMID;
 
-	/** Teleport sound */
+	/** The interactor that started aiming to teleport */
 	UPROPERTY()
-	class USoundCue* TeleportSound;
+	UViewportInteractor* InteractorTryingTeleport;
+
+	/** When offset between the hoverlocation of the laser and the calculated  teleport*/
+	FVector OffsetDistance;
+
+	/** The goal world to meters scale. This value is used to scale the visuals and will be used to set world to meters scale after teleporting */
+	float TeleportGoalScale;
+
+	/** The current length of the laser where the teleport should be at */
+	float DragRayLength;
+	
+	/** The current drag velocity to push or pull the teleport along the laser */
+	float DragRayLengthVelocity;
+
+	/** If the teleporter has been pushed by the trackpad */
+	bool bPushedFromEndOfLaser;
+
+	/** Is this the first time aiming for teleporting, we don't have to smooth movement for initial tick */
+	bool bInitialTeleportAim;
+
+	/** Fade alpha, for visibility transitions */
+	float FadeAlpha;
+
+	/** If the teleporter should fade in */
+	TOptional<bool> bShouldBeVisible;
 
-};
+	/** Delay to start the actual moving to the end location */
+	uint32 TeleportTickDelay;
+};
\ No newline at end of file
