Merged all changes from Private-Geometry up to CL 3339870
- This includes mostly VR Editor usability improvements

diff --git a/Engine/Config/BaseEditorPerProjectUserSettings.ini b/Engine/Config/BaseEditorPerProjectUserSettings.ini
--- a/Engine/Config/BaseEditorPerProjectUserSettings.ini
+++ b/Engine/Config/BaseEditorPerProjectUserSettings.ini
@@ -153,8 +153,6 @@ bBreakOnExceptions=False
 bToolbarCustomization=False
 bBehaviorTreeEditor=False
 bEnableFindAndReplaceReferences =False
-; If true, wearing a Vive or Oculus Rift headset will automatically enter VR Editing mode 
-bEnableAutoVREditMode=false
 
 [/Script/UnrealEd.EditorLoadingSavingSettings]
 ; True if we should automatically load a default level at start up
diff --git a/Engine/Source/Editor/EditorSettingsViewer/EditorSettingsViewer.Build.cs b/Engine/Source/Editor/EditorSettingsViewer/EditorSettingsViewer.Build.cs
--- a/Engine/Source/Editor/EditorSettingsViewer/EditorSettingsViewer.Build.cs
+++ b/Engine/Source/Editor/EditorSettingsViewer/EditorSettingsViewer.Build.cs
@@ -26,7 +26,8 @@ namespace UnrealBuildTool.Rules
                     "InternationalizationSettings",
 					"BlueprintGraph",
                     "EditorStyle",
-                    "Analytics"
+                    "Analytics",
+                    "VREditor"
 				}
 			);
 
diff --git a/Engine/Source/Editor/EditorSettingsViewer/Private/EditorSettingsViewerModule.cpp b/Engine/Source/Editor/EditorSettingsViewer/Private/EditorSettingsViewerModule.cpp
--- a/Engine/Source/Editor/EditorSettingsViewer/Private/EditorSettingsViewerModule.cpp
+++ b/Engine/Source/Editor/EditorSettingsViewer/Private/EditorSettingsViewerModule.cpp
@@ -41,6 +41,7 @@
 
 #include "CrashReporterSettings.h"
 #include "Analytics/AnalyticsPrivacySettings.h"
+#include "VRModeSettings.h"
 
 #define LOCTEXT_NAMESPACE "FEditorSettingsViewerModule"
 
@@ -169,6 +170,13 @@ protected:
 			LOCTEXT("ExperimentalSettingsDescription", "Enable and configure experimental Editor features."),
 			GetMutableDefault<UEditorExperimentalSettings>()
 		);
+
+		// experimental features
+		SettingsModule.RegisterSettings("Editor", "General", "VR Mode",
+			LOCTEXT("VRModeSettings", "VR Mode"),
+			LOCTEXT("VRModeSettingsDescription", "Configure VR editing features."),
+			GetMutableDefault<UVRModeSettings>()
+			);
 	}
 
 	/**
diff --git a/Engine/Source/Editor/LevelEditor/Private/LevelEditorToolBar.cpp b/Engine/Source/Editor/LevelEditor/Private/LevelEditorToolBar.cpp
--- a/Engine/Source/Editor/LevelEditor/Private/LevelEditorToolBar.cpp
+++ b/Engine/Source/Editor/LevelEditor/Private/LevelEditorToolBar.cpp
@@ -1269,7 +1269,7 @@ TSharedRef< SWidget > FLevelEditorToolBar::MakeLevelEditorToolBar( const TShared
 			FSlateIcon(FEditorStyle::GetStyleSetName(), "LevelEditor.EditMatinee") 
 			);
 
-		ToolbarBuilder.AddToolBarButton( FLevelEditorCommands::Get().ToggleVR, NAME_None, LOCTEXT("ToggleVR", "VR Editor") );
+		ToolbarBuilder.AddToolBarButton( FLevelEditorCommands::Get().ToggleVR, NAME_None, LOCTEXT("ToggleVR", "VR Mode") );
 	}
 	ToolbarBuilder.EndSection();
 	
diff --git a/Engine/Source/Editor/MeshEditor/MeshEditorMode.cpp b/Engine/Source/Editor/MeshEditor/MeshEditorMode.cpp
--- a/Engine/Source/Editor/MeshEditor/MeshEditorMode.cpp
+++ b/Engine/Source/Editor/MeshEditor/MeshEditorMode.cpp
@@ -6202,7 +6202,7 @@ void FMeshEditorMode::MeshEditActionsGenerator(FMenuBuilder MenuBuilder, TShared
 			FExecuteAction::CreateSP(this, &FMeshEditorMode::AddOrRemoveSubdivisionLevel, true )
 		),
 		NAME_None,
-		EUserInterfaceActionType::Check
+		EUserInterfaceActionType::ToggleButton
 		);
 	MenuBuilder.AddMenuEntry(
 		LOCTEXT("RemoveSubdivision", "Remove SubD"),
@@ -6213,7 +6213,7 @@ void FMeshEditorMode::MeshEditActionsGenerator(FMenuBuilder MenuBuilder, TShared
 			FExecuteAction::CreateSP(this, &FMeshEditorMode::AddOrRemoveSubdivisionLevel, false)
 		),
 		NAME_None,
-		EUserInterfaceActionType::Check
+		EUserInterfaceActionType::ToggleButton
 		);
 	MenuBuilder.AddMenuEntry(
 		LOCTEXT("EditInstance", "Edit Instance"),
@@ -6226,7 +6226,7 @@ void FMeshEditorMode::MeshEditActionsGenerator(FMenuBuilder MenuBuilder, TShared
 			FIsActionChecked::CreateLambda([this] { return IsEditingPerInstance(); })
 		),
 		NAME_None,
-		EUserInterfaceActionType::Check
+		EUserInterfaceActionType::ToggleButton
 		);
 
 	if (GetMeshElementSelectionMode() == EEditableMeshElementType::Polygon)
@@ -6242,7 +6242,7 @@ void FMeshEditorMode::MeshEditActionsGenerator(FMenuBuilder MenuBuilder, TShared
 				FIsActionChecked::CreateLambda([this] { return (EquippedPolygonAction == EMeshEditAction::Move); })
 				),
 			NAME_None,
-			EUserInterfaceActionType::Check
+			EUserInterfaceActionType::ToggleButton
 			);
 		MenuBuilder.AddMenuEntry(
 			LOCTEXT("Delete", "Delete"),
@@ -6254,7 +6254,7 @@ void FMeshEditorMode::MeshEditActionsGenerator(FMenuBuilder MenuBuilder, TShared
 				FCanExecuteAction::CreateSP(this, &FMeshEditorMode::IsMeshElementTypeSelected, EEditableMeshElementType::Polygon)
 				),
 			NAME_None,
-			EUserInterfaceActionType::Button
+			EUserInterfaceActionType::CollapsedButton
 			);
 		MenuBuilder.AddMenuEntry(
 			LOCTEXT("Extrude", "Extrude"),
@@ -6267,7 +6267,7 @@ void FMeshEditorMode::MeshEditActionsGenerator(FMenuBuilder MenuBuilder, TShared
 				FIsActionChecked::CreateLambda([this] { return (EquippedPolygonAction == EMeshEditAction::ExtrudePolygon); })
 			),
 			NAME_None,
-			EUserInterfaceActionType::Check
+			EUserInterfaceActionType::ToggleButton
 			);
 		MenuBuilder.AddMenuEntry(
 			LOCTEXT("Inset", "Inset"),
@@ -6280,7 +6280,7 @@ void FMeshEditorMode::MeshEditActionsGenerator(FMenuBuilder MenuBuilder, TShared
 				FIsActionChecked::CreateLambda([this] { return (EquippedPolygonAction == EMeshEditAction::InsetPolygon); })
 			),
 			NAME_None,
-			EUserInterfaceActionType::Check
+			EUserInterfaceActionType::ToggleButton
 			);
 	}
 	else if (GetMeshElementSelectionMode() == EEditableMeshElementType::Edge)
@@ -6296,7 +6296,7 @@ void FMeshEditorMode::MeshEditActionsGenerator(FMenuBuilder MenuBuilder, TShared
 				FIsActionChecked::CreateLambda([this] { return (EquippedEdgeAction == EMeshEditAction::Move); })
 			),
 			NAME_None,
-			EUserInterfaceActionType::Check
+			EUserInterfaceActionType::ToggleButton
 			);
 		MenuBuilder.AddMenuEntry(
 			LOCTEXT("Delete", "Delete"),
@@ -6308,7 +6308,7 @@ void FMeshEditorMode::MeshEditActionsGenerator(FMenuBuilder MenuBuilder, TShared
 				FCanExecuteAction::CreateSP(this, &FMeshEditorMode::IsMeshElementTypeSelected, EEditableMeshElementType::Edge)
 				),
 			NAME_None,
-			EUserInterfaceActionType::Button
+			EUserInterfaceActionType::CollapsedButton
 			);
 		MenuBuilder.AddMenuEntry(
 			LOCTEXT("RemoveSelected", "Remove Selected"),
@@ -6320,7 +6320,7 @@ void FMeshEditorMode::MeshEditActionsGenerator(FMenuBuilder MenuBuilder, TShared
 				FCanExecuteAction::CreateSP(this, &FMeshEditorMode::IsMeshElementTypeSelected, EEditableMeshElementType::Edge)
 				),
 			NAME_None,
-			EUserInterfaceActionType::Button
+			EUserInterfaceActionType::CollapsedButton
 			);
 		MenuBuilder.AddMenuEntry(
 			LOCTEXT( "SelectEdgeLoop", "Select Edge Loop" ),
@@ -6332,7 +6332,7 @@ void FMeshEditorMode::MeshEditActionsGenerator(FMenuBuilder MenuBuilder, TShared
 				FCanExecuteAction::CreateSP(this, &FMeshEditorMode::IsMeshElementTypeSelected, EEditableMeshElementType::Edge)
 				),
 			NAME_None,
-			EUserInterfaceActionType::Button
+			EUserInterfaceActionType::CollapsedButton
 			);
 		MenuBuilder.AddMenuEntry(
 			LOCTEXT("InsertEdgeLoop", "Insert Loop"),
@@ -6345,7 +6345,7 @@ void FMeshEditorMode::MeshEditActionsGenerator(FMenuBuilder MenuBuilder, TShared
 				FIsActionChecked::CreateLambda([this] { return (EquippedEdgeAction == EMeshEditAction::InsertEdgeLoop); })
 				),
 			NAME_None,
-			EUserInterfaceActionType::Check
+			EUserInterfaceActionType::ToggleButton
 			);
 		MenuBuilder.AddMenuEntry(
 			LOCTEXT("ExtendEdge", "Extend"),
@@ -6358,7 +6358,7 @@ void FMeshEditorMode::MeshEditActionsGenerator(FMenuBuilder MenuBuilder, TShared
 				FIsActionChecked::CreateLambda([this] { return (EquippedEdgeAction == EMeshEditAction::ExtendEdge); })
 				),
 			NAME_None,
-			EUserInterfaceActionType::Check
+			EUserInterfaceActionType::ToggleButton
 			);
 	}
 	else if (GetMeshElementSelectionMode() == EEditableMeshElementType::Vertex)
@@ -6374,7 +6374,7 @@ void FMeshEditorMode::MeshEditActionsGenerator(FMenuBuilder MenuBuilder, TShared
 				FIsActionChecked::CreateLambda([this] { return (EquippedVertexAction == EMeshEditAction::Move); })
 				),
 			NAME_None,
-			EUserInterfaceActionType::Check
+			EUserInterfaceActionType::ToggleButton
 			);
 		MenuBuilder.AddMenuEntry(
 			LOCTEXT("Delete", "Delete"),
@@ -6386,7 +6386,7 @@ void FMeshEditorMode::MeshEditActionsGenerator(FMenuBuilder MenuBuilder, TShared
 				FCanExecuteAction::CreateSP(this, &FMeshEditorMode::IsMeshElementTypeSelected, EEditableMeshElementType::Vertex)
 			),
 			NAME_None,
-			EUserInterfaceActionType::Button
+			EUserInterfaceActionType::CollapsedButton
 			);
 		MenuBuilder.AddMenuEntry(
 			LOCTEXT("RemoveSelected", "Remove Selected"),
@@ -6398,7 +6398,7 @@ void FMeshEditorMode::MeshEditActionsGenerator(FMenuBuilder MenuBuilder, TShared
 				FCanExecuteAction::CreateSP(this, &FMeshEditorMode::IsMeshElementTypeSelected, EEditableMeshElementType::Vertex)
 			),
 			NAME_None,
-			EUserInterfaceActionType::Button
+			EUserInterfaceActionType::CollapsedButton
 			);
 		MenuBuilder.AddMenuEntry(
 			LOCTEXT("Extend", "Extend"),
@@ -6411,7 +6411,7 @@ void FMeshEditorMode::MeshEditActionsGenerator(FMenuBuilder MenuBuilder, TShared
 				FIsActionChecked::CreateLambda([this] { return (EquippedVertexAction == EMeshEditAction::ExtendVertex); })
 				),
 			NAME_None,
-			EUserInterfaceActionType::Check
+			EUserInterfaceActionType::ToggleButton
 			);
 		MenuBuilder.AddMenuEntry(
 			LOCTEXT("WeldSelected", "Weld Selected"),
@@ -6423,7 +6423,7 @@ void FMeshEditorMode::MeshEditActionsGenerator(FMenuBuilder MenuBuilder, TShared
 				FCanExecuteAction::CreateSP(this, &FMeshEditorMode::IsMeshElementTypeSelected, EEditableMeshElementType::Vertex)
 				),
 			NAME_None,
-			EUserInterfaceActionType::Button
+			EUserInterfaceActionType::CollapsedButton
 			);
 	}
 }
diff --git a/Engine/Source/Editor/UnrealEd/Classes/Settings/EditorExperimentalSettings.h b/Engine/Source/Editor/UnrealEd/Classes/Settings/EditorExperimentalSettings.h
--- a/Engine/Source/Editor/UnrealEd/Classes/Settings/EditorExperimentalSettings.h
+++ b/Engine/Source/Editor/UnrealEd/Classes/Settings/EditorExperimentalSettings.h
@@ -125,10 +125,6 @@ public:
 	UPROPERTY(EditAnywhere, config, Category = Tools, meta = (DisplayName = "Use OpenCL for convex hull decomposition"))
 	bool bUseOpenCLForConvexHullDecomp;
 
-	/**If true, wearing a Vive or Oculus Rift headset will automatically enter VR Editing mode */
-	UPROPERTY(EditAnywhere, config, Category = VR, meta = (DisplayName = "Enable VR Mode Auto-Entry"))
-	bool bEnableAutoVREditMode;
-
 	/** Allows editing of potentially unsafe properties during PIE. Advanced use only - use with caution. */
 	UPROPERTY(EditAnywhere, config, Category = Tools, meta = (DisplayName = "Allow editing of potentially unsafe properties."))
 	bool bAllowPotentiallyUnsafePropertyEditing;
diff --git a/Engine/Source/Editor/VREditor/Gizmo/VREditorTransformGizmo.cpp b/Engine/Source/Editor/VREditor/Gizmo/VREditorTransformGizmo.cpp
--- a/Engine/Source/Editor/VREditor/Gizmo/VREditorTransformGizmo.cpp
+++ b/Engine/Source/Editor/VREditor/Gizmo/VREditorTransformGizmo.cpp
@@ -15,11 +15,11 @@
 #include "VREditorPlaneTranslationGizmoHandle.h"
 #include "VIStretchGizmoHandle.h"
 #include "VIUniformScaleGizmoHandle.h"
+#include "VRModeSettings.h"
 
 namespace VREd
 {
 	// @todo vreditor tweak: Tweak out console variables
-	static FAutoConsoleVariable GizmoScale( TEXT( "VREd.GizmoScale" ), 0.80f, TEXT( "How big the gizmo handles should be" ) );
 	static FAutoConsoleVariable GizmoDistanceScaleFactor( TEXT( "VREd.GizmoDistanceScaleFactor" ), 0.002f, TEXT( "How much the gizmo handles should increase in size with distance from the camera, to make it easier to select" ) );
 	static FAutoConsoleVariable MinActorSizeForTransformGizmo( TEXT( "VREd.MinActorSizeForTransformGizmo" ), 50.0f, TEXT( "How big an object must be in scaled world units before we'll start to shrink the gizmo" ) );
 }
@@ -167,7 +167,7 @@ void ATransformGizmo::UpdateGizmo(const EGizmoHandleTypes InGizmoType, const ECo
 	// @todo vreditor: Should probably be a curve, not linear
 	// @todo vreditor: Should take FOV into account (especially in non-stereo/HMD mode)
 	const float WorldSpaceDistanceToGizmoBounds = FMath::Sqrt(WorldSpaceBounds.ComputeSquaredDistanceToPoint(InViewLocation));
-	const float GizmoScaleUpClose = VREd::GizmoScale->GetFloat();
+	const float GizmoScaleUpClose = GetDefault<UVRModeSettings>()->GizmoScale;
 	const float GizmoScale(InScaleMultiplier * (GizmoScaleUpClose + (WorldSpaceDistanceToGizmoBounds / WorldScaleFactor) * VREd::GizmoDistanceScaleFactor->GetFloat()) * ScaleCompensationForTinyGizmo * WorldScaleFactor);
 
 	// Update animation
diff --git a/Engine/Source/Editor/VREditor/Public/VRModeSettings.h b/Engine/Source/Editor/VREditor/Public/VRModeSettings.h
new file mode 100644
--- /dev/null
+++ b/Engine/Source/Editor/VREditor/Public/VRModeSettings.h
@@ -0,0 +1,69 @@
+// Copyright 1998-2017 Epic Games, Inc. All Rights Reserved.
+
+#pragma once
+
+#include "CoreMinimal.h"
+#include "UObject/ObjectMacros.h"
+#include "UObject/Object.h"
+#include "VRModeSettings.generated.h"
+
+UENUM()
+enum class EInteractorHand : uint8
+{
+	/** Right hand */
+	Right,
+
+	/** Left hand */
+	Left,
+};
+
+/**
+* Implements the settings for VR Mode.
+*/
+UCLASS(config = EditorSettings)
+class VREDITOR_API UVRModeSettings : public UObject
+{
+public:
+	GENERATED_UCLASS_BODY()
+
+	/**If true, wearing a Vive or Oculus Rift headset will automatically enter VR Editing mode */
+	UPROPERTY(EditAnywhere, config, Category = "General", meta = (DisplayName = "Enable VR Mode Auto-Entry"))
+	uint32 bEnableAutoVREditMode : 1;
+
+	// Which hand should have the primary interactor laser on it
+	UPROPERTY(Category = "General", EditAnywhere)
+	EInteractorHand InteractorHand;
+
+	/** Show the movement grid for orientation while moving through the world */
+	UPROPERTY(EditAnywhere, config, Category = "UI Customization")
+	uint32 bShowWorldMovementGrid : 1;
+
+	/** Dim the surroundings while moving through the world */
+	UPROPERTY(EditAnywhere, config, Category = "UI Customization")
+	uint32 bShowWorldMovementPostProcess : 1;
+
+	/** Display a progress bar while scaling that shows your current scale */
+	UPROPERTY(EditAnywhere, config, Category = "UI Customization")
+	uint32 bShowWorldScaleProgressBar : 1;
+
+	/** Adjusts the brightness of the UI panels */
+	UPROPERTY(EditAnywhere, config, Category = "UI Customization", meta = (DisplayName = "UI Panel Brightness", ClampMin = 0.01, ClampMax = 10.0))
+	float UIBrightness;
+
+	/** The size of the transform gizmo */
+	UPROPERTY(EditAnywhere, config, Category = "UI Customization", meta = (ClampMin = 0.1, ClampMax = 2.0))
+	float GizmoScale;
+
+	/** The maximum time in seconds between two clicks for a double-click to register */
+	UPROPERTY(EditAnywhere, config, Category = "Motion Controllers", meta = (ClampMin = 0.01, ClampMax = 1.0))
+	float DoubleClickTime;
+
+	/** The amount (between 0-1) you have to depress the Vive controller trigger to register a press */
+	UPROPERTY(EditAnywhere, config, Category = "Motion Controllers", meta = (DisplayName = "Trigger Pressed Threshold (Vive)", ClampMin = 0.01, ClampMax = 1.0))
+	float TriggerPressedThreshold_Vive;
+
+	/** The amount (between 0-1) you have to depress the Oculus Touch controller trigger to register a press */
+	UPROPERTY(EditAnywhere, config, Category = "Motion Controllers", meta = (DisplayName = "Trigger Pressed Threshold (Oculus Touch)", ClampMin = 0.01, ClampMax = 1.0))
+	float TriggerPressedThreshold_Rift;
+
+};
diff --git a/Engine/Source/Editor/VREditor/UI/VREditorFloatingUI.cpp b/Engine/Source/Editor/VREditor/UI/VREditorFloatingUI.cpp
--- a/Engine/Source/Editor/VREditor/UI/VREditorFloatingUI.cpp
+++ b/Engine/Source/Editor/VREditor/UI/VREditorFloatingUI.cpp
@@ -7,11 +7,11 @@
 #include "Components/WidgetComponent.h"
 #include "VREditorWidgetComponent.h"
 #include "Components/StaticMeshComponent.h"
+#include "VRModeSettings.h"
 
 namespace VREd
 {
 	static FAutoConsoleVariable UIFadeSpeed( TEXT( "VREd.UIFadeSpeed" ), 6.0f, TEXT( "How fast UI should fade in and out" ) );
-	static FAutoConsoleVariable UIBrightness( TEXT( "VREd.UIBrightness" ), 1.5f, TEXT( "How bright the UI should be" ) );
 }
 
 
@@ -240,7 +240,7 @@ void AVREditorFloatingUI::UpdateFadingState( const float DeltaTime )
 		}
 
  		// Set material color
-		const float UIBrightness = FadeAlpha * VREd::UIBrightness->GetFloat();
+		const float UIBrightness = FadeAlpha * GetDefault<UVRModeSettings>()->UIBrightness;
 		WidgetComponent->SetTintColorAndOpacity( FLinearColor( UIBrightness, UIBrightness, UIBrightness ).CopyWithNewOpacity( FadeAlpha ) );
 
 	}
diff --git a/Engine/Source/Editor/VREditor/UI/VREditorRadialFloatingUI.cpp b/Engine/Source/Editor/VREditor/UI/VREditorRadialFloatingUI.cpp
--- a/Engine/Source/Editor/VREditor/UI/VREditorRadialFloatingUI.cpp
+++ b/Engine/Source/Editor/VREditor/UI/VREditorRadialFloatingUI.cpp
@@ -75,12 +75,17 @@ AVREditorRadialFloatingUI::AVREditorRadialFloatingUI()
 		ArrowMeshComponent->SetCollisionResponseToAllChannels(ECollisionResponse::ECR_Ignore);
 
 		ArrowMeshComponent->SetStaticMesh(AssetContainer->RadialMenuPointerMesh);
+		ArrowMeshComponent->CreateAndSetMaterialInstanceDynamic(0);
+		UMaterialInstanceDynamic* ArrowMaterial = Cast<UMaterialInstanceDynamic>(ArrowMeshComponent->GetMaterial(0));
+		ArrowAlpha = 0.0f;
+		ArrowMaterial->SetScalarParameterValue("Alpha", ArrowAlpha);
 		ArrowMeshComponent->bGenerateOverlapEvents = false;
 		ArrowMeshComponent->SetCanEverAffectNavigation(false);
 		ArrowMeshComponent->bCastDynamicShadow = false;
 		ArrowMeshComponent->bCastStaticShadow = false;
 		ArrowMeshComponent->bAffectDistanceFieldLighting = false;
 		ArrowMeshComponent->bSelectable = false;
+		ArrowMeshComponent->SetVisibility(false);
 	}
 
 	{
@@ -257,6 +262,19 @@ void AVREditorRadialFloatingUI::UpdateFadingState( const float DeltaTime )
 		}
 		DiskMaterial->SetScalarParameterValue("GlowAmount", GlowAmount);
 	}
+	if (ArrowMeshComponent != nullptr)
+	{
+		UMaterialInstanceDynamic* ArrowMaterial = Cast<UMaterialInstanceDynamic>(ArrowMeshComponent->GetMaterial(0));
+		if (ArrowAlpha < 1.0 && ArrowMeshComponent->IsVisible())
+		{
+			ArrowAlpha = FMath::Min(1.0f, ArrowAlpha + VREd::RadialUIFadeSpeed->GetFloat() * DeltaTime);
+		}
+		if (!ArrowMeshComponent->IsVisible())
+		{
+			ArrowAlpha = 0.0f;
+		}
+		ArrowMaterial->SetScalarParameterValue("Alpha", ArrowAlpha);
+	}
 	if ( FadeDelay > 0.f )
 	{
 		FadeDelay -= DeltaTime;
diff --git a/Engine/Source/Editor/VREditor/UI/VREditorRadialFloatingUI.h b/Engine/Source/Editor/VREditor/UI/VREditorRadialFloatingUI.h
--- a/Engine/Source/Editor/VREditor/UI/VREditorRadialFloatingUI.h
+++ b/Engine/Source/Editor/VREditor/UI/VREditorRadialFloatingUI.h
@@ -207,5 +207,9 @@ private:
 	float GlowAmount;
 
 	float DefaultGlowAmount;
+
+	/** Alpha of arrow - VFX */
+	float ArrowAlpha;
+
 };
 
diff --git a/Engine/Source/Editor/VREditor/UI/VREditorUISystem.cpp b/Engine/Source/Editor/VREditor/UI/VREditorUISystem.cpp
--- a/Engine/Source/Editor/VREditor/UI/VREditorUISystem.cpp
+++ b/Engine/Source/Editor/VREditor/UI/VREditorUISystem.cpp
@@ -85,6 +85,7 @@
 #include "Widgets/Layout/SDPIScaler.h"
 #include "Widgets/Docking/SDockTab.h"
 #include "Widgets/Colors/SColorPicker.h"
+#include "Widgets/Input/SCheckBox.h"
 #include "MultiBox.h"
 #include "LevelEditorActions.h"
 #include "VREditorActions.h"
@@ -98,6 +99,7 @@
 #include "SImage.h"
 
 #include "DrawDebugHelpers.h"
+#include "VRModeSettings.h"
 
 #define LOCTEXT_NAMESPACE "VREditor"
 
@@ -123,7 +125,6 @@ namespace VREd
 	static FAutoConsoleVariable UIAbsoluteScrollSpeed( TEXT( "VREd.UIAbsoluteScrollSpeed" ), 8.0f, TEXT( "How fast the UI scrolls when dragging the touchpad" ) );
 	static FAutoConsoleVariable UIRelativeScrollSpeed( TEXT( "VREd.UIRelativeScrollSpeed" ), 0.75f, TEXT( "How fast the UI scrolls when holding an analog stick" ) );
 	static FAutoConsoleVariable MinUIScrollDeltaForInertia( TEXT( "VREd.MinUIScrollDeltaForInertia" ), 0.25f, TEXT( "Minimum amount of touch pad input before inertial UI scrolling kicks in" ) );
-	static FAutoConsoleVariable DoubleClickTime( TEXT( "VREd.DoubleClickTime" ), 0.25f, TEXT( "Minimum duration between clicks for a double click event to fire" ) );
 	static FAutoConsoleVariable UIPressHapticFeedbackStrength( TEXT( "VREd.UIPressHapticFeedbackStrength" ), 0.4f, TEXT( "Strenth of haptic when clicking on the UI" ) );
 	static FAutoConsoleVariable UIAssetEditorSummonedOnHandHapticFeedbackStrength( TEXT( "VREd.UIAssetEditorSummonedOnHandHapticFeedbackStrength" ), 1.0f, TEXT( "Strenth of haptic to play to remind a user which hand an asset editor was spawned on" ) );
 	static FAutoConsoleVariable MaxDockWindowSize( TEXT( "VREd.MaxDockWindowSize" ), 250, TEXT( "Maximum size for dockable windows" ) );
@@ -132,7 +133,6 @@ namespace VREd
 	static FAutoConsoleVariable UIPanelOpenRotationPitchOffset( TEXT( "VREd.UIPanelOpenRotationPitchOffset" ), 45.0f, TEXT( "The pitch rotation offset in degrees when spawning a panel in front of the motioncontroller" ) );
 	static FAutoConsoleVariable ColorPickerDockSpawnOffset( TEXT( "VREd.ColorPickerDockSpawnOffset" ),  3.f, TEXT( "Offset of where the color picker spawns" ) );
 	static FAutoConsoleVariable EnableQuickMenu(TEXT("VREd.EnableQuickMenu"), 0, TEXT("Whether the quick menu should be enabled. Disabled by default."));
-	static FAutoConsoleVariable LeftHandDominant(TEXT("VREd.LeftHandDominant"), 0, TEXT("Set to 1 to make the left hand the dominant hand"));
 	static FAutoConsoleVariable SteamVRTrackpadDeadzone(TEXT("VREd.SteamVRTrackpadDeadzone"), 0.3f, TEXT("The deadzone for the Vive motion controller trackpad"));
 	static const FTransform DefaultColorPickerTransform(	FRotator( -10, 180, 0), FVector( 30, 35, 0 ), FVector( 1 ) );
 }
@@ -151,7 +151,7 @@ UVREditorUISystem::UVREditorUISystem( const FObjectInitializer& Initializer ) :
 	ColorPickerUI( nullptr ),
 	bRefocusViewport(false),
 	LastDraggingHoverLocation( FVector::ZeroVector ),
-	bSequencerOpenendFromRadialMenu(false),
+	bSequencerOpenedFromRadialMenu(false),
 	bDragPanelFromOpen(false),
 	DragPanelFromOpenTime(0.0f),
 	bPanelCanScale(true)
@@ -169,7 +169,7 @@ void UVREditorUISystem::Init()
 	FSlateApplication::Get().OnDragDropCheckOverride.BindUObject(this, &UVREditorUISystem::CheckForVRDragDrop);
 
 	// TODO: asymmetry hardcoded right now
-	if (VREd::LeftHandDominant->GetInt() == 1)
+	if (GetDefault<UVRModeSettings>()->InteractorHand == EInteractorHand::Left)
 	{
 		LaserInteractor = Cast<UVREditorMotionControllerInteractor>(GetOwner().GetHandInteractor(EControllerHand::Left));
 		UIInteractor = Cast<UVREditorMotionControllerInteractor>(GetOwner().GetHandInteractor(EControllerHand::Right));
@@ -266,7 +266,10 @@ void UVREditorUISystem::OnPreviewInputAction(FEditorViewportClient& ViewportClie
 		VREditorInteractor == UIInteractor && 
 		(VREditorInteractor->GetDraggingMode() != EViewportInteractionDraggingMode::World || (VREditorInteractor->GetOtherInteractor() != nullptr && VREditorInteractor->GetOtherInteractor()->GetDraggingMode() != EViewportInteractionDraggingMode::World && VREditorInteractor->GetDraggingMode() == EViewportInteractionDraggingMode::AssistingDrag )))
 	{
-		if (Action.Event == EInputEvent::IE_Pressed && VRMode->GetHMDDeviceType() == EHMDDeviceType::DT_OculusRift && Action.ActionType == VRActionTypes::ConfirmRadialSelection)
+		if (Action.Event == EInputEvent::IE_Pressed && 
+			VRMode->GetHMDDeviceType() == EHMDDeviceType::DT_OculusRift 
+			&& Action.ActionType == VRActionTypes::ConfirmRadialSelection
+			&& GIsDemoMode)
 		{
 			ResetAll();
 			bWasHandled = true;
@@ -460,7 +463,7 @@ void UVREditorUISystem::OnPreviewInputAction(FEditorViewportClient& ViewportClie
 									if (Action.Event == IE_Pressed)
 									{
 										const double CurrentTime = FPlatformTime::Seconds();
-										if (CurrentTime - VREditorInteractor->GetLastUIPressTime() <= VREd::DoubleClickTime->GetFloat())
+										if (CurrentTime - VREditorInteractor->GetLastUIPressTime() <= GetDefault<UVRModeSettings>()->DoubleClickTime)
 										{
 											// Trigger a double click event!
 											Reply = FSlateApplication::Get().RoutePointerDoubleClickEvent(WidgetPathUnderFinger, PointerEvent);
@@ -1634,8 +1637,8 @@ void UVREditorUISystem::TogglePanelVisibility( const EEditorUIPanel EditorUIPane
 		}
 		else
 		{
+			const bool bDragFromOpen = ShouldPreviewPanel();
 			const bool bSpawnInFront = true;
-			const bool bDragFromOpen = true;
 			const bool bPlaySound = false;
 			ShowEditorUIPanel(Panel, UIInteractor, !bIsShowing, bSpawnInFront, bDragFromOpen, bPlaySound);
 		}
@@ -1769,7 +1772,7 @@ void UVREditorUISystem::MakeUniformGridMenu(const TSharedRef<FMultiBox>& MultiBo
 				];
 
 			//AddHoverableButton(BlockWidget, TEXT("SMenuEntryButton"), TestOverlay);
-			SetButtonTextWrap(BlockWidget, 50.0);
+			//SetButtonTextWrap(BlockWidget, 50.0);
 			Index++;
 		}
 	}
@@ -1803,41 +1806,49 @@ TSharedRef<SWidget> UVREditorUISystem::AddHoverableButton(TSharedRef<SWidget>& B
 	return BlockWidget;
 }
 
-TSharedRef<SWidget> UVREditorUISystem::SetButtonTextWrap(TSharedRef<SWidget>& BlockWidget, float WrapSize)
+TSharedRef<SWidget> UVREditorUISystem::SetButtonFormatting(TSharedRef<SWidget>& BlockWidget, float WrapSize)
 {
-
-	TArray<TSharedRef<SWidget>> TestWidgets;
-	FindAllWidgetsOfType(TestWidgets, BlockWidget, FName(TEXT("SImage")));
-	if (TestWidgets.Num() > 0)
-	{
-		for (const TSharedRef<SWidget> TestWidget : TestWidgets)
-		{
-			TSharedRef<SImage> Image = StaticCastSharedRef<SImage>(TestWidget);
-			Image->SetRenderTransformPivot(FVector2D(0.5f, 0.5f));
-			Image->SetRenderTransform(FSlateRenderTransform::FTransform2D(4.0f));
-		}
-	}
-
-	TestWidgets.Empty();
-	FindAllWidgetsOfType(TestWidgets, BlockWidget, FName(TEXT("SBox")));
-	if (TestWidgets.Num() > 0)
+	TSharedRef<SWidget> TestWidget = FindWidgetOfType(BlockWidget, FName(TEXT("SImage")));
+	if (TestWidget != SNullWidget::NullWidget)
 	{
-		for (const TSharedRef<SWidget> TestWidget : TestWidgets)
-		{
-				TSharedRef<SBox> Box = StaticCastSharedRef<SBox>(TestWidget);
-				Box->SetMinDesiredHeight(64.0f);
-		}
+		TSharedRef<SImage> Image = StaticCastSharedRef<SImage>(TestWidget);
+		Image->SetRenderTransformPivot(FVector2D(0.5f, 0.5f));
+		Image->SetRenderTransform(FSlateRenderTransform::FTransform2D(4.0f));	
 	}
 
-	TSharedRef<SWidget> TestWidget = FindWidgetOfType(BlockWidget, FName(TEXT("STextBlock")));
+	// Format the button text
+	TestWidget = FindWidgetOfType(BlockWidget, FName(TEXT("STextBlock")));
 	if (TestWidget != SNullWidget::NullWidget)
 	{
 		TSharedRef<STextBlock> TextBlock = StaticCastSharedRef<STextBlock>(TestWidget);
 		TextBlock->SetWrapTextAt(WrapSize);
 		TextBlock->SetJustification(ETextJustify::Center);
-	}	
-
+		
+		// Change the button font color based on whether the option is checked or not
+		TSharedRef<SWidget> TestCheckbox = FindWidgetOfType(BlockWidget, FName(TEXT("SCheckBox")));
+		if (TestCheckbox != SNullWidget::NullWidget)
+		{
+			TSharedRef<SCheckBox> CheckBox = StaticCastSharedRef<SCheckBox>(TestCheckbox);
+			TAttribute<FSlateColor> DynamicSelectedTextColor = TAttribute<FSlateColor>::Create(
+				TAttribute<FSlateColor>::FGetter::CreateLambda([this, CheckBox]() -> FSlateColor
+				{
+					return CheckBox->IsChecked() ? this->GetOwner().GetColor(UVREditorMode::EColors::UIColor) : FSlateColor(FLinearColor::White);
+				}
+				)
+				);
+				TAttribute<FSlateFontInfo> DynamicSelectedTextFont = TAttribute<FSlateFontInfo>::Create(
+					TAttribute<FSlateFontInfo>::FGetter::CreateLambda([CheckBox]() -> FSlateFontInfo
+				{
+					return CheckBox->IsChecked() ? FVREditorStyle::Get().GetFontStyle("VRRadialStyle.ActiveFont") : FVREditorStyle::Get().GetFontStyle("VRRadialStyle.InactiveFont");
+				}
+				)
+				);
+				TextBlock->SetColorAndOpacity(DynamicSelectedTextColor);
+				TextBlock->SetFont(DynamicSelectedTextFont);
+			}
+		
 
+	}	
 	return BlockWidget;
 }
 
@@ -1888,7 +1899,7 @@ void UVREditorUISystem::MakeRadialBoxMenu(const TSharedRef<FMultiBox>& MultiBox,
 			TSharedRef<SOverlay> TestOverlay = SNew(SOverlay);
 
 			BlockWidget = AddHoverableButton(BlockWidget, ButtonTypeOverride, TestOverlay);
-			BlockWidget = SetButtonTextWrap(BlockWidget, 50.0f);
+			BlockWidget = SetButtonFormatting(BlockWidget, 50.0f);
 			TSharedRef<SBox> RadialMenuElement = SNew(SBox);
 			RadialMenuElement->SetContent(
 				SNew(SDPIScaler)
@@ -1968,7 +1979,7 @@ TSharedRef<SWidget> UVREditorUISystem::BuildQuickMenuWidget()
 			FGetActionCheckState::CreateStatic(&FVREditorActionCallbacks::AreAligningToActors, VRMode)
 			),
 		NAME_None,
-		EUserInterfaceActionType::Check
+		EUserInterfaceActionType::ToggleButton
 		);
 
 	TAttribute<FText> DynamicAlignSelectionLabel;
@@ -2005,7 +2016,7 @@ TSharedRef<SWidget> UVREditorUISystem::BuildQuickMenuWidget()
 			FGetActionCheckState::CreateStatic(&FVREditorActionCallbacks::GetTranslationSnapState)
 		),
 		NAME_None,
-		EUserInterfaceActionType::Check
+		EUserInterfaceActionType::ToggleButton
 		);
 	TAttribute<FText> DynamicTranslationSizeLabel;
 	DynamicTranslationSizeLabel.BindStatic(&FVREditorActionCallbacks::GetTranslationSnapSizeText);
@@ -2032,7 +2043,7 @@ TSharedRef<SWidget> UVREditorUISystem::BuildQuickMenuWidget()
 			FGetActionCheckState::CreateStatic(&FVREditorActionCallbacks::GetRotationSnapState)
 			),
 		NAME_None,
-		EUserInterfaceActionType::Check
+		EUserInterfaceActionType::ToggleButton
 		);
 	TAttribute<FText> DynamicRotationSizeLabel;
 	DynamicRotationSizeLabel.BindStatic(&FVREditorActionCallbacks::GetRotationSnapSizeText);
@@ -2059,7 +2070,7 @@ TSharedRef<SWidget> UVREditorUISystem::BuildQuickMenuWidget()
 			FGetActionCheckState::CreateStatic(&FVREditorActionCallbacks::GetScaleSnapState)
 			),
 		NAME_None,
-		EUserInterfaceActionType::Check
+		EUserInterfaceActionType::ToggleButton
 		);
 	TAttribute<FText> DynamicScaleSizeLabel;
 	DynamicScaleSizeLabel.BindStatic(&FVREditorActionCallbacks::GetScaleSnapSizeText);
@@ -2094,7 +2105,7 @@ TSharedRef<SWidget> UVREditorUISystem::BuildQuickMenuWidget()
 			FGetActionCheckState::CreateStatic(&FVREditorActionCallbacks::GetUIToggledState, VRMode, UVREditorUISystem::EEditorUIPanel::ActorDetails)
 			),
 		NAME_None,
-		EUserInterfaceActionType::Check
+		EUserInterfaceActionType::ToggleButton
 		);
 	InterfaceMenuBuilder.AddMenuEntry(
 		LOCTEXT("ContentBrowser", "Content Browser"),
@@ -2107,7 +2118,7 @@ TSharedRef<SWidget> UVREditorUISystem::BuildQuickMenuWidget()
 			FGetActionCheckState::CreateStatic(&FVREditorActionCallbacks::GetUIToggledState, VRMode, UVREditorUISystem::EEditorUIPanel::ContentBrowser)
 			),
 		NAME_None,
-		EUserInterfaceActionType::Check
+		EUserInterfaceActionType::ToggleButton
 		);
 	InterfaceMenuBuilder.AddMenuEntry(
 		LOCTEXT("Modes", "Modes"),
@@ -2120,7 +2131,7 @@ TSharedRef<SWidget> UVREditorUISystem::BuildQuickMenuWidget()
 			FGetActionCheckState::CreateStatic(&FVREditorActionCallbacks::GetUIToggledState, VRMode, UVREditorUISystem::EEditorUIPanel::Modes)
 			),
 		NAME_None,
-		EUserInterfaceActionType::Check
+		EUserInterfaceActionType::ToggleButton
 		);
 	InterfaceMenuBuilder.AddMenuEntry(
 		LOCTEXT("WorldOutliner", "World Outliner"),
@@ -2133,7 +2144,7 @@ TSharedRef<SWidget> UVREditorUISystem::BuildQuickMenuWidget()
 			FGetActionCheckState::CreateStatic(&FVREditorActionCallbacks::GetUIToggledState, VRMode, UVREditorUISystem::EEditorUIPanel::WorldOutliner)
 			),
 		NAME_None,
-		EUserInterfaceActionType::Check
+		EUserInterfaceActionType::ToggleButton
 		);
 	InterfaceMenuBuilder.AddMenuEntry(
 		LOCTEXT("WorldSettings", "World Settings"),
@@ -2146,7 +2157,7 @@ TSharedRef<SWidget> UVREditorUISystem::BuildQuickMenuWidget()
 			FGetActionCheckState::CreateStatic(&FVREditorActionCallbacks::GetUIToggledState, VRMode, UVREditorUISystem::EEditorUIPanel::WorldSettings)
 			),
 		NAME_None,
-		EUserInterfaceActionType::Check
+		EUserInterfaceActionType::ToggleButton
 		);
 	InterfaceMenuBuilder.AddMenuEntry(
 		LOCTEXT("CreateNewSequence", "Create Sequence"), 
@@ -2603,7 +2614,7 @@ void UVREditorUISystem::OnHoverEndEffect(UVREditorWidgetComponent* Button)
 
 void UVREditorUISystem::SequencerOpenendFromRadialMenu(const bool bInOpenedFromRadialMenu /*= true*/)
 {
-	bSequencerOpenendFromRadialMenu = bInOpenedFromRadialMenu;
+	bSequencerOpenedFromRadialMenu = bInOpenedFromRadialMenu;
 }
 
 bool UVREditorUISystem::CanScalePanel() const
@@ -2611,6 +2622,11 @@ bool UVREditorUISystem::CanScalePanel() const
 	return bPanelCanScale || (DraggingUI != nullptr && LaserInteractor != nullptr && InteractorDraggingUI != nullptr && LaserInteractor == InteractorDraggingUI);
 }
 
+UVREditorMotionControllerInteractor* UVREditorUISystem::GetUIInteractor()
+{
+	return UIInteractor;
+}
+
 void UVREditorUISystem::SequencerRadialMenuGenerator(FMenuBuilder MenuBuilder, TSharedPtr<FUICommandList> CommandList, class UVREditorMode* VRMode, float& RadiusOverride)
 {
 	RadiusOverride = 1.0f;
@@ -2666,7 +2682,7 @@ void UVREditorUISystem::SequencerRadialMenuGenerator(FMenuBuilder MenuBuilder, T
 			FGetActionCheckState::CreateStatic(&FVREditorActionCallbacks::GetSequencerScrubState, UIInteractor)
 			),
 		NAME_None,
-		EUserInterfaceActionType::Check
+		EUserInterfaceActionType::ToggleButton
 		);
 	MenuBuilder.AddMenuEntry(
 		LOCTEXT("ToggleLooping", "Toggle Looping"),
@@ -2679,7 +2695,7 @@ void UVREditorUISystem::SequencerRadialMenuGenerator(FMenuBuilder MenuBuilder, T
 			FGetActionCheckState::CreateStatic(&FVREditorActionCallbacks::IsLoopingChecked, VRMode)
 			),
 		NAME_None,
-		EUserInterfaceActionType::Check
+		EUserInterfaceActionType::ToggleButton
 		);
 }
 
@@ -2696,7 +2712,7 @@ void UVREditorUISystem::HandleEditorModeChanged(FEdMode* Mode, bool IsEnabled)
 			{
 				const bool bShouldShow = true;
 				const bool bSpawnInFront = true;
-				const bool bDragFromOpen = true;
+				const bool bDragFromOpen = ShouldPreviewPanel();
 				ShowEditorUIPanel(Panel, UIInteractor, bShouldShow, bSpawnInFront, bDragFromOpen);
 			}
 		}
@@ -2760,6 +2776,24 @@ bool UVREditorUISystem::CheckForVRDragDrop()
 	}
 }
 
+bool UVREditorUISystem::ShouldPreviewPanel()
+{
+	bool bDragFromOpen = true;
+	if (LaserInteractor != nullptr)
+	{
+		// If we are clicking with the laser interactor, instantly spawn the panel
+		FHitResult HitResult = LaserInteractor->GetHitResultFromLaserPointer();
+		if (HitResult.Actor.IsValid())
+		{
+			if (HitResult.Actor == QuickRadialMenu)
+			{
+				bDragFromOpen = !LaserInteractor->IsClickingOnUI();
+			}
+		}
+	}
+	return bDragFromOpen;
+}
+
 AVREditorRadialFloatingUI* UVREditorUISystem::GetRadialMenuFloatingUI() const
 {
 	return QuickRadialMenu;
@@ -2800,10 +2834,13 @@ void UVREditorUISystem::UpdateSequencerUI()
 			const bool bWithSceneComponent = false;
 			SequencerPanel->SetSlateWidget(*this, WidgetToDraw, FIntPoint(VREd::SequencerUIResolutionX->GetFloat(), VREd::SequencerUIResolutionY->GetFloat()), VREd::EditorUISize->GetFloat(), AVREditorFloatingUI::EDockedTo::Nothing);
 	
-			if (bSequencerOpenendFromRadialMenu == true)
+			if (bSequencerOpenedFromRadialMenu == true)
 			{
-				ShowEditorUIPanel(SequencerPanel, UIInteractor, true, true, false, true);
-				bSequencerOpenendFromRadialMenu = false;
+				const bool bDragFromOpen = ShouldPreviewPanel();
+				const bool bShouldShow = true;
+				const bool bSpawnInFront = true;
+				ShowEditorUIPanel(SequencerPanel, UIInteractor, bShouldShow, bSpawnInFront, bDragFromOpen);
+				bSequencerOpenedFromRadialMenu = false;
 			}
 			else
 			{
diff --git a/Engine/Source/Editor/VREditor/UI/VREditorUISystem.h b/Engine/Source/Editor/VREditor/UI/VREditorUISystem.h
--- a/Engine/Source/Editor/VREditor/UI/VREditorUISystem.h
+++ b/Engine/Source/Editor/VREditor/UI/VREditorUISystem.h
@@ -235,6 +235,9 @@ public:
 	/** If a dockable window can be scaled */
 	bool CanScalePanel() const;
 
+	/** Get the interactor that holds the radial menu */
+	class UVREditorMotionControllerInteractor* GetUIInteractor();
+
 protected:
 
 	/** Called to "preview" an input event to get a first chance at it. */
@@ -286,7 +289,7 @@ protected:
 	/** Adds a hoverable button of a given type to an overlay, using menu data from a BlockWidget */
 	TSharedRef<SWidget> AddHoverableButton(TSharedRef<SWidget>& BlockWidget, FName ButtonType, TSharedRef<SOverlay>& TestOverlay);
 	/** Sets the text wrap size of the text block element nested in a BlockWidget */
-	TSharedRef<SWidget> SetButtonTextWrap(TSharedRef<SWidget>& BlockWidget, float WrapSize);
+	TSharedRef<SWidget> SetButtonFormatting(TSharedRef<SWidget>& BlockWidget, float WrapSize);
 	
 	/** Builds the quick menu Slate widget */
 	TSharedRef<SWidget> BuildQuickMenuWidget();
@@ -313,6 +316,9 @@ protected:
 	/**The VR Editor UI System's rules for when drag drop should be checked for */
 	bool CheckForVRDragDrop();
 
+	/** Preview the UI panel's location if spawning with the UI interactor, else spawn immediately */
+	bool ShouldPreviewPanel();
+
 protected:
 
 	/** Owning object */
@@ -403,7 +409,7 @@ protected:
 	FTimespan RadialMenuModifierSpawnTime;
 
 	/** If sequenced was opened from the radial menu or somewhere else such as the content browser. */
-	bool bSequencerOpenendFromRadialMenu;
+	bool bSequencerOpenedFromRadialMenu;
 
 	/** If started dragging from opening a UI panel */
 	bool bDragPanelFromOpen;
diff --git a/Engine/Source/Editor/VREditor/UI/VRRadialMenuHandler.cpp b/Engine/Source/Editor/VREditor/UI/VRRadialMenuHandler.cpp
--- a/Engine/Source/Editor/VREditor/UI/VRRadialMenuHandler.cpp
+++ b/Engine/Source/Editor/VREditor/UI/VRRadialMenuHandler.cpp
@@ -275,7 +275,7 @@ void UVRRadialMenuHandler::SnapMenuGenerator(FMenuBuilder MenuBuilder, TSharedPt
 			FGetActionCheckState::CreateStatic(&FVREditorActionCallbacks::GetTranslationSnapState)
 			),
 		NAME_None,
-		EUserInterfaceActionType::Check
+		EUserInterfaceActionType::ToggleButton
 		);
 	TAttribute<FText> DynamicTranslationSizeLabel;
 	DynamicTranslationSizeLabel.BindStatic(&FVREditorActionCallbacks::GetTranslationSnapSizeText);
@@ -289,7 +289,7 @@ void UVRRadialMenuHandler::SnapMenuGenerator(FMenuBuilder MenuBuilder, TSharedPt
 			FCanExecuteAction::CreateStatic(&FLevelEditorActionCallbacks::DefaultCanExecuteAction)
 			),
 		NAME_None,
-		EUserInterfaceActionType::CollapsedButton
+		EUserInterfaceActionType::ToggleButton
 		);
 	MenuBuilder.AddMenuEntry(
 		LOCTEXT("ToggleRotationSnap", "Rotate Snap"),
@@ -302,7 +302,7 @@ void UVRRadialMenuHandler::SnapMenuGenerator(FMenuBuilder MenuBuilder, TSharedPt
 			FGetActionCheckState::CreateStatic(&FVREditorActionCallbacks::GetRotationSnapState)
 			),
 		NAME_None,
-		EUserInterfaceActionType::Check
+		EUserInterfaceActionType::ToggleButton
 		);
 	TAttribute<FText> DynamicRotationSizeLabel;
 	DynamicRotationSizeLabel.BindStatic(&FVREditorActionCallbacks::GetRotationSnapSizeText);
@@ -329,7 +329,7 @@ void UVRRadialMenuHandler::SnapMenuGenerator(FMenuBuilder MenuBuilder, TSharedPt
 			FGetActionCheckState::CreateStatic(&FVREditorActionCallbacks::GetScaleSnapState)
 			),
 		NAME_None,
-		EUserInterfaceActionType::Check
+		EUserInterfaceActionType::ToggleButton
 		);
 	TAttribute<FText> DynamicScaleSizeLabel;
 	DynamicScaleSizeLabel.BindStatic(&FVREditorActionCallbacks::GetScaleSnapSizeText);
@@ -356,7 +356,7 @@ void UVRRadialMenuHandler::SnapMenuGenerator(FMenuBuilder MenuBuilder, TSharedPt
 			FGetActionCheckState::CreateStatic(&FVREditorActionCallbacks::AreAligningToActors, VRMode)
 			),
 		NAME_None,
-		EUserInterfaceActionType::Check
+		EUserInterfaceActionType::ToggleButton
 		);
 
 	TAttribute<FText> DynamicAlignSelectionLabel;
@@ -388,7 +388,7 @@ void UVRRadialMenuHandler::GizmoMenuGenerator(FMenuBuilder MenuBuilder, TSharedP
 			FGetActionCheckState::CreateStatic(&FVREditorActionCallbacks::IsActiveCoordinateSystem, VRMode, ECoordSystem::COORD_Local)
 			),
 		NAME_None,
-		EUserInterfaceActionType::Check
+		EUserInterfaceActionType::ToggleButton
 		);
 	MenuBuilder.AddMenuEntry(
 		LOCTEXT("WorldSpace", "World Space"),
@@ -401,7 +401,7 @@ void UVRRadialMenuHandler::GizmoMenuGenerator(FMenuBuilder MenuBuilder, TSharedP
 			FGetActionCheckState::CreateStatic(&FVREditorActionCallbacks::IsActiveCoordinateSystem, VRMode, ECoordSystem::COORD_World)
 		),
 		NAME_None,
-		EUserInterfaceActionType::Check
+		EUserInterfaceActionType::ToggleButton
 		);
 	MenuBuilder.AddMenuEntry(
 		LOCTEXT("Universal", "Universal"),
@@ -414,7 +414,7 @@ void UVRRadialMenuHandler::GizmoMenuGenerator(FMenuBuilder MenuBuilder, TSharedP
 			FGetActionCheckState::CreateStatic(&FVREditorActionCallbacks::IsActiveGizmoMode, VRMode, EGizmoHandleTypes::All)
 			),
 		NAME_None,
-		EUserInterfaceActionType::Check
+		EUserInterfaceActionType::ToggleButton
 		);
 	MenuBuilder.AddMenuEntry(
 		LOCTEXT("Translate", "Translate"),
@@ -427,7 +427,7 @@ void UVRRadialMenuHandler::GizmoMenuGenerator(FMenuBuilder MenuBuilder, TSharedP
 			FGetActionCheckState::CreateStatic(&FVREditorActionCallbacks::IsActiveGizmoMode, VRMode, EGizmoHandleTypes::Translate)
 			),
 		NAME_None,
-		EUserInterfaceActionType::Check
+		EUserInterfaceActionType::ToggleButton
 		);
 	MenuBuilder.AddMenuEntry(
 		LOCTEXT("Rotate", "Rotate"),
@@ -440,7 +440,7 @@ void UVRRadialMenuHandler::GizmoMenuGenerator(FMenuBuilder MenuBuilder, TSharedP
 			FGetActionCheckState::CreateStatic(&FVREditorActionCallbacks::IsActiveGizmoMode, VRMode, EGizmoHandleTypes::Rotate)
 			),
 		NAME_None,
-		EUserInterfaceActionType::Check
+		EUserInterfaceActionType::ToggleButton
 		);
 	MenuBuilder.AddMenuEntry(
 		LOCTEXT("Scale", "Scale"),
@@ -453,7 +453,7 @@ void UVRRadialMenuHandler::GizmoMenuGenerator(FMenuBuilder MenuBuilder, TSharedP
 			FGetActionCheckState::CreateStatic(&FVREditorActionCallbacks::IsActiveGizmoMode, VRMode, EGizmoHandleTypes::Scale)
 			),
 		NAME_None,
-		EUserInterfaceActionType::Check
+		EUserInterfaceActionType::ToggleButton
 		);
 }
 
@@ -470,7 +470,7 @@ void UVRRadialMenuHandler::UIMenuGenerator(FMenuBuilder MenuBuilder, TSharedPtr<
 			FGetActionCheckState::CreateStatic(&FVREditorActionCallbacks::GetUIToggledState, VRMode, UVREditorUISystem::EEditorUIPanel::ActorDetails)
 			),
 		NAME_None,
-		EUserInterfaceActionType::Check
+		EUserInterfaceActionType::ToggleButton
 		);
 	MenuBuilder.AddMenuEntry(
 		LOCTEXT("ContentBrowser", "Content Browser"),
@@ -483,7 +483,7 @@ void UVRRadialMenuHandler::UIMenuGenerator(FMenuBuilder MenuBuilder, TSharedPtr<
 			FGetActionCheckState::CreateStatic(&FVREditorActionCallbacks::GetUIToggledState, VRMode, UVREditorUISystem::EEditorUIPanel::ContentBrowser)
 			),
 		NAME_None,
-		EUserInterfaceActionType::Check
+		EUserInterfaceActionType::ToggleButton
 		);
 	MenuBuilder.AddMenuEntry(
 		LOCTEXT("Modes", "Modes Panel"),
@@ -496,7 +496,7 @@ void UVRRadialMenuHandler::UIMenuGenerator(FMenuBuilder MenuBuilder, TSharedPtr<
 			FGetActionCheckState::CreateStatic(&FVREditorActionCallbacks::GetUIToggledState, VRMode, UVREditorUISystem::EEditorUIPanel::Modes)
 			),
 		NAME_None,
-		EUserInterfaceActionType::Check
+		EUserInterfaceActionType::ToggleButton
 		);
 	MenuBuilder.AddMenuEntry(
 		LOCTEXT("WorldOutliner", "World Outliner"),
@@ -509,7 +509,7 @@ void UVRRadialMenuHandler::UIMenuGenerator(FMenuBuilder MenuBuilder, TSharedPtr<
 			FGetActionCheckState::CreateStatic(&FVREditorActionCallbacks::GetUIToggledState, VRMode, UVREditorUISystem::EEditorUIPanel::WorldOutliner)
 			),
 		NAME_None,
-		EUserInterfaceActionType::Check
+		EUserInterfaceActionType::ToggleButton
 		);
 	MenuBuilder.AddMenuEntry(
 		LOCTEXT("WorldSettings", "World Settings"),
@@ -522,7 +522,7 @@ void UVRRadialMenuHandler::UIMenuGenerator(FMenuBuilder MenuBuilder, TSharedPtr<
 			FGetActionCheckState::CreateStatic(&FVREditorActionCallbacks::GetUIToggledState, VRMode, UVREditorUISystem::EEditorUIPanel::WorldSettings)
 			),
 		NAME_None,
-		EUserInterfaceActionType::Check
+		EUserInterfaceActionType::ToggleButton
 		);
 	MenuBuilder.AddMenuEntry(
 		LOCTEXT("CreateNewSequence", "Create Sequence"),
@@ -726,7 +726,7 @@ void UVRRadialMenuHandler::ModesMenuGenerator(FMenuBuilder MenuBuilder, TSharedP
 			FGetActionCheckState::CreateStatic(&FVREditorActionCallbacks::EditorModeActive, FBuiltinEditorModes::EM_Placement)
 			),
 		NAME_None,
-		EUserInterfaceActionType::Check
+		EUserInterfaceActionType::ToggleButton
 		);
 	MenuBuilder.AddMenuEntry(
 		LOCTEXT("Foliage", "Foliage"),
@@ -739,7 +739,7 @@ void UVRRadialMenuHandler::ModesMenuGenerator(FMenuBuilder MenuBuilder, TSharedP
 			FGetActionCheckState::CreateStatic(&FVREditorActionCallbacks::EditorModeActive, FBuiltinEditorModes::EM_Foliage)
 			),
 		NAME_None,
-		EUserInterfaceActionType::Check
+		EUserInterfaceActionType::ToggleButton
 		);
 	MenuBuilder.AddMenuEntry(
 		LOCTEXT("Landscape", "Landscape"),
@@ -752,7 +752,7 @@ void UVRRadialMenuHandler::ModesMenuGenerator(FMenuBuilder MenuBuilder, TSharedP
 			FGetActionCheckState::CreateStatic(&FVREditorActionCallbacks::EditorModeActive, FBuiltinEditorModes::EM_Landscape)
 			),
 		NAME_None,
-		EUserInterfaceActionType::Check
+		EUserInterfaceActionType::ToggleButton
 		);
 	MenuBuilder.AddMenuEntry(
 		LOCTEXT("MeshPaint", "Paint"),
@@ -765,7 +765,7 @@ void UVRRadialMenuHandler::ModesMenuGenerator(FMenuBuilder MenuBuilder, TSharedP
 			FGetActionCheckState::CreateStatic(&FVREditorActionCallbacks::EditorModeActive, FBuiltinEditorModes::EM_MeshPaint)
 			),
 		NAME_None,
-		EUserInterfaceActionType::Check
+		EUserInterfaceActionType::ToggleButton
 		);
 	MenuBuilder.AddMenuEntry(
 		LOCTEXT("Mesh", "Mesh"),
@@ -778,7 +778,7 @@ void UVRRadialMenuHandler::ModesMenuGenerator(FMenuBuilder MenuBuilder, TSharedP
 			FGetActionCheckState::CreateStatic(&FVREditorActionCallbacks::MeshEditModeActive, EEditableMeshElementType::Any)
 			),
 		NAME_None,
-		EUserInterfaceActionType::Check
+		EUserInterfaceActionType::ToggleButton
 		);
 	MenuBuilder.AddMenuEntry(
 		LOCTEXT("Polygon", "Polygon"),
@@ -791,7 +791,7 @@ void UVRRadialMenuHandler::ModesMenuGenerator(FMenuBuilder MenuBuilder, TSharedP
 			FGetActionCheckState::CreateStatic(&FVREditorActionCallbacks::MeshEditModeActive, EEditableMeshElementType::Polygon)
 			),
 		NAME_None,
-		EUserInterfaceActionType::Check
+		EUserInterfaceActionType::ToggleButton
 		);
 	MenuBuilder.AddMenuEntry(
 		LOCTEXT("Edge", "Edge"),
@@ -804,7 +804,7 @@ void UVRRadialMenuHandler::ModesMenuGenerator(FMenuBuilder MenuBuilder, TSharedP
 			FGetActionCheckState::CreateStatic(&FVREditorActionCallbacks::MeshEditModeActive, EEditableMeshElementType::Edge)
 			),
 		NAME_None,
-		EUserInterfaceActionType::Check
+		EUserInterfaceActionType::ToggleButton
 		);
 	MenuBuilder.AddMenuEntry(
 		LOCTEXT("Vertex", "Vertex"),
@@ -817,7 +817,7 @@ void UVRRadialMenuHandler::ModesMenuGenerator(FMenuBuilder MenuBuilder, TSharedP
 			FGetActionCheckState::CreateStatic(&FVREditorActionCallbacks::MeshEditModeActive, EEditableMeshElementType::Vertex)
 			),
 		NAME_None,
-		EUserInterfaceActionType::Check
+		EUserInterfaceActionType::ToggleButton
 		);
 	
 }
diff --git a/Engine/Source/Editor/VREditor/VREditorAvatarActor.cpp b/Engine/Source/Editor/VREditor/VREditorAvatarActor.cpp
--- a/Engine/Source/Editor/VREditor/VREditorAvatarActor.cpp
+++ b/Engine/Source/Editor/VREditor/VREditorAvatarActor.cpp
@@ -20,6 +20,7 @@
 #include "EngineGlobals.h"
 #include "IHeadMountedDisplay.h"
 #include "VREditorAssetContainer.h"
+#include "VRModeSettings.h"
 
 namespace VREd
 {
@@ -37,10 +38,6 @@ namespace VREd
 
 	static FAutoConsoleVariable ScaleProgressBarLength( TEXT( "VREd.ScaleProgressBarLength" ), 50.0f, TEXT( "Length of the progressbar that appears when scaling" ) );
 	static FAutoConsoleVariable ScaleProgressBarRadius( TEXT( "VREd.ScaleProgressBarRadius" ), 1.0f, TEXT( "Radius of the progressbar that appears when scaling" ) );
-	
-	static FAutoConsoleVariable ShowWorldScaleProgressBar( TEXT( "VREd.ShowWorldScaleProgressBar" ), 0, TEXT( "Whether we show a visual progress bar when scaling the world" ) );
-	static FAutoConsoleVariable ShowMovementGrid( TEXT( "VREd.ShowMovementGrid" ), 0, TEXT( "Showing the ground movement grid" ) );
-	static FAutoConsoleVariable ShowWorldMovementPostProcess( TEXT( "VREd.ShowWorldMovementPostProcess" ), 0, TEXT( "Showing the movement post processing" ) );
 }
 
 AVREditorAvatarActor::AVREditorAvatarActor( const FObjectInitializer& ObjectInitializer ) :
@@ -212,7 +209,7 @@ AVREditorAvatarActor::AVREditorAvatarActor( const FObjectInitializer& ObjectInit
 		PostProcessComponent->SetupAttachment( this->GetRootComponent( ) );
 
 		// Unlimited size
-		PostProcessComponent->bEnabled = VREd::ShowWorldMovementPostProcess->GetInt( ) == 1 ? true : false;
+		PostProcessComponent->bEnabled = GetDefault<UVRModeSettings>()->bShowWorldMovementPostProcess;
 		PostProcessComponent->bUnbound = true;
 	}
 
@@ -292,7 +289,7 @@ void AVREditorAvatarActor::TickManually( const float DeltaTime )
 		UVREditorInteractor* LeftHandInteractor = VRMode->GetHandInteractor( EControllerHand::Left );
 		UVREditorInteractor* RightHandInteractor = VRMode->GetHandInteractor( EControllerHand::Right );
 
-		if ( VREd::ShowWorldScaleProgressBar->GetInt() != 0 &&
+		if (GetDefault<UVRModeSettings>()->bShowWorldScaleProgressBar &&
 			LeftHandInteractor != nullptr && RightHandInteractor != nullptr &&
 			 ( ( LeftHandInteractor->GetDraggingMode() == EViewportInteractionDraggingMode::World && RightHandInteractor->GetDraggingMode() == EViewportInteractionDraggingMode::AssistingDrag ) ||
 			   ( LeftHandInteractor->GetDraggingMode() == EViewportInteractionDraggingMode::AssistingDrag && RightHandInteractor->GetDraggingMode() == EViewportInteractionDraggingMode::World ) ) )
@@ -364,7 +361,7 @@ void AVREditorAvatarActor::TickManually( const float DeltaTime )
 
 	// Updating the opacity and visibility of the grid according to the controllers //@todo
 	{
-		if ( VREd::ShowMovementGrid->GetInt( ) == 1 )
+		if (GetDefault<UVRModeSettings>()->bShowWorldMovementGrid)
 		{
 			UVREditorInteractor* LeftHandInteractor = VRMode->GetHandInteractor( EControllerHand::Left );
 			UVREditorInteractor* RightHandInteractor = VRMode->GetHandInteractor( EControllerHand::Right );
@@ -428,7 +425,7 @@ void AVREditorAvatarActor::TickManually( const float DeltaTime )
 	// that are not nearby, and completely hides the skybox and other very distant objects.  This is to help
 	// prevent simulation sickness when moving/rotating/scaling the entire world around them.
 	{
-		PostProcessComponent->bEnabled = VREd::ShowWorldMovementPostProcess->GetInt( ) == 1 ? true : false;
+		PostProcessComponent->bEnabled = GetDefault<UVRModeSettings>()->bShowWorldMovementPostProcess;
 
 		// Make sure our world size is reflected in the post process material
 		static FName WorldScaleFactorParameterName( "WorldScaleFactor" );
diff --git a/Engine/Source/Editor/VREditor/VREditorMode.cpp b/Engine/Source/Editor/VREditor/VREditorMode.cpp
--- a/Engine/Source/Editor/VREditor/VREditorMode.cpp
+++ b/Engine/Source/Editor/VREditor/VREditorMode.cpp
@@ -38,13 +38,15 @@
 #include "SequencerSettings.h"
 #include "Kismet/GameplayStatics.h"
 #include "Components/AudioComponent.h"
+#include "VREditorActions.h"
+#include "EditorModes.h"
+#include "VRModeSettings.h"
 
 #define LOCTEXT_NAMESPACE "VREditorMode"
 
 namespace VREd
 {
 	static FAutoConsoleVariable DefaultVRNearClipPlane(TEXT("VREd.DefaultVRNearClipPlane"), 1.0f, TEXT("The near clip plane to use for VR"));
-	static FAutoConsoleVariable GizmoScaleInVR( TEXT( "VREd.GizmoScaleInVR" ), 0.8f, TEXT( "Size of transform gizmo while in VR" ) );
 	static FAutoConsoleVariable SlateDragDistanceOverride( TEXT( "VREd.SlateDragDistanceOverride" ), 40.0f, TEXT( "How many pixels you need to drag before a drag and drop operation starts in VR" ) );
 	static FAutoConsoleVariable DefaultWorldToMeters(TEXT("VREd.DefaultWorldToMeters"), 100.0f, TEXT("Default world to meters scale"));
 
@@ -198,12 +200,17 @@ void UVREditorMode::Enter()
 
 			// When actually in VR, make sure the transform gizmo is big!
 			SavedEditorState.TransformGizmoScale = WorldInteraction->GetTransformGizmoScale();
-			WorldInteraction->SetTransformGizmoScale( VREd::GizmoScaleInVR->GetFloat() );
+			WorldInteraction->SetTransformGizmoScale(GetDefault<UVRModeSettings>()->GizmoScale);
 			WorldInteraction->SetShouldSuppressExistingCursor(true);
 			WorldInteraction->SetInVR(true);
 		}
 	}
 
+	// Switch us back to placement mode and close any open sequencer windows
+	FVREditorActionCallbacks::ChangeEditorModes(FBuiltinEditorModes::EM_Placement);
+	FLevelEditorModule& LevelEditorModule = FModuleManager::GetModuleChecked<FLevelEditorModule>(TEXT("LevelEditor"));
+	LevelEditorModule.GetLevelEditorTabManager()->InvokeTab(FTabId("Sequencer"))->RequestCloseTab();
+
 	// Setup our avatar
 	if (AvatarActor == nullptr)
 	{
@@ -283,6 +290,8 @@ void UVREditorMode::Enter()
 void UVREditorMode::Exit(const bool bShouldDisableStereo)
 {
 	{
+		FVREditorActionCallbacks::ChangeEditorModes(FBuiltinEditorModes::EM_Placement);
+
 		//Destroy the avatar
 		{
 			DestroyTransientActor(AvatarActor);
@@ -669,7 +678,7 @@ void UVREditorMode::RefreshVREditorSequencer(class ISequencer* InCurrentSequence
 {
 	CurrentSequencer = InCurrentSequencer;
 	// Tell the VR Editor UI system to refresh the Sequencer UI
-	if (bActuallyUsingVR)
+	if (bActuallyUsingVR && InCurrentSequencer != nullptr)
 	{
 		GetUISystem().UpdateSequencerUI();
 	}
diff --git a/Engine/Source/Editor/VREditor/VREditorModeManager.cpp b/Engine/Source/Editor/VREditor/VREditorModeManager.cpp
--- a/Engine/Source/Editor/VREditor/VREditorModeManager.cpp
+++ b/Engine/Source/Editor/VREditor/VREditorModeManager.cpp
@@ -3,7 +3,6 @@
 #include "VREditorModeManager.h"
 #include "InputCoreTypes.h"
 #include "VREditorMode.h"
-#include "Settings/EditorExperimentalSettings.h"
 #include "Modules/ModuleManager.h"
 #include "GameFramework/PlayerController.h"
 #include "GameFramework/WorldSettings.h"
@@ -14,6 +13,7 @@
 #include "IHeadMountedDisplay.h"
 #include "EditorWorldExtension.h"
 #include "ViewportWorldInteraction.h"
+#include "VRModeSettings.h"
 
 FVREditorModeManager::FVREditorModeManager() :
 	CurrentVREditorMode( nullptr ),
@@ -36,7 +36,7 @@ void FVREditorModeManager::Tick( const float DeltaTime )
 	TimeSinceHMDChecked += DeltaTime;
 
 	// You can only auto-enter VR if the setting is enabled. Other criteria are that the VR Editor is enabled in experimental settings, that you are not in PIE, and that the editor is foreground.
-	bool bCanAutoEnterVR = GetDefault<UEditorExperimentalSettings>()->bEnableAutoVREditMode && 
+	bool bCanAutoEnterVR = GetDefault<UVRModeSettings>()->bEnableAutoVREditMode && 
 		(GEditor->PlayWorld == nullptr || (CurrentVREditorMode != nullptr && CurrentVREditorMode->GetStartedPlayFromVREditor())) && 
 		FPlatformProcess::IsThisApplicationForeground();
 
diff --git a/Engine/Source/Editor/VREditor/VREditorMotionControllerInteractor.cpp b/Engine/Source/Editor/VREditor/VREditorMotionControllerInteractor.cpp
--- a/Engine/Source/Editor/VREditor/VREditorMotionControllerInteractor.cpp
+++ b/Engine/Source/Editor/VREditor/VREditorMotionControllerInteractor.cpp
@@ -33,6 +33,7 @@
 #include "DrawDebugHelpers.h"
 #include "VREditorActions.h"
 #include "VREditorAssetContainer.h"
+#include "VRModeSettings.h"
 
 namespace VREd
 {
@@ -50,9 +51,7 @@ namespace VREd
 	static FAutoConsoleVariable LaserRadiusScaleWhenOverUI( TEXT( "VREd.LaserRadiusScaleWhenOverUI" ), 0.25f, TEXT( "How much to scale down the size of the laser pointer radius when over UI" ) );
 	static FAutoConsoleVariable HoverBallRadiusScaleWhenOverUI(TEXT("VREd.HoverBallRadiusScaleWhenOverUI"), 0.4f, TEXT("How much to scale down the size of the hover ball when over UI"));
 	//Trigger
-	static FAutoConsoleVariable TriggerPressedThreshold_Vive( TEXT( "VI.TriggerPressedThreshold_Vive" ), 0.03f, TEXT( "Minimum trigger threshold before we consider the trigger at least 'lightly pressed'" ) );
 	static FAutoConsoleVariable TriggerDeadZone_Vive( TEXT( "VI.TriggerDeadZone_Vive" ), 0.01f, TEXT( "Trigger dead zone.  The trigger must be fully released before we'll trigger a new 'light press'" ) );
-	static FAutoConsoleVariable TriggerPressedThreshold_Rift( TEXT( "VI.TriggerPressedThreshold_Rift" ), 0.2f, TEXT( "Minimum trigger threshold before we consider the trigger at least 'lightly pressed'" ) );
 	static FAutoConsoleVariable TriggerDeadZone_Rift( TEXT( "VI.TriggerDeadZone_Rift" ), 0.15f, TEXT( "Trigger dead zone.  The trigger must be fully released before we'll trigger a new 'light press'" ) );
 	static FAutoConsoleVariable TriggerFullyPressedThreshold_Vive( TEXT( "VI.TriggerFullyPressedThreshold_Vive" ), 0.90f, TEXT( "Minimum trigger threshold before we consider the trigger 'fully pressed'" ) );
 	static FAutoConsoleVariable TriggerFullyPressedThreshold_Rift( TEXT( "VI.TriggerFullyPressedThreshold_Rift" ), 0.99f, TEXT( "Minimum trigger threshold before we consider the trigger 'fully pressed'" ) );
@@ -759,7 +758,7 @@ void UVREditorMotionControllerInteractor::HandleInputAxis( FEditorViewportClient
 {
 	if ( Action.ActionType == TriggerAxis )
 	{
-		const float TriggerPressedThreshold = ( GetHMDDeviceType() == EHMDDeviceType::DT_OculusRift ) ? VREd::TriggerPressedThreshold_Rift->GetFloat() : VREd::TriggerPressedThreshold_Vive->GetFloat();
+		const float TriggerPressedThreshold = ( GetHMDDeviceType() == EHMDDeviceType::DT_OculusRift ) ? GetDefault<UVRModeSettings>()->TriggerPressedThreshold_Rift : GetDefault<UVRModeSettings>()->TriggerPressedThreshold_Vive;
 		const float TriggerDeadZone = ( GetHMDDeviceType() == EHMDDeviceType::DT_OculusRift ) ? VREd::TriggerDeadZone_Rift->GetFloat() : VREd::TriggerDeadZone_Vive->GetFloat();
 
 		// Synthesize "lightly pressed" events for the trigger
diff --git a/Engine/Source/Editor/VREditor/VREditorStyle.cpp b/Engine/Source/Editor/VREditor/VREditorStyle.cpp
--- a/Engine/Source/Editor/VREditor/VREditorStyle.cpp
+++ b/Engine/Source/Editor/VREditor/VREditorStyle.cpp
@@ -146,6 +146,15 @@ TSharedRef< FSlateStyleSet > FVREditorStyle::Create()
 	Style->Set("VRRadialStyle.Label", FTextBlockStyle(NormalText)
 		.SetFont(RadialFont)
 		.SetColorAndOpacity(FLinearColor::White));
+
+	Style->Set("VRRadialStyle.InactiveFont", FSlateFontInfo(RadialFont));
+
+	FSlateFontInfo ActiveRadialFont = TTF_CORE_FONT("Fonts/Roboto-Regular", 10);
+	FFontOutlineSettings ActiveRadialOutline;
+	ActiveRadialOutline.OutlineColor = FLinearColor::Black;
+	ActiveRadialOutline.OutlineSize = 1.0f;
+	ActiveRadialFont.OutlineSettings = ActiveRadialOutline;
+	Style->Set("VRRadialStyle.ActiveFont", FSlateFontInfo(ActiveRadialFont));
 	
 	FSlateFontInfo NumpadRadialFont = TTF_CORE_FONT("Fonts/Roboto-Regular", 24);
 	NumpadRadialFont.OutlineSettings = RadialOutline;
diff --git a/Engine/Source/Editor/VREditor/VREditorWorldInteraction.cpp b/Engine/Source/Editor/VREditor/VREditorWorldInteraction.cpp
--- a/Engine/Source/Editor/VREditor/VREditorWorldInteraction.cpp
+++ b/Engine/Source/Editor/VREditor/VREditorWorldInteraction.cpp
@@ -21,6 +21,7 @@
 #include "VREditorUISystem.h"
 #include "VREditorFloatingUI.h"
 #include "VREditorInteractor.h"
+#include "VREditorMotionControllerInteractor.h"
 
 // For actor placement
 #include "ObjectTools.h"
@@ -79,17 +80,16 @@ void UVREditorWorldInteraction::StopDragging( UViewportInteractor* Interactor )
 {
 	EViewportInteractionDraggingMode InteractorDraggingMode = Interactor->GetDraggingMode();
 
-	// If we were placing a material, go ahead and do that now
 	if ( InteractorDraggingMode == EViewportInteractionDraggingMode::Material )
 	{
+		// If we were placing a material, go ahead and do that now
 		PlaceDraggedMaterialOrTexture( Interactor );
 	}	
-
-	// If we were placing something, bring the window back
-	if ( FloatingUIAssetDraggedFrom != nullptr && 
+	else if ( FloatingUIAssetDraggedFrom != nullptr && 
 		( InteractorDraggingMode == EViewportInteractionDraggingMode::TransformablesAtLaserImpact ||
 		  InteractorDraggingMode == EViewportInteractionDraggingMode::Material ) )
 	{
+		// If we were placing something, bring the window back
 		const bool bShouldShow = true;
 		const bool bSpawnInFront = false;
 		const bool bDragFromOpen = false;
@@ -97,6 +97,14 @@ void UVREditorWorldInteraction::StopDragging( UViewportInteractor* Interactor )
 		Owner->GetUISystem().ShowEditorUIPanel(FloatingUIAssetDraggedFrom, Cast<UVREditorInteractor>(Interactor->GetOtherInteractor()), bShouldShow, bSpawnInFront, bDragFromOpen, bPlaySound);
 		FloatingUIAssetDraggedFrom = nullptr;
 	}
+	else if (Interactor->GetDraggingMode() == EViewportInteractionDraggingMode::TransformablesFreely)
+	{
+		UVREditorMotionControllerInteractor* MotionController = Cast<UVREditorMotionControllerInteractor>(Interactor);
+		if (Owner->GetUISystem().GetUIInteractor() == MotionController)
+		{
+			MotionController->SetControllerType(EControllerType::UI);
+		}
+	}
 }
 
 void UVREditorWorldInteraction::UpdateNearClipPlaneOnScaleChange(const float NewWorldToMetersScale)
diff --git a/Engine/Source/Editor/VREditor/VRModeSettings.cpp b/Engine/Source/Editor/VREditor/VRModeSettings.cpp
new file mode 100644
--- /dev/null
+++ b/Engine/Source/Editor/VREditor/VRModeSettings.cpp
@@ -0,0 +1,18 @@
+// Copyright 1998-2017 Epic Games, Inc. All Rights Reserved.
+
+#include "VRModeSettings.h"
+
+UVRModeSettings::UVRModeSettings(const FObjectInitializer& ObjectInitializer)
+	: Super(ObjectInitializer)
+{
+	bEnableAutoVREditMode = false;
+	InteractorHand = EInteractorHand::Right;
+	bShowWorldMovementGrid = true;
+	bShowWorldMovementPostProcess = true;
+	bShowWorldScaleProgressBar = true;
+	UIBrightness = 1.5f;
+	GizmoScale = 0.8f;
+	DoubleClickTime = 0.25f;
+	TriggerPressedThreshold_Vive = 0.03f;
+	TriggerPressedThreshold_Rift = 0.2f;
+}
diff --git a/Engine/Source/Editor/ViewportInteraction/Gizmo/VIPivotTransformGizmo.cpp b/Engine/Source/Editor/ViewportInteraction/Gizmo/VIPivotTransformGizmo.cpp
--- a/Engine/Source/Editor/ViewportInteraction/Gizmo/VIPivotTransformGizmo.cpp
+++ b/Engine/Source/Editor/ViewportInteraction/Gizmo/VIPivotTransformGizmo.cpp
@@ -34,7 +34,8 @@ namespace VREd //@todo VREditor: Duplicates of TransformGizmo
 
 APivotTransformGizmo::APivotTransformGizmo() :
 	Super(),
-	AimingAtMeFadeAlpha(0.0f)
+	AimingAtGizmoScaleAlpha(0.0f),
+	LastDraggingHandle(nullptr)
 {
 	UniformScaleGizmoHandleGroup = CreateDefaultSubobject<UUniformScaleGizmoHandleGroup>( TEXT( "UniformScaleHandles" ), true );
 	UniformScaleGizmoHandleGroup->SetOwningTransformGizmo(this);
@@ -97,40 +98,55 @@ void APivotTransformGizmo::UpdateGizmo(const EGizmoHandleTypes InGizmoType, cons
 	float GizmoScale = (InScaleMultiplier * ((WorldSpaceDistanceToToPivot / WorldScaleFactor) * VREd::PivotGizmoDistanceScaleFactor->GetFloat())) * WorldScaleFactor;
 	const bool bIsWorldSpaceGizmo = (WorldInteraction->GetTransformGizmoCoordinateSpace() == COORD_World);
 
+	if (LastDraggingHandle != nullptr && DraggingHandle == nullptr)
+	{
+		AimingAtGizmoScaleAlpha = 0.0f;
+	}
+
 	// Only scale the gizmo down when not aiming at it for VR implementations
 	if (WorldInteraction->IsInVR())
 	{
 		bool bIsAimingTowards = false;
-		const float GizmoRadius = (GizmoScale * 350) * 0.5f; //@todo ViewportInteraction: Hardcoded radius multiplier.
+		bool bAnyInteractorInCorrectState = false;
+		const float GizmoRadius = (GizmoScale * 350) * 0.5f;
 
 		// Check if any interactor has a laser close to the gizmo
 		for (UViewportInteractor* Interactor : WorldInteraction->GetInteractors())
 		{	
-			FVector LaserPointerStart, LaserPointerEnd;
-			if (Interactor->GetLaserPointer(/*Out*/ LaserPointerStart, /*Out*/ LaserPointerEnd))
+			// We only want the interactor to affect the size when aiming at the gizmo if they are dragging nothing
+			if (Interactor->GetDraggingMode() == EViewportInteractionDraggingMode::Nothing)
 			{
-				const FVector ClosestPointOnLaser = FMath::ClosestPointOnLine(LaserPointerStart, LaserPointerEnd, InLocalToWorld.GetLocation());
-				const float ClosestPointDistance = (ClosestPointOnLaser - InLocalToWorld.GetLocation()).Size();
-				if (ClosestPointDistance <= GizmoRadius)
+				bAnyInteractorInCorrectState = true;			
+
+				FVector LaserPointerStart, LaserPointerEnd;
+				if (Interactor->GetLaserPointer(/*Out*/ LaserPointerStart, /*Out*/ LaserPointerEnd))
 				{
-					bIsAimingTowards = true;
-					break;
+					const FVector ClosestPointOnLaser = FMath::ClosestPointOnLine(LaserPointerStart, LaserPointerEnd, InLocalToWorld.GetLocation());
+					const float ClosestPointDistance = (ClosestPointOnLaser - InLocalToWorld.GetLocation()).Size();
+					if (ClosestPointDistance <= GizmoRadius)
+					{
+						bIsAimingTowards = true;
+						break;
+					}
 				}
 			}
 		}
 
-		const float DeltaTime = WorldInteraction->GetCurrentDeltaTime();
-		if (bIsAimingTowards)
+		if (bAnyInteractorInCorrectState)
 		{
-			AimingAtMeFadeAlpha += DeltaTime / VREd::PivotGizmoAimAtAnimationSpeed->GetFloat();
-		}
-		else
-		{
-			AimingAtMeFadeAlpha -= DeltaTime / VREd::PivotGizmoAimAtAnimationSpeed->GetFloat();
-		}
-		AimingAtMeFadeAlpha = FMath::Clamp(AimingAtMeFadeAlpha, VREd::PivotGizmoAimAtShrinkSize->GetFloat(), 1.0f);
+			const float DeltaTime = WorldInteraction->GetCurrentDeltaTime();
+			if (bIsAimingTowards)
+			{
+				AimingAtGizmoScaleAlpha += DeltaTime / VREd::PivotGizmoAimAtAnimationSpeed->GetFloat();
+			}
+			else
+			{
+				AimingAtGizmoScaleAlpha -= DeltaTime / VREd::PivotGizmoAimAtAnimationSpeed->GetFloat();
+			}
+			AimingAtGizmoScaleAlpha = FMath::Clamp(AimingAtGizmoScaleAlpha, VREd::PivotGizmoAimAtShrinkSize->GetFloat(), 1.0f);
 
-		GizmoScale *= AimingAtMeFadeAlpha;
+			GizmoScale *= AimingAtGizmoScaleAlpha;
+		}
 	}
 
 	// Update animation
@@ -152,6 +168,8 @@ void APivotTransformGizmo::UpdateGizmo(const EGizmoHandleTypes InGizmoType, cons
 			}
 		}
 	}
+
+	LastDraggingHandle = DraggingHandle;
 }
 
 /************************************************************************/
diff --git a/Engine/Source/Editor/ViewportInteraction/Gizmo/VIPivotTransformGizmo.h b/Engine/Source/Editor/ViewportInteraction/Gizmo/VIPivotTransformGizmo.h
--- a/Engine/Source/Editor/ViewportInteraction/Gizmo/VIPivotTransformGizmo.h
+++ b/Engine/Source/Editor/ViewportInteraction/Gizmo/VIPivotTransformGizmo.h
@@ -20,7 +20,7 @@ public:
 
 	/** Default constructor that sets up CDO properties */
 	APivotTransformGizmo();
-
+	
 	/** Called by the world interaction system after we've been spawned into the world, to allow
 	    us to create components and set everything up nicely for the selected objects that we'll be
 		used to manipulate */
@@ -53,7 +53,12 @@ private:
 	UPROPERTY()
 	class UStretchGizmoHandleGroup* StretchGizmoHandleGroup;
 
-	float AimingAtMeFadeAlpha;
+	/** The alpha for gizmo animation when aiming at it with a laser */
+	float AimingAtGizmoScaleAlpha;
+
+	/** Handle from previous tick that was dragged */
+	UPROPERTY()
+	UActorComponent* LastDraggingHandle;
 };
 
 /**
