Updated VR Mode button to become inactive during SIE (instead of disappearing altogether)

diff --git a/Engine/Source/Editor/LevelEditor/Private/LevelEditor.cpp b/Engine/Source/Editor/LevelEditor/Private/LevelEditor.cpp
--- a/Engine/Source/Editor/LevelEditor/Private/LevelEditor.cpp
+++ b/Engine/Source/Editor/LevelEditor/Private/LevelEditor.cpp
@@ -609,7 +609,7 @@ void FLevelEditorModule::BindGlobalLevelEditorCommands()
 	ActionList.MapAction( 
 		Commands.ToggleVR, 
 		FExecuteAction::CreateStatic( &FLevelEditorActionCallbacks::ToggleVR ), 
-		FCanExecuteAction::CreateStatic( &FLevelEditorActionCallbacks::ToggleVR_CanExecute ),
+		FCanExecuteAction::CreateStatic( &FLevelEditorActionCallbacks::ToggleVR_IsButtonActive ),
 		FIsActionChecked::CreateStatic( &FLevelEditorActionCallbacks::ToggleVR_IsChecked ) );
 
 	ActionList.MapAction(Commands.ImportScene,
diff --git a/Engine/Source/Editor/LevelEditor/Private/LevelEditorActions.cpp b/Engine/Source/Editor/LevelEditor/Private/LevelEditorActions.cpp
--- a/Engine/Source/Editor/LevelEditor/Private/LevelEditorActions.cpp
+++ b/Engine/Source/Editor/LevelEditor/Private/LevelEditorActions.cpp
@@ -2115,6 +2115,11 @@ bool FLevelEditorActionCallbacks::ToggleVR_CanExecute()
 	return VREditorModule.IsVREditorAvailable();
 }
 
+bool FLevelEditorActionCallbacks::ToggleVR_IsButtonActive()
+{
+	IVREditorModule& VREditorModule = IVREditorModule::Get();
+	return VREditorModule.IsVREditorButtonActive();
+}
 
 bool FLevelEditorActionCallbacks::ToggleVR_IsChecked()
 {
diff --git a/Engine/Source/Editor/LevelEditor/Private/SLevelEditor.cpp b/Engine/Source/Editor/LevelEditor/Private/SLevelEditor.cpp
--- a/Engine/Source/Editor/LevelEditor/Private/SLevelEditor.cpp
+++ b/Engine/Source/Editor/LevelEditor/Private/SLevelEditor.cpp
@@ -130,7 +130,7 @@ void SLevelEditor::BindCommands()
 		FExecuteAction::CreateStatic( &FLevelEditorActionCallbacks::ToggleVR ),
 		FCanExecuteAction::CreateStatic( &FLevelEditorActionCallbacks::ToggleVR_CanExecute ),
 		FIsActionChecked::CreateStatic( &FLevelEditorActionCallbacks::ToggleVR_IsChecked ),
-		FIsActionButtonVisible::CreateStatic(&FLevelEditorActionCallbacks::ToggleVR_CanExecute));
+		FIsActionButtonVisible::CreateStatic(&FLevelEditorActionCallbacks::ToggleVR_IsButtonActive));
 
 	LevelEditorCommands->MapAction(
 		Actions.WorldProperties,
diff --git a/Engine/Source/Editor/LevelEditor/Public/LevelEditorActions.h b/Engine/Source/Editor/LevelEditor/Public/LevelEditorActions.h
--- a/Engine/Source/Editor/LevelEditor/Public/LevelEditorActions.h
+++ b/Engine/Source/Editor/LevelEditor/Public/LevelEditorActions.h
@@ -654,6 +654,7 @@ public:
 	/** Toggles VR mode */
 	static void ToggleVR();
 	static bool ToggleVR_CanExecute();
+	static bool ToggleVR_IsButtonActive();
 	static bool ToggleVR_IsChecked();
 
 	/** Opens delta transform */
diff --git a/Engine/Source/Editor/VREditor/Public/IVREditorModule.h b/Engine/Source/Editor/VREditor/Public/IVREditorModule.h
--- a/Engine/Source/Editor/VREditor/Public/IVREditorModule.h
+++ b/Engine/Source/Editor/VREditor/Public/IVREditorModule.h
@@ -50,6 +50,13 @@ public:
 	 */
 	virtual bool IsVREditorAvailable() const = 0;
 
+	/**
+	* Checks to see if the VR Mode button should be active or grayed out (such as during SIE)
+	*
+	* @return	True if button should be active
+	*/
+	virtual bool IsVREditorButtonActive() const = 0;
+
 	/**
 	 * Enables or disables editor VR features.  Calling this to active VR will turn on the HMD and setup
 	 * the editor UI for VR interaction.
diff --git a/Engine/Source/Editor/VREditor/VREditorModeManager.cpp b/Engine/Source/Editor/VREditor/VREditorModeManager.cpp
--- a/Engine/Source/Editor/VREditor/VREditorModeManager.cpp
+++ b/Engine/Source/Editor/VREditor/VREditorModeManager.cpp
@@ -152,6 +152,12 @@ bool FVREditorModeManager::IsVREditorAvailable() const
 	return bHasHMDDevice && !GEditor->bIsSimulatingInEditor;
 }
 
+bool FVREditorModeManager::IsVREditorButtonActive() const
+{
+	const bool bHasHMDDevice = GEngine->XRSystem.IsValid() && GEngine->XRSystem->GetHMDDevice() && GEngine->XRSystem->GetHMDDevice()->IsHMDEnabled();
+	return bHasHMDDevice;
+}
+
 
 UVREditorMode* FVREditorModeManager::GetCurrentVREditorMode()
 {
diff --git a/Engine/Source/Editor/VREditor/VREditorModeManager.h b/Engine/Source/Editor/VREditor/VREditorModeManager.h
--- a/Engine/Source/Editor/VREditor/VREditorModeManager.h
+++ b/Engine/Source/Editor/VREditor/VREditorModeManager.h
@@ -42,6 +42,9 @@ public:
 	/** If the VR Editor is currently available */
 	bool IsVREditorAvailable() const;
 
+	/** If the VR Editor button active */
+	bool IsVREditorButtonActive() const;
+
 	/** Gets the current VR Editor mode that was enabled */
 	UVREditorMode* GetCurrentVREditorMode();
 
diff --git a/Engine/Source/Editor/VREditor/VREditorModule.cpp b/Engine/Source/Editor/VREditor/VREditorModule.cpp
--- a/Engine/Source/Editor/VREditor/VREditorModule.cpp
+++ b/Engine/Source/Editor/VREditor/VREditorModule.cpp
@@ -31,6 +31,7 @@ public:
 	// IVREditorModule overrides
 	virtual bool IsVREditorEnabled() const override;
 	virtual bool IsVREditorAvailable() const override;
+	virtual bool IsVREditorButtonActive() const override;
 	virtual void EnableVREditor( const bool bEnable, const bool bForceWithoutHMD ) override;
 	virtual bool IsVREditorModeActive() override;
 	virtual UVREditorMode* GetVRMode() override;
@@ -81,8 +82,12 @@ bool FVREditorModule::IsVREditorEnabled() const
 bool FVREditorModule::IsVREditorAvailable() const
 {
 	return ModeManager.IsVREditorAvailable();
-}	
+}
 
+bool FVREditorModule::IsVREditorButtonActive() const
+{
+	return ModeManager.IsVREditorButtonActive();
+}
 
 void FVREditorModule::EnableVREditor( const bool bEnable, const bool bForceWithoutHMD )
 {
