Preventing crash on floating text if asset container does not exist. VR Editor floating text is now not placeable or blueprintable.

diff --git a/Engine/Source/Editor/VREditor/VREditorFloatingText.cpp b/Engine/Source/Editor/VREditor/VREditorFloatingText.cpp
--- a/Engine/Source/Editor/VREditor/VREditorFloatingText.cpp
+++ b/Engine/Source/Editor/VREditor/VREditorFloatingText.cpp
@@ -94,13 +94,8 @@ AFloatingText::AFloatingText()
 	}
 
 	LineMaterial = AssetContainer.LineMaterial;
-	check(LineMaterial != nullptr);
-
 	MaskedTextMaterial = AssetContainer.TextMaterial;
-	check(MaskedTextMaterial != nullptr);
-
 	TranslucentTextMaterial = AssetContainer.TranslucentTextMaterial;
-	check(TranslucentTextMaterial != nullptr);
 
 	{
 		TextComponent = CreateDefaultSubobject<UTextRenderComponent>( TEXT( "Text" ) );
@@ -124,9 +119,11 @@ AFloatingText::AFloatingText()
 		// Use a custom font.  The text will be visible up close.	   
 		TextComponent->SetFont( AssetContainer.TextFont );
 
-		// Assign our custom text rendering material.
-		TextComponent->SetTextMaterial( MaskedTextMaterial );
-
+		if (MaskedTextMaterial != nullptr)
+		{
+			// Assign our custom text rendering material.
+			TextComponent->SetTextMaterial(MaskedTextMaterial);
+		}
 		TextComponent->SetTextRenderColor( FLinearColor::White.ToFColor( false ) );
 
 		// Left justify the text
@@ -141,12 +138,14 @@ void AFloatingText::PostActorCreated()
 	Super::PostActorCreated();
 
 	// Create an MID so that we can change parameters on the fly (fading)
-	check( LineMaterial != nullptr );
-	this->LineMaterialMID = UMaterialInstanceDynamic::Create( LineMaterial, this );
+	if (LineMaterial != nullptr)
+	{
+		this->LineMaterialMID = UMaterialInstanceDynamic::Create(LineMaterial, this);
 
-	FirstLineComponent->SetMaterial( 0, LineMaterialMID );
-	JointSphereComponent->SetMaterial( 0, LineMaterialMID );
-	SecondLineComponent->SetMaterial( 0, LineMaterialMID );
+		FirstLineComponent->SetMaterial(0, LineMaterialMID);
+		JointSphereComponent->SetMaterial(0, LineMaterialMID);
+		SecondLineComponent->SetMaterial(0, LineMaterialMID);
+	}
 }
 
 
diff --git a/Engine/Source/Editor/VREditor/VREditorFloatingText.h b/Engine/Source/Editor/VREditor/VREditorFloatingText.h
--- a/Engine/Source/Editor/VREditor/VREditorFloatingText.h
+++ b/Engine/Source/Editor/VREditor/VREditorFloatingText.h
@@ -10,7 +10,7 @@
 /**
  * Draws 3D text in the world along with targeting line cues
  */
-UCLASS()
+UCLASS(NotPlaceable, NotBlueprintable)
 class AFloatingText : public AActor
 {
 	GENERATED_BODY()
