Moving all asset container access to PostActorCreated to avoid VR editor assets in cooks

diff --git a/Engine/Source/Editor/VREditor/UI/VREditorDockableWindow.cpp b/Engine/Source/Editor/VREditor/UI/VREditorDockableWindow.cpp
--- a/Engine/Source/Editor/VREditor/UI/VREditorDockableWindow.cpp
+++ b/Engine/Source/Editor/VREditor/UI/VREditorDockableWindow.cpp
@@ -34,10 +34,11 @@ AVREditorDockableWindow::AVREditorDockableWindow() :
 	CloseButtonHoverAlpha(0.0f),
 	DockSelectDistance(0.0f)
 {
-	if (HasAnyFlags(RF_ClassDefaultObject))
-	{
-		return;
-	}
+}
+
+void AVREditorDockableWindow::PostActorCreated()
+{
+	Super::PostActorCreated();
 
 	const UVREditorAssetContainer& AssetContainer = UVREditorMode::LoadAssetContainer();
 
@@ -54,60 +55,62 @@ AVREditorDockableWindow::AVREditorDockableWindow() :
 	{
 		UStaticMesh* SelectionMesh = AssetContainer.WindowSelectionBarMesh;
 
-		SelectionBarMeshComponent = CreateDefaultSubobject<UStaticMeshComponent>( TEXT( "SelectionBarMesh" ) );
-		SelectionBarMeshComponent->SetStaticMesh( SelectionMesh );
-		SelectionBarMeshComponent->SetMobility( EComponentMobility::Movable );
-		SelectionBarMeshComponent->SetupAttachment( RootComponent );
-
+		SelectionBarMeshComponent = NewObject<UStaticMeshComponent>(this, TEXT("SelectionBarMesh"));
+		SelectionBarMeshComponent->SetStaticMesh(SelectionMesh);
+		SelectionBarMeshComponent->SetMobility(EComponentMobility::Movable);
+		SelectionBarMeshComponent->SetupAttachment(RootComponent);
+		SelectionBarMeshComponent->RegisterComponent();
 		SelectionBarMeshComponent->SetGenerateOverlapEvents(false);
-		SelectionBarMeshComponent->SetCanEverAffectNavigation( false );
+		SelectionBarMeshComponent->SetCanEverAffectNavigation(false);
 		SelectionBarMeshComponent->bCastDynamicShadow = false;
 		SelectionBarMeshComponent->bCastStaticShadow = false;
 		SelectionBarMeshComponent->bAffectDistanceFieldLighting = false;
 		SelectionBarMeshComponent->SetRelativeRotation(RelativeRotation);
 
 
-		SelectionBarMID = UMaterialInstanceDynamic::Create( HoverMaterial, GetTransientPackage() );
-		check( SelectionBarMID != nullptr );
-		SelectionBarMeshComponent->SetMaterial( 0, SelectionBarMID );
-		SelectionBarTranslucentMID = UMaterialInstanceDynamic::Create( TranslucentHoverMaterial, GetTransientPackage() );
-		check( SelectionBarTranslucentMID != nullptr );
-		SelectionBarMeshComponent->SetMaterial( 1, SelectionBarTranslucentMID );
+		SelectionBarMID = UMaterialInstanceDynamic::Create(HoverMaterial, GetTransientPackage());
+		check(SelectionBarMID != nullptr);
+		SelectionBarMeshComponent->SetMaterial(0, SelectionBarMID);
+		SelectionBarTranslucentMID = UMaterialInstanceDynamic::Create(TranslucentHoverMaterial, GetTransientPackage());
+		check(SelectionBarTranslucentMID != nullptr);
+		SelectionBarMeshComponent->SetMaterial(1, SelectionBarTranslucentMID);
 	}
 
 	{
 		UStaticMesh* CloseButtonMesh = AssetContainer.WindowCloseButtonMesh;
 
-		CloseButtonMeshComponent= CreateDefaultSubobject<UStaticMeshComponent>( TEXT( "CloseButtonMesh" ) );
-		CloseButtonMeshComponent->SetStaticMesh( CloseButtonMesh );
-		CloseButtonMeshComponent->SetMobility( EComponentMobility::Movable );
-		CloseButtonMeshComponent->SetupAttachment( RootComponent );
-
+		CloseButtonMeshComponent = NewObject<UStaticMeshComponent>(this, TEXT("CloseButtonMesh"));
+		CloseButtonMeshComponent->SetStaticMesh(CloseButtonMesh);
+		CloseButtonMeshComponent->SetMobility(EComponentMobility::Movable);
+		CloseButtonMeshComponent->SetupAttachment(RootComponent);
+		CloseButtonMeshComponent->RegisterComponent();
 		CloseButtonMeshComponent->SetGenerateOverlapEvents(false);
-		CloseButtonMeshComponent->SetCanEverAffectNavigation( false );
+		CloseButtonMeshComponent->SetCanEverAffectNavigation(false);
 		CloseButtonMeshComponent->bCastDynamicShadow = false;
 		CloseButtonMeshComponent->bCastStaticShadow = false;
 		CloseButtonMeshComponent->bAffectDistanceFieldLighting = false;
 		CloseButtonMeshComponent->SetRelativeRotation(RelativeRotation);
 
-		CloseButtonMID = UMaterialInstanceDynamic::Create( HoverMaterial, GetTransientPackage() );
-		check( CloseButtonMID != nullptr );
-		CloseButtonMeshComponent->SetMaterial( 0, CloseButtonMID );
-		CloseButtonTranslucentMID = UMaterialInstanceDynamic::Create( TranslucentHoverMaterial, GetTransientPackage() );
-		check( CloseButtonTranslucentMID != nullptr );
-		CloseButtonMeshComponent->SetMaterial( 1, CloseButtonTranslucentMID );
+		CloseButtonMID = UMaterialInstanceDynamic::Create(HoverMaterial, GetTransientPackage());
+		check(CloseButtonMID != nullptr);
+		CloseButtonMeshComponent->SetMaterial(0, CloseButtonMID);
+		CloseButtonTranslucentMID = UMaterialInstanceDynamic::Create(TranslucentHoverMaterial, GetTransientPackage());
+		check(CloseButtonTranslucentMID != nullptr);
+		CloseButtonMeshComponent->SetMaterial(1, CloseButtonTranslucentMID);
 	}
 
 	// The selection bar and close button will not be initially visible.  They'll appear when the user aims
 	// their laser toward the UI
-	SelectionBarMeshComponent->SetVisibility( false );
-	CloseButtonMeshComponent->SetVisibility( false );
+	SelectionBarMeshComponent->SetVisibility(false);
+	CloseButtonMeshComponent->SetVisibility(false);
 
 	// Create the drag operation
-	DragOperationComponent = CreateDefaultSubobject<UViewportDragOperationComponent>( TEXT( "DragOperation" ) );
-	DragOperationComponent->SetDragOperationClass( UDockableWindowDragOperation::StaticClass() );
+	DragOperationComponent = NewObject<UViewportDragOperationComponent>(this, TEXT("DragOperation"));
+	DragOperationComponent->SetDragOperationClass(UDockableWindowDragOperation::StaticClass());
 }
 
+
+
 void AVREditorDockableWindow::SetupWidgetComponent()
 {
 	Super::SetupWidgetComponent();
diff --git a/Engine/Source/Editor/VREditor/UI/VREditorDockableWindow.h b/Engine/Source/Editor/VREditor/UI/VREditorDockableWindow.h
--- a/Engine/Source/Editor/VREditor/UI/VREditorDockableWindow.h
+++ b/Engine/Source/Editor/VREditor/UI/VREditorDockableWindow.h
@@ -24,6 +24,7 @@ public:
 
 	/** Default constructor */
 	AVREditorDockableWindow();
+	virtual void PostActorCreated() override;
 
 	/** Updates the meshes for the UI */
 	virtual void TickManually( float DeltaTime ) override;
diff --git a/Engine/Source/Editor/VREditor/UI/VREditorRadialFloatingUI.cpp b/Engine/Source/Editor/VREditor/UI/VREditorRadialFloatingUI.cpp
--- a/Engine/Source/Editor/VREditor/UI/VREditorRadialFloatingUI.cpp
+++ b/Engine/Source/Editor/VREditor/UI/VREditorRadialFloatingUI.cpp
@@ -33,25 +33,26 @@ AVREditorRadialFloatingUI::AVREditorRadialFloatingUI()
 	FadeDelay( 0.0f ),
 	InitialScale( 1.0f )
 {
-	if (HasAnyFlags(RF_ClassDefaultObject))
-	{
-		return;
-	}
-
 	const bool bTransient = true;
-	USceneComponent* SceneComponent = CreateDefaultSubobject<USceneComponent>( TEXT( "SceneComponent" ), bTransient );
-	check( SceneComponent != nullptr );
+	USceneComponent* SceneComponent = CreateDefaultSubobject<USceneComponent>(TEXT("SceneComponent"), bTransient);
+	check(SceneComponent != nullptr);
 	this->RootComponent = SceneComponent;
 
 	DefaultGlowAmount = 2.0f;
+}
+
+
+void AVREditorRadialFloatingUI::PostActorCreated()
+{
+	Super::PostActorCreated();
 
 	const UVREditorAssetContainer& AssetContainer = UVREditorMode::LoadAssetContainer();
 
 	{
-		WindowMeshComponent = CreateDefaultSubobject<UStaticMeshComponent>(TEXT("WindowMesh"));
+		WindowMeshComponent = NewObject<UStaticMeshComponent>(this, TEXT("WindowMesh"));
 		WindowMeshComponent->SetMobility(EComponentMobility::Movable);
 		WindowMeshComponent->SetupAttachment(RootComponent);
-
+		WindowMeshComponent->RegisterComponent();
 		WindowMeshComponent->SetCollisionEnabled(ECollisionEnabled::QueryOnly);
 		WindowMeshComponent->SetCollisionResponseToAllChannels(ECollisionResponse::ECR_Ignore);
 
@@ -73,10 +74,10 @@ AVREditorRadialFloatingUI::AVREditorRadialFloatingUI()
 	}
 
 	{
-		ArrowMeshComponent = CreateDefaultSubobject<UStaticMeshComponent>(TEXT("ArrowMesh"));
+		ArrowMeshComponent = NewObject<UStaticMeshComponent>(this, TEXT("ArrowMesh"));
 		ArrowMeshComponent->SetMobility(EComponentMobility::Movable);
 		ArrowMeshComponent->SetupAttachment(WindowMeshComponent);
-
+		ArrowMeshComponent->RegisterComponent();
 		ArrowMeshComponent->SetCollisionEnabled(ECollisionEnabled::QueryOnly);
 		ArrowMeshComponent->SetCollisionResponseToAllChannels(ECollisionResponse::ECR_Ignore);
 
@@ -95,8 +96,9 @@ AVREditorRadialFloatingUI::AVREditorRadialFloatingUI()
 	}
 
 	{
-		CentralWidgetComponent = CreateDefaultSubobject<UVREditorWidgetComponent>(TEXT("CentralWidget"));
+		CentralWidgetComponent = NewObject<UVREditorWidgetComponent>(this, TEXT("CentralWidget"));
 		CentralWidgetComponent->SetupAttachment(RootComponent);
+		CentralWidgetComponent->RegisterComponent();
 		CentralWidgetComponent->SetCollisionEnabled(ECollisionEnabled::QueryOnly);
 		CentralWidgetComponent->SetCollisionResponseToAllChannels(ECollisionResponse::ECR_Ignore);
 		CentralWidgetComponent->SetGenerateOverlapEvents(false);
@@ -119,7 +121,6 @@ AVREditorRadialFloatingUI::AVREditorRadialFloatingUI()
 	}
 }
 
-
 void AVREditorRadialFloatingUI::SetupWidgetComponent(TSharedPtr<SWidget> SlateWidget)
 {
 
diff --git a/Engine/Source/Editor/VREditor/UI/VREditorRadialFloatingUI.h b/Engine/Source/Editor/VREditor/UI/VREditorRadialFloatingUI.h
--- a/Engine/Source/Editor/VREditor/UI/VREditorRadialFloatingUI.h
+++ b/Engine/Source/Editor/VREditor/UI/VREditorRadialFloatingUI.h
@@ -25,6 +25,7 @@ public:
 
 	/** Default constructor which sets up safe defaults */
 	AVREditorRadialFloatingUI();
+	virtual void PostActorCreated() override;
 
 	/** Creates a FVREditorFloatingUI using a Slate widget, and sets up safe defaults */
 	void SetSlateWidget( class UVREditorUISystem& InitOwner, const TSharedRef<SWidget>& InitSlateWidget, const FIntPoint InitResolution, const float InitScale, const EDockedTo InitDockedTo );
diff --git a/Engine/Source/Editor/VREditor/VREditorFloatingText.cpp b/Engine/Source/Editor/VREditor/VREditorFloatingText.cpp
--- a/Engine/Source/Editor/VREditor/VREditorFloatingText.cpp
+++ b/Engine/Source/Editor/VREditor/VREditorFloatingText.cpp
@@ -18,37 +18,38 @@
 
 AFloatingText::AFloatingText()
 {
-	if (UNLIKELY(IsRunningDedicatedServer()) || HasAnyFlags(RF_ClassDefaultObject))
-	{
-		return;
-	}
-
-	const UVREditorAssetContainer& AssetContainer = UVREditorMode::LoadAssetContainer();
-
 	// Create root default scene component
 	{
-		SceneComponent = CreateDefaultSubobject<USceneComponent>( TEXT( "SceneComponent" ) );
-		check( SceneComponent != nullptr );
+		SceneComponent = CreateDefaultSubobject<USceneComponent>(TEXT("SceneComponent"));
+		check(SceneComponent != nullptr);
 
 		RootComponent = SceneComponent;
 	}
+}
+
+
+void AFloatingText::PostActorCreated()
+{
+	Super::PostActorCreated();
 
 	// @todo vreditor: Tweak
 	const bool bAllowTextLighting = false;
 	const float TextSize = 1.5f;
 
+	const UVREditorAssetContainer& AssetContainer = UVREditorMode::LoadAssetContainer();
+
 	{
-		FirstLineComponent = CreateDefaultSubobject<UStaticMeshComponent>( TEXT( "FirstLine" ) );
-		check( FirstLineComponent != nullptr );
+		FirstLineComponent = NewObject<UStaticMeshComponent>(this, TEXT("FirstLine"));
+		check(FirstLineComponent != nullptr);
 
-		FirstLineComponent->SetStaticMesh( AssetContainer.LineSegmentCylinderMesh );
-		FirstLineComponent->SetMobility( EComponentMobility::Movable );
-		FirstLineComponent->SetupAttachment( SceneComponent );
-		
-		FirstLineComponent->SetCollisionEnabled( ECollisionEnabled::NoCollision );
+		FirstLineComponent->SetStaticMesh(AssetContainer.LineSegmentCylinderMesh);
+		FirstLineComponent->SetMobility(EComponentMobility::Movable);
+		FirstLineComponent->SetupAttachment(SceneComponent);
+		FirstLineComponent->RegisterComponent();
+		FirstLineComponent->SetCollisionEnabled(ECollisionEnabled::NoCollision);
 
 		FirstLineComponent->SetGenerateOverlapEvents(false);
-		FirstLineComponent->SetCanEverAffectNavigation( false );
+		FirstLineComponent->SetCanEverAffectNavigation(false);
 		FirstLineComponent->bCastDynamicShadow = bAllowTextLighting;
 		FirstLineComponent->bCastStaticShadow = false;
 		FirstLineComponent->bAffectDistanceFieldLighting = bAllowTextLighting;
@@ -56,36 +57,36 @@ AFloatingText::AFloatingText()
 	}
 
 	{
-		JointSphereComponent = CreateDefaultSubobject<UStaticMeshComponent>( TEXT( "JointSphere" ) );
-		check( JointSphereComponent != nullptr );
-
-		JointSphereComponent->SetStaticMesh( AssetContainer.JointSphereMesh );
-		JointSphereComponent->SetMobility( EComponentMobility::Movable );
-		JointSphereComponent->SetupAttachment( SceneComponent );
+		JointSphereComponent = NewObject<UStaticMeshComponent>(this, TEXT("JointSphere"));
+		check(JointSphereComponent != nullptr);
 
-		JointSphereComponent->SetCollisionEnabled( ECollisionEnabled::NoCollision );
+		JointSphereComponent->SetStaticMesh(AssetContainer.JointSphereMesh);
+		JointSphereComponent->SetMobility(EComponentMobility::Movable);
+		JointSphereComponent->SetupAttachment(SceneComponent);
+		JointSphereComponent->RegisterComponent();
+		JointSphereComponent->SetCollisionEnabled(ECollisionEnabled::NoCollision);
 
 		JointSphereComponent->SetGenerateOverlapEvents(false);
-		JointSphereComponent->SetCanEverAffectNavigation( false );
+		JointSphereComponent->SetCanEverAffectNavigation(false);
 		JointSphereComponent->bCastDynamicShadow = bAllowTextLighting;
 		JointSphereComponent->bCastStaticShadow = false;
 		JointSphereComponent->bAffectDistanceFieldLighting = bAllowTextLighting;
 		JointSphereComponent->bAffectDynamicIndirectLighting = bAllowTextLighting;
-		
+
 	}
 
 	{
-		SecondLineComponent = CreateDefaultSubobject<UStaticMeshComponent>( TEXT( "SecondLine" ) );
-		check( SecondLineComponent != nullptr );
+		SecondLineComponent = NewObject<UStaticMeshComponent>(this, TEXT("SecondLine"));
+		check(SecondLineComponent != nullptr);
 
-		SecondLineComponent->SetStaticMesh( AssetContainer.LineSegmentCylinderMesh );
-		SecondLineComponent->SetMobility( EComponentMobility::Movable );
-		SecondLineComponent->SetupAttachment( SceneComponent );
-
-		SecondLineComponent->SetCollisionEnabled( ECollisionEnabled::NoCollision );
+		SecondLineComponent->SetStaticMesh(AssetContainer.LineSegmentCylinderMesh);
+		SecondLineComponent->SetMobility(EComponentMobility::Movable);
+		SecondLineComponent->SetupAttachment(SceneComponent);
+		SecondLineComponent->RegisterComponent();
+		SecondLineComponent->SetCollisionEnabled(ECollisionEnabled::NoCollision);
 
 		SecondLineComponent->SetGenerateOverlapEvents(false);
-		SecondLineComponent->SetCanEverAffectNavigation( false );
+		SecondLineComponent->SetCanEverAffectNavigation(false);
 		SecondLineComponent->bCastDynamicShadow = bAllowTextLighting;
 		SecondLineComponent->bCastStaticShadow = false;
 		SecondLineComponent->bAffectDistanceFieldLighting = bAllowTextLighting;
@@ -98,44 +99,38 @@ AFloatingText::AFloatingText()
 	TranslucentTextMaterial = AssetContainer.TranslucentTextMaterial;
 
 	{
-		TextComponent = CreateDefaultSubobject<UTextRenderComponent>( TEXT( "Text" ) );
-		check( TextComponent != nullptr );
-
-		TextComponent->SetMobility( EComponentMobility::Movable );
-		TextComponent->SetupAttachment( SceneComponent );
+		TextComponent = NewObject<UTextRenderComponent>(this, TEXT("Text"));
+		check(TextComponent != nullptr);
 
-		TextComponent->SetCollisionProfileName( UCollisionProfile::NoCollision_ProfileName );
+		TextComponent->SetMobility(EComponentMobility::Movable);
+		TextComponent->SetupAttachment(SceneComponent);
+		TextComponent->RegisterComponent();
+		TextComponent->SetCollisionProfileName(UCollisionProfile::NoCollision_ProfileName);
 
 		TextComponent->SetGenerateOverlapEvents(false);
-		TextComponent->SetCanEverAffectNavigation( false );
+		TextComponent->SetCanEverAffectNavigation(false);
 		TextComponent->bCastDynamicShadow = bAllowTextLighting;
 		TextComponent->bCastStaticShadow = false;
 		TextComponent->bAffectDistanceFieldLighting = bAllowTextLighting;
 		TextComponent->bAffectDynamicIndirectLighting = bAllowTextLighting;
 
 
-		TextComponent->SetWorldSize( TextSize );
+		TextComponent->SetWorldSize(TextSize);
 
 		// Use a custom font.  The text will be visible up close.	   
-		TextComponent->SetFont( AssetContainer.TextFont );
+		TextComponent->SetFont(AssetContainer.TextFont);
 
 		if (MaskedTextMaterial != nullptr)
 		{
 			// Assign our custom text rendering material.
 			TextComponent->SetTextMaterial(MaskedTextMaterial);
 		}
-		TextComponent->SetTextRenderColor( FLinearColor::White.ToFColor( false ) );
+		TextComponent->SetTextRenderColor(FLinearColor::White.ToFColor(false));
 
 		// Left justify the text
-		TextComponent->SetHorizontalAlignment( EHTA_Left );
+		TextComponent->SetHorizontalAlignment(EHTA_Left);
 
 	}
-}
-
-
-void AFloatingText::PostActorCreated()
-{
-	Super::PostActorCreated();
 
 	// Create an MID so that we can change parameters on the fly (fading)
 	if (LineMaterial != nullptr)
