Fix for crash in vr mode

diff --git a/Engine/Source/Editor/VREditor/UI/VREditorDockableWindow.cpp b/Engine/Source/Editor/VREditor/UI/VREditorDockableWindow.cpp
--- a/Engine/Source/Editor/VREditor/UI/VREditorDockableWindow.cpp
+++ b/Engine/Source/Editor/VREditor/UI/VREditorDockableWindow.cpp
@@ -55,7 +55,7 @@ void AVREditorDockableWindow::PostActorCreated()
 	{
 		UStaticMesh* DockingMesh = AssetContainer.DockingButtonMesh;
 
-		DockButtonMeshComponent = CreateDefaultSubobject<UStaticMeshComponent>( TEXT( "DockMesh" ) );
+		DockButtonMeshComponent = NewObject<UStaticMeshComponent>(this, TEXT( "DockMesh" ) );
 		DockButtonMeshComponent->SetStaticMesh(DockingMesh);
 		DockButtonMeshComponent->SetMobility( EComponentMobility::Movable );
 		DockButtonMeshComponent->SetupAttachment( RootComponent );
