In VR editor, only scale the actor when the actor can not be carried

diff --git a/Engine/Source/Editor/VREditor/VREditorPlacement.cpp b/Engine/Source/Editor/VREditor/VREditorPlacement.cpp
--- a/Engine/Source/Editor/VREditor/VREditorPlacement.cpp
+++ b/Engine/Source/Editor/VREditor/VREditorPlacement.cpp
@@ -372,25 +372,34 @@ void UVREditorPlacement::StartPlacingObjects( const TArray<UObject*>& ObjectsToP
 				FEditorDelegates::OnNewActorsDropped.Broadcast( DroppedObjects, AllNewActors );
 			}
 
-			FBox BoundsOfAllActors;
+			static const auto CVarAllowCarryingCertainObjects = IConsoleManager::Get().FindTConsoleVariableDataInt(TEXT("VI.AllowCarryingCertainObjects"));
+			const bool bCanBeCarried = ViewportWorldInteraction->GetTransformables().Num() == 1 && ViewportWorldInteraction->GetTransformables()[0].Get()->ShouldBeCarried();
+			const bool bShouldInterpolateScaleFromDragLocation = bShouldInterpolateFromDragLocation && (CVarAllowCarryingCertainObjects->GetValueOnAnyThread() == 0 || !bCanBeCarried);
+
+			float DesiredScale = 1.f;
+			if( bShouldInterpolateScaleFromDragLocation )
 			{
-				BoundsOfAllActors.Init();
+				FBox BoundsOfAllActors;
+				{
+					BoundsOfAllActors.Init();
+					for( AActor* NewActor : AllNewActors )
+					{
+						BoundsOfAllActors += NewActor->CalculateComponentsBoundingBoxInLocalSpace();
+					}
+				}
+				const float BoundsOfAllActorsSize = BoundsOfAllActors.GetSize().GetAbsMax();
+				const float UsedBoundsOfAllActorsSize = BoundsOfAllActorsSize == 0 ? 1 : BoundsOfAllActorsSize;
+				DesiredScale = (VREd::SizeOfActorsOverContentBrowserThumbnail->GetFloat() / UsedBoundsOfAllActorsSize) * ViewportWorldInteraction->GetWorldScaleFactor();
+
+				// Start the placed objects off scaled down to match the content browser thumbnail
 				for( AActor* NewActor : AllNewActors )
 				{
-					BoundsOfAllActors += NewActor->CalculateComponentsBoundingBoxInLocalSpace();
+					NewActor->SetActorScale3D(FVector(DesiredScale));
 				}
 			}
-			const float BoundsOfAllActorsSize = BoundsOfAllActors.GetSize().GetAbsMax();
-			const float UsedBoundsOfAllActorsSize = BoundsOfAllActorsSize == 0 ? 1 : BoundsOfAllActorsSize;
-			const float DesiredScale = ( VREd::SizeOfActorsOverContentBrowserThumbnail->GetFloat() / UsedBoundsOfAllActorsSize ) * ViewportWorldInteraction->GetWorldScaleFactor();
-
-			// Start the placed objects off scaled down to match the content browser thumbnail
-			if( bShouldInterpolateFromDragLocation )
+			else if( bShouldInterpolateFromDragLocation )
 			{
- 				for( AActor* NewActor : AllNewActors )
- 				{
- 					NewActor->SetActorScale3D( FVector( DesiredScale ) );
- 				}
+				PlaceAt = PlacingWithInteractor->GetInteractorData().Transform.GetLocation();
 			}
 
 			// We changed the initial scale of selected actors, so make sure our transformables and gizmo start transform are up to date.
