Fix for crash on undo after placing cameras

diff --git a/Engine/Source/Editor/VREditor/UI/VREditorUISystem.cpp b/Engine/Source/Editor/VREditor/UI/VREditorUISystem.cpp
--- a/Engine/Source/Editor/VREditor/UI/VREditorUISystem.cpp
+++ b/Engine/Source/Editor/VREditor/UI/VREditorUISystem.cpp
@@ -2456,8 +2456,12 @@ void UVREditorUISystem::UpdateSequencerUI()
 
 void UVREditorUISystem::UpdateActorPreviewUI( TSharedRef<SWidget> InWidget, int32 Index, AActor *Actor )
 {
-	check(Actor != nullptr)
-	FString PanelString = ActorPreviewPrefix + FString::FromInt( Actor->GetUniqueID() );
+	
+	FString PanelString = ActorPreviewPrefix;
+	if (Actor)
+	{
+		PanelString += FString::FromInt(Actor->GetUniqueID());
+	}
 	VREditorPanelID PreviewPanelID = FName(*PanelString);
 	AVREditorFloatingUI* PreviewPanel = GetPanel(PreviewPanelID);
 
