Disable VRMode for WMR devices.

diff --git a/Engine/Source/Editor/VREditor/VREditorModeManager.cpp b/Engine/Source/Editor/VREditor/VREditorModeManager.cpp
--- a/Engine/Source/Editor/VREditor/VREditorModeManager.cpp
+++ b/Engine/Source/Editor/VREditor/VREditorModeManager.cpp
@@ -146,11 +146,19 @@ bool FVREditorModeManager::IsVREditorActive() const
 	return CurrentVREditorMode != nullptr && CurrentVREditorMode->IsActive();
 }
 
-
+const static FName WMRSytemName = FName(TEXT("WindowsMixedRealityHMD"));
 bool FVREditorModeManager::IsVREditorAvailable() const
 {
-	const bool bHasHMDDevice = GEngine->XRSystem.IsValid() && GEngine->XRSystem->GetHMDDevice() && GEngine->XRSystem->GetHMDDevice()->IsHMDEnabled();
-	return bHasHMDDevice && !GEditor->bIsSimulatingInEditor;
+	if (GEngine->XRSystem.IsValid() && GEngine->XRSystem->GetHMDDevice() && GEngine->XRSystem->GetHMDDevice()->IsHMDEnabled())
+	{
+		// TODO: UE-71871 Work around for avoiding starting VRMode when using WMR
+		const bool bIsWMR = GEngine->XRSystem->GetSystemName() == WMRSytemName;
+		return !bIsWMR && !GEditor->bIsSimulatingInEditor;
+	}
+	else
+	{
+		return false;
+	}
 }
 
 bool FVREditorModeManager::IsVREditorButtonActive() const
