PR #5727: Replace comment of GetSelectedActors (Contributed by [REDACTED_PERSON])

diff --git a/Engine/Source/Editor/VREditor/VRScoutingInteractor.h b/Engine/Source/Editor/VREditor/VRScoutingInteractor.h
--- a/Engine/Source/Editor/VREditor/VRScoutingInteractor.h
+++ b/Engine/Source/Editor/VREditor/VRScoutingInteractor.h
@@ -35,9 +35,9 @@ public:
 	// IViewportInteractorInterface overrides
 	virtual void Shutdown_Implementation() override;
 
-	/** Set the VR flight speed cvar */
+	/** Gets all actors that are selected in the world editor */
 	UFUNCTION(BlueprintCallable, Category = "Scouting")
-	static TArray<AActor*> GetSelectedActors() ;
+	static TArray<AActor*> GetSelectedActors();
 
 	/** Shown in Navigation mode */
 	UPROPERTY(Category = Interactor, EditAnywhere, BlueprintReadOnly)
