Disable VRMode when using the OpenXR plugin

diff --git a/Engine/Source/Editor/VREditor/VREditorModeManager.cpp b/Engine/Source/Editor/VREditor/VREditorModeManager.cpp
--- a/Engine/Source/Editor/VREditor/VREditorModeManager.cpp
+++ b/Engine/Source/Editor/VREditor/VREditorModeManager.cpp
@@ -147,12 +147,14 @@ bool FVREditorModeManager::IsVREditorActive() const
 }
 
 const static FName WMRSytemName = FName(TEXT("WindowsMixedRealityHMD"));
+const static FName OXRSytemName = FName(TEXT("OpenXR"));
 bool FVREditorModeManager::IsVREditorAvailable() const
 {
 	if (GEngine->XRSystem.IsValid() && GEngine->XRSystem->GetHMDDevice() && GEngine->XRSystem->GetHMDDevice()->IsHMDEnabled())
 	{
-		// TODO: UE-71871 Work around for avoiding starting VRMode when using WMR
-		const bool bIsWMR = GEngine->XRSystem->GetSystemName() == WMRSytemName;
+		// TODO: UE-71871/UE-73237 Work around for avoiding starting VRMode when using WMR
+		FName SystemName = GEngine->XRSystem->GetSystemName();
+		const bool bIsWMR = SystemName == WMRSytemName || SystemName == OXRSytemName;
 		return !bIsWMR && !GEditor->bIsSimulatingInEditor;
 	}
 	else
