Removing blutility dependencies in VREditor module, making data validation an editor plugin
#rnx

diff --git a/Engine/Plugins/Editor/DataValidation/DataValidation.uplugin b/Engine/Plugins/Editor/DataValidation/DataValidation.uplugin
--- a/Engine/Plugins/Editor/DataValidation/DataValidation.uplugin
+++ b/Engine/Plugins/Editor/DataValidation/DataValidation.uplugin
@@ -19,7 +19,7 @@
 	[
 		{
 			"Name" : "DataValidation",
-			"Type": "Developer",
+			"Type": "Editor",
 			"LoadingPhase" : "PreDefault"
 		}
 	]
diff --git a/Engine/Source/Editor/Blutility/Classes/EditorUtilityBlueprint.h b/Engine/Source/Editor/Blutility/Classes/EditorUtilityBlueprint.h
--- a/Engine/Source/Editor/Blutility/Classes/EditorUtilityBlueprint.h
+++ b/Engine/Source/Editor/Blutility/Classes/EditorUtilityBlueprint.h
@@ -16,18 +16,13 @@ class BLUTILITY_API UEditorUtilityBlueprint : public UBlueprint
 {
 	GENERATED_UCLASS_BODY()
 
-#if WITH_EDITOR
 	// UBlueprint interface
 	virtual bool SupportedByDefaultBlueprintFactory() const override
 	{
 		return false;
 	}
-	
-#endif
 
-#if WITH_EDITORONLY_DATA
 protected:
 	virtual void LoadModulesRequiredForCompilation() override;
-#endif
 // End of UBlueprint interface
 };
diff --git a/Engine/Source/Editor/VREditor/UI/VREditorFloatingUI.h b/Engine/Source/Editor/VREditor/UI/VREditorFloatingUI.h
--- a/Engine/Source/Editor/VREditor/UI/VREditorFloatingUI.h
+++ b/Engine/Source/Editor/VREditor/UI/VREditorFloatingUI.h
@@ -6,7 +6,7 @@
 #include "UObject/ObjectMacros.h"
 #include "Widgets/SWidget.h"
 #include "VREditorBaseActor.h"
-#include "EditorUtilityWidget.h"
+#include "Blueprint/UserWidget.h"
 #include "VREditorFloatingUI.generated.h"
 
 class UVREditorBaseUserWidget;
@@ -25,32 +25,32 @@ struct FVREditorFloatingUICreationContext
 public:	
 	
 	/** Widget to open in the VR window. null to close an open window (if if matches the PanelID) */
-	UPROPERTY(EditAnywhere, BlueprintReadWrite, Category = "Virtual Production UI")
-	TSubclassOf<UEditorUtilityWidget> WidgetClass;
+	UPROPERTY(EditAnywhere, BlueprintReadWrite, Category = "VR Mode UI")
+	TSubclassOf<UUserWidget> WidgetClass;
 
 	// @todo: As we make this user-definable, we allow possible name collisions - how to deal with that? MakeUniqueName won't work because the name will be different on repeat calls. 
 	/** ID that the UI system will use to identify the panel. MUST BE UNIQUE! */ 
-	UPROPERTY(EditAnywhere, BlueprintReadWrite, Category = "Virtual Production UI")
+	UPROPERTY(EditAnywhere, BlueprintReadWrite, Category = "VR Mode UI")
 	FName PanelID;
 
 	/** Optional offset from HMD where the window opens. Pass FTransform::Identity for default logic - window will open at controller location. */
-	UPROPERTY(EditAnywhere, BlueprintReadWrite, Category = "Virtual Production UI")
+	UPROPERTY(EditAnywhere, BlueprintReadWrite, Category = "VR Mode UI")
 	FTransform PanelSpawnOffset;
 
 	/** Panel size. Should match the size of the UMG passed in. */
-	UPROPERTY(EditAnywhere, BlueprintReadWrite, Category = "Virtual Production UI")
+	UPROPERTY(EditAnywhere, BlueprintReadWrite, Category = "VR Mode UI")
 	FVector2D PanelSize;
 
 	/** NOT IN USE YET! Custom mesh to use for the VR window. */
-	UPROPERTY(EditAnywhere, BlueprintReadWrite, Category = "Virtual Production UI")
+	UPROPERTY(EditAnywhere, BlueprintReadWrite, Category = "VR Mode UI")
 	UStaticMesh* PanelMesh;
 
 	/** Optional override for "VREd.EditorUISize". Leave at 0 for default. */
-	UPROPERTY(EditAnywhere, BlueprintReadWrite, Category = "Virtual Production UI")
+	UPROPERTY(EditAnywhere, BlueprintReadWrite, Category = "VR Mode UI")
 	float EditorUISize;
 
 	//** Turn off handles under window? (X-To-Close, movement bar...) */
-	UPROPERTY(EditAnywhere, BlueprintReadWrite, Category = "Virtual Production UI")
+	UPROPERTY(EditAnywhere, BlueprintReadWrite, Category = "VR Mode UI")
 	bool bHideWindowHandles;
 };
 
diff --git a/Engine/Source/Editor/VREditor/VREditor.Build.cs b/Engine/Source/Editor/VREditor/VREditor.Build.cs
--- a/Engine/Source/Editor/VREditor/VREditor.Build.cs
+++ b/Engine/Source/Editor/VREditor/VREditor.Build.cs
@@ -26,8 +26,7 @@ namespace UnrealBuildTool.Rules
 					"Analytics",
                     "LevelSequence",
                     "Sequencer",
-                    "Projects",
-					"Blutility"
+                    "Projects"
 				}
 			);
 
