Fix for disappearing UI handle in VR mode

diff --git a/Engine/Source/Editor/VREditor/UI/VREditorDockableWindow.cpp b/Engine/Source/Editor/VREditor/UI/VREditorDockableWindow.cpp
--- a/Engine/Source/Editor/VREditor/UI/VREditorDockableWindow.cpp
+++ b/Engine/Source/Editor/VREditor/UI/VREditorDockableWindow.cpp
@@ -56,9 +56,11 @@ void AVREditorDockableWindow::PostActorCreated()
 		UStaticMesh* DockingMesh = AssetContainer.DockingButtonMesh;
 
 		DockButtonMeshComponent = NewObject<UStaticMeshComponent>(this, TEXT( "DockMesh" ) );
+		DockButtonMeshComponent->SetupAttachment(RootComponent);
+		DockButtonMeshComponent->RegisterComponent();
 		DockButtonMeshComponent->SetStaticMesh(DockingMesh);
 		DockButtonMeshComponent->SetMobility( EComponentMobility::Movable );
-		DockButtonMeshComponent->SetupAttachment( RootComponent );
+
 
 		DockButtonMeshComponent->SetGenerateOverlapEvents(false);
 		DockButtonMeshComponent->SetCanEverAffectNavigation( false );
@@ -76,7 +78,10 @@ void AVREditorDockableWindow::PostActorCreated()
 		UStaticMesh* SelectionMesh = AssetContainer.WindowSelectionBarMesh;
 
 		SelectionBarMeshComponent = NewObject<UStaticMeshComponent>(this, TEXT("SelectionBarMesh"));
+		SelectionBarMeshComponent->SetupAttachment(RootComponent);
+		SelectionBarMeshComponent->SetMobility(EComponentMobility::Movable);
 		SelectionBarMeshComponent->RegisterComponent();
+		SelectionBarMeshComponent->SetStaticMesh(SelectionMesh);
 		SelectionBarMeshComponent->SetGenerateOverlapEvents(false);
 		SelectionBarMeshComponent->SetCanEverAffectNavigation(false);
 		SelectionBarMeshComponent->bCastDynamicShadow = false;
