Fix for crash on simulate in VR mode

diff --git a/Engine/Source/Editor/VREditor/UI/VREditorRadialFloatingUI.cpp b/Engine/Source/Editor/VREditor/UI/VREditorRadialFloatingUI.cpp
--- a/Engine/Source/Editor/VREditor/UI/VREditorRadialFloatingUI.cpp
+++ b/Engine/Source/Editor/VREditor/UI/VREditorRadialFloatingUI.cpp
@@ -499,11 +499,14 @@ const void AVREditorRadialFloatingUI::HighlightSlot(const FVector2D& TrackpadPos
 		ArrowMeshComponent->SetVisibility(false);
 	}
 
+	TSharedRef<SWidget> TestWidget = SNullWidget::NullWidget;
 	const int32 Index = (Angle / AnglePerItem);
 
-	TSharedRef<SWidget> CurrentChild = WidgetComponents[Index]->GetSlateWidget().ToSharedRef();
-	TSharedRef<SWidget> TestWidget = UVREditorUISystem::FindWidgetOfType(CurrentChild, ButtonTypeOverride);
-
+	if (WidgetComponents.IsValidIndex(Index) && WidgetComponents[Index]->GetSlateWidget())
+	{
+		TSharedRef<SWidget> CurrentChild = WidgetComponents[Index]->GetSlateWidget().ToSharedRef();
+		TestWidget = UVREditorUISystem::FindWidgetOfType(CurrentChild, ButtonTypeOverride);
+	}
 	if (TestWidget != SNullWidget::NullWidget)
 	{
 		CurrentlyHoveredButton = StaticCastSharedRef<SButton>(TestWidget);
