Moving "reset all other widgets" to outside a conditional

diff --git a/Engine/Source/Editor/VREditor/UI/VREditorRadialFloatingUI.cpp b/Engine/Source/Editor/VREditor/UI/VREditorRadialFloatingUI.cpp
--- a/Engine/Source/Editor/VREditor/UI/VREditorRadialFloatingUI.cpp
+++ b/Engine/Source/Editor/VREditor/UI/VREditorRadialFloatingUI.cpp
@@ -499,9 +499,11 @@ const void AVREditorRadialFloatingUI::HighlightSlot(const FVector2D& TrackpadPos
 		ArrowMeshComponent->SetVisibility(false);
 	}
 
+	const FPointerEvent& SimulatedPointer = FPointerEvent();
+	const FGeometry& ChildGeometry = FGeometry();
 	TSharedRef<SWidget> TestWidget = SNullWidget::NullWidget;
 	const int32 Index = (Angle / AnglePerItem);
-
+	// Make sure we are checking for a valid radial menu widget
 	if (WidgetComponents.IsValidIndex(Index) && WidgetComponents[Index]->GetSlateWidget())
 	{
 		TSharedRef<SWidget> CurrentChild = WidgetComponents[Index]->GetSlateWidget().ToSharedRef();
@@ -511,34 +513,30 @@ const void AVREditorRadialFloatingUI::HighlightSlot(const FVector2D& TrackpadPos
 	{
 		CurrentlyHoveredButton = StaticCastSharedRef<SButton>(TestWidget);
 		CurrentlyHoveredWidget = WidgetComponents[Index];
-		const FPointerEvent& SimulatedPointer = FPointerEvent();
-		const FGeometry& ChildGeometry = FGeometry();
-
+	
 		// Simulate mouse entering event for the button if it was not previously hovered
 		if (!(CurrentlyHoveredButton->IsHovered()))
 		{
 			CurrentlyHoveredButton->OnMouseEnter(ChildGeometry, SimulatedPointer);
 			Owner->OnHoverBeginEffect(CurrentlyHoveredWidget);
 		}
-
-		// Simulate mouse leaving events for any buttons that were previously hovered
-		for (int32 ButtonCount = 0; ButtonCount < (NumberOfEntries); ButtonCount++)
+	}
+	// Simulate mouse leaving events for any buttons that were previously hovered
+	for (int32 ButtonCount = 0; ButtonCount < (NumberOfEntries); ButtonCount++)
+	{
+		if (ButtonCount != Index && WidgetComponents.IsValidIndex(ButtonCount))
 		{
-			if (ButtonCount != Index)
+			TSharedRef<SWidget> ChildWidget = WidgetComponents[ButtonCount]->GetSlateWidget().ToSharedRef();
+			TestWidget = UVREditorUISystem::FindWidgetOfType(ChildWidget, ButtonTypeOverride);
+			TSharedRef<SButton> TestButton = StaticCastSharedRef<SButton>(TestWidget);
+			if (TestButton->IsHovered())
 			{
-				TSharedRef<SWidget> ChildWidget = WidgetComponents[ButtonCount]->GetSlateWidget().ToSharedRef();
-				TestWidget = UVREditorUISystem::FindWidgetOfType(ChildWidget, ButtonTypeOverride);
-				TSharedRef<SButton> TestButton = StaticCastSharedRef<SButton>(TestWidget);
-				if (TestButton->IsHovered())
-				{
-					TestButton->OnMouseLeave(SimulatedPointer);
-					Owner->OnHoverEndEffect(WidgetComponents[ButtonCount]);
-				}
+				TestButton->OnMouseLeave(SimulatedPointer);
+				Owner->OnHoverEndEffect(WidgetComponents[ButtonCount]);
 			}
-
 		}
-	}
 
+	}
 }
 
 void AVREditorRadialFloatingUI::SimulateLeftClick()
