Guarding against updating the menu while it is transitioning.

diff --git a/Engine/Source/Editor/VREditor/UI/VREditorRadialFloatingUI.cpp b/Engine/Source/Editor/VREditor/UI/VREditorRadialFloatingUI.cpp
--- a/Engine/Source/Editor/VREditor/UI/VREditorRadialFloatingUI.cpp
+++ b/Engine/Source/Editor/VREditor/UI/VREditorRadialFloatingUI.cpp
@@ -524,7 +524,7 @@ const void AVREditorRadialFloatingUI::HighlightSlot(const FVector2D& TrackpadPos
 	// Simulate mouse leaving events for any buttons that were previously hovered
 	for (int32 ButtonCount = 0; ButtonCount < (NumberOfEntries); ButtonCount++)
 	{
-		if (ButtonCount != Index && WidgetComponents.IsValidIndex(ButtonCount))
+		if (ButtonCount != Index && WidgetComponents.IsValidIndex(ButtonCount) && WidgetComponents[ButtonCount]->GetSlateWidget())
 		{
 			TSharedRef<SWidget> ChildWidget = WidgetComponents[ButtonCount]->GetSlateWidget().ToSharedRef();
 			TestWidget = UVREditorUISystem::FindWidgetOfType(ChildWidget, ButtonTypeOverride);
diff --git a/Engine/Source/Editor/VREditor/VREditorMode.cpp b/Engine/Source/Editor/VREditor/VREditorMode.cpp
--- a/Engine/Source/Editor/VREditor/VREditorMode.cpp
+++ b/Engine/Source/Editor/VREditor/VREditorMode.cpp
@@ -1251,6 +1251,7 @@ void UVREditorMode::PostPIEStarted( bool bIsSimulatingInEditor )
 	{
 		GEnableVREditorHacks = false;
 	}
+	UISystem->TryToSpawnRadialMenu(UISystem->GetUIInteractor(), true, false);
 }
 
 
diff --git a/Engine/Source/Editor/VREditor/VREditorMotionControllerInteractor.cpp b/Engine/Source/Editor/VREditor/VREditorMotionControllerInteractor.cpp
--- a/Engine/Source/Editor/VREditor/VREditorMotionControllerInteractor.cpp
+++ b/Engine/Source/Editor/VREditor/VREditorMotionControllerInteractor.cpp
@@ -1395,6 +1395,12 @@ void UVREditorMotionControllerInteractor::SetLaserVisuals( const FLinearColor& N
 
 void UVREditorMotionControllerInteractor::UpdateRadialMenuInput( const float DeltaTime )
 {
+	// Don't try to update the menu if we are about to start/stop simulating and the menu needs to transition
+	if (GEditor->bIsSimulateInEditorQueued || GEditor->bRequestEndPlayMapQueued)
+	{
+		return;
+	}
+
 	UVREditorUISystem& UISystem = GetVRMode().GetUISystem();
 	const FName HMDDeviceType = GetVRMode().GetHMDDeviceType();
 	//Update the radial menu
