Icons don't disappear on the radial menu during simulate anymore

diff --git a/Engine/Source/Editor/VREditor/UI/VREditorUISystem.cpp b/Engine/Source/Editor/VREditor/UI/VREditorUISystem.cpp
--- a/Engine/Source/Editor/VREditor/UI/VREditorUISystem.cpp
+++ b/Engine/Source/Editor/VREditor/UI/VREditorUISystem.cpp
@@ -2625,6 +2625,8 @@ void UVREditorUISystem::TransitionWorld(UWorld* NewWorld, EEditorWorldExtensionT
 			TabManagerUI->GetWidgetComponent()->UpdateWidget();
 			ProxyTabManager->SetParentWindow(TabManagerUI->GetWidgetComponent()->GetSlateWindow().ToSharedRef());
 		}
+
+		TryToSpawnRadialMenu(GetUIInteractor(), true, false);
 	}
 }
 
diff --git a/Engine/Source/Editor/VREditor/VREditorMode.cpp b/Engine/Source/Editor/VREditor/VREditorMode.cpp
--- a/Engine/Source/Editor/VREditor/VREditorMode.cpp
+++ b/Engine/Source/Editor/VREditor/VREditorMode.cpp
@@ -1251,7 +1251,6 @@ void UVREditorMode::PostPIEStarted( bool bIsSimulatingInEditor )
 	{
 		GEnableVREditorHacks = false;
 	}
-	UISystem->TryToSpawnRadialMenu(UISystem->GetUIInteractor(), true, false);
 }
 
 
diff --git a/Engine/Source/Editor/VREditor/VREditorMotionControllerInteractor.cpp b/Engine/Source/Editor/VREditor/VREditorMotionControllerInteractor.cpp
--- a/Engine/Source/Editor/VREditor/VREditorMotionControllerInteractor.cpp
+++ b/Engine/Source/Editor/VREditor/VREditorMotionControllerInteractor.cpp
@@ -1395,12 +1395,6 @@ void UVREditorMotionControllerInteractor::SetLaserVisuals( const FLinearColor& N
 
 void UVREditorMotionControllerInteractor::UpdateRadialMenuInput( const float DeltaTime )
 {
-	// Don't try to update the menu if we are about to start/stop simulating and the menu needs to transition
-	if (GEditor->bIsSimulateInEditorQueued || GEditor->bRequestEndPlayMapQueued)
-	{
-		return;
-	}
-
 	UVREditorUISystem& UISystem = GetVRMode().GetUISystem();
 	const FName HMDDeviceType = GetVRMode().GetHMDDeviceType();
 	//Update the radial menu
