Fixed radial menu inputs to not be touch events

diff --git a/Engine/Source/Editor/VREditor/UI/VREditorRadialFloatingUI.cpp b/Engine/Source/Editor/VREditor/UI/VREditorRadialFloatingUI.cpp
--- a/Engine/Source/Editor/VREditor/UI/VREditorRadialFloatingUI.cpp
+++ b/Engine/Source/Editor/VREditor/UI/VREditorRadialFloatingUI.cpp
@@ -552,7 +552,7 @@ void AVREditorRadialFloatingUI::SimulateLeftClick()
 		}
 		if (ButtonTypeOverride == FName(TEXT("SButton")))
 		{
-			const FPointerEvent& SimulatedPointer = FPointerEvent(uint32(0), uint32(0), FVector2D::ZeroVector, FVector2D::ZeroVector, 1.0f, true);
+			const FPointerEvent& SimulatedPointer = FPointerEvent(uint32(0), uint32(0), FVector2D::ZeroVector, FVector2D::ZeroVector, TSet<FKey>(), EKeys::LeftMouseButton, 0.0f, FModifierKeysState());
 			const FGeometry& ChildGeometry = FGeometry();
 			CurrentlyHoveredButton->OnMouseButtonDown(ChildGeometry, SimulatedPointer);
 			CurrentlyHoveredButton->OnMouseButtonUp(ChildGeometry, SimulatedPointer);
