This issue was introduced by the fix that allowed MotionControllerComponents to correctly replicate in networked games.  Previously the local motion controllers overrode the remote controller's positions.  This was fixed by only updating if HasLocalNetOwner was true.  But that was returning false for VREditorAvatarActor, as it's Owner was null.
So for now, VREditorAvatarActor's HasLocalNetOwner now always returns true.  When using Multi-User with two VRMode editors, we might have to do something else here for when this is called by MotionControllerComponent::PollControllerState()

diff --git a/Engine/Source/Editor/VREditor/VREditorAvatarActor.h b/Engine/Source/Editor/VREditor/VREditorAvatarActor.h
--- a/Engine/Source/Editor/VREditor/VREditorAvatarActor.h
+++ b/Engine/Source/Editor/VREditor/VREditorAvatarActor.h
@@ -34,6 +34,12 @@ public:
 		return true;
 	}
 
+	virtual bool HasLocalNetOwner() const override
+	{
+		// TODO: When using Multi-User with two VRMode editors, we might have to do something else here for when this is called by MotionControllerComponent::PollControllerState()
+		return true;
+	}
+
 private:
 
 	/** Our avatar's head mesh */
