Remove the developer path to compile blueprint interactor while in VR.

diff --git a/Engine/Source/Editor/VREditor/Public/VREditorMode.h b/Engine/Source/Editor/VREditor/Public/VREditorMode.h
--- a/Engine/Source/Editor/VREditor/Public/VREditorMode.h
+++ b/Engine/Source/Editor/VREditor/Public/VREditorMode.h
@@ -68,9 +68,6 @@ public:
 	/** When the user actually enters the VR Editor mode */
 	void Enter();
 
-	/** When Blueprints are reinstanced so we can re-connect our interactors. */
-	void OnBlueprintReinstanced();
-
 	/** When the user leaves the VR Editor mode */
 	void Exit( const bool bShouldDisableStereo );
 
diff --git a/Engine/Source/Editor/VREditor/VREditorMode.cpp b/Engine/Source/Editor/VREditor/VREditorMode.cpp
--- a/Engine/Source/Editor/VREditor/VREditorMode.cpp
+++ b/Engine/Source/Editor/VREditor/VREditorMode.cpp
@@ -138,8 +138,6 @@ void UVREditorMode::Init()
 		check( WorldInteraction != nullptr );
 	}
 
-	GEditor->OnBlueprintReinstanced().AddUObject( this, &UVREditorMode::OnBlueprintReinstanced );
-
 	// Setup the asset container.
 	AssetContainer = &LoadAssetContainer();
 	check(AssetContainer != nullptr);
@@ -150,114 +148,6 @@ void UVREditorMode::Init()
 	bIsFullyInitialized = true;
 }
 
-// If we're recompiling Blueprints, we'll need to re-create and re-connect our interactors. 
-void UVREditorMode::OnBlueprintReinstanced()
-{
-	// Tear down
-	{
-		DestroyTransientActor( AvatarActor );
-		AvatarActor = nullptr;
-		FlashlightComponent = nullptr;
-	}
-
-	// Kill subsystems
-	if (UISystem != nullptr)
-	{
-		UISystem->Shutdown();
-		UISystem->MarkPendingKill();
-		UISystem = nullptr;
-	}
-
-	if (PlacementSystem != nullptr)
-	{
-		PlacementSystem->Shutdown();
-		PlacementSystem->MarkPendingKill();
-		PlacementSystem = nullptr;
-	}
-
-	if (TeleportActor != nullptr)
-	{
-		DestroyTransientActor( TeleportActor );
-		TeleportActor = nullptr;
-	}
-
-	if (AutoScalerSystem != nullptr)
-	{
-		AutoScalerSystem->Shutdown();
-		AutoScalerSystem->MarkPendingKill();
-		AutoScalerSystem = nullptr;
-	}
-
-	//	Rebuild
-	if (bActuallyUsingVR)
-	{
-		// Tell Slate to require a larger pixel distance threshold before the drag starts.  This is important for things
-		// like Content Browser drag and drop.
-		SavedEditorState.DragTriggerDistance = FSlateApplication::Get().GetDragTriggerDistance();
-		FSlateApplication::Get().SetDragTriggerDistance( VREd::SlateDragDistanceOverride->GetFloat() );
-
-		// When actually in VR, make sure the transform gizmo is big!
-		SavedEditorState.TransformGizmoScale = WorldInteraction->GetTransformGizmoScale();
-		WorldInteraction->SetTransformGizmoScale( GetDefault<UVRModeSettings>()->GizmoScale );
-		WorldInteraction->SetShouldSuppressExistingCursor( true );
-		WorldInteraction->SetInVR( true );
-	}
-
-	// Setup our avatar
-	if (AvatarActor == nullptr)
-	{
-		const bool bWithSceneComponent = true;
-		AvatarActor = SpawnTransientSceneActor<AVREditorAvatarActor>( TEXT( "AvatarActor" ), bWithSceneComponent );
-		AvatarActor->Init( this );
-
-		WorldInteraction->AddActorToExcludeFromHitTests( AvatarActor );
-	}
-
-		// Setup sub systems
-	{
-		// Setup world interaction
-		// We need input preprocessing for VR so that we can receive motion controller input without any viewports having 
-		// to be focused.  This is mainly because Slate UI injected into the 3D world can cause focus to be lost unexpectedly,
-		// but we need the user to still be able to interact with UI.
-		WorldInteraction->SetUseInputPreprocessor( true );
-
-		// Motion controllers
-//		LeftHandInteractor->SetControllerHandSide( FXRMotionControllerBase::LeftHandSourceId );
-
-//		RightHandInteractor->SetControllerHandSide( FXRMotionControllerBase::RightHandSourceId );
-
-		WorldInteraction->PairInteractors(GetHandInteractor(EControllerHand::Left), GetHandInteractor(EControllerHand::Right));
-
-		for (UVREditorInteractor* Interactor : Interactors)
-		{
-			Interactor->Init( this );
-			WorldInteraction->AddInteractor( Interactor );
-		}
-
-		// Setup the UI system
-		UISystem = NewObject<UVREditorUISystem>();
-		UISystem->Init( this );
-
-		PlacementSystem = NewObject<UVREditorPlacement>();
-		PlacementSystem->Init( this );
-
-		// Setup teleporter
-		TeleportActor = SpawnTransientSceneActor<AVREditorTeleporter>( TEXT( "Teleporter" ), true );
-		TeleportActor->Init( this );
-		WorldInteraction->AddActorToExcludeFromHitTests( TeleportActor );
-
-		// Setup autoscaler
-		AutoScalerSystem = NewObject<UVREditorAutoScaler>();
-		AutoScalerSystem->Init( this );
-	}
-
-	for (UVREditorInteractor* Interactor : Interactors)
-	{
-		Interactor->SetupComponent( AvatarActor );
-	}
-}
-
-
 /*
 * @EventName Editor.Usage.EnterVRMode
 *
@@ -288,8 +178,6 @@ void UVREditorMode::Shutdown()
 	GEnableVREditorHacks = false;
 
 	FEditorDelegates::EndPIE.RemoveAll(this);
-
-	GEditor->OnBlueprintReinstanced().RemoveAll( this );
 }
 
 void UVREditorMode::AllocateInteractors()
