Fixed crash when serializing VREditorFloatingUI because of uninitialized pointer

diff --git a/Engine/Source/Editor/VREditor/UI/VREditorFloatingUI.h b/Engine/Source/Editor/VREditor/UI/VREditorFloatingUI.h
--- a/Engine/Source/Editor/VREditor/UI/VREditorFloatingUI.h
+++ b/Engine/Source/Editor/VREditor/UI/VREditorFloatingUI.h
@@ -43,7 +43,7 @@ public:
 
 	/** NOT IN USE YET! Custom mesh to use for the VR window. */
 	UPROPERTY(EditAnywhere, BlueprintReadWrite, Category = "VR Mode UI")
-	UStaticMesh* PanelMesh;
+	UStaticMesh* PanelMesh = nullptr;
 
 	/** Optional override for "VREd.EditorUISize". Leave at 0 for default. */
 	UPROPERTY(EditAnywhere, BlueprintReadWrite, Category = "VR Mode UI")
