Only the correct UED VR controllers show up in the editor config. "None" option has been deleted (we always want a controller selection). VirtualScoutingController and VirtualScoutingTeleporter match in name.

diff --git a/Engine/Source/Editor/VREditor/Public/VRModeSettings.h b/Engine/Source/Editor/VREditor/Public/VRModeSettings.h
--- a/Engine/Source/Editor/VREditor/Public/VRModeSettings.h
+++ b/Engine/Source/Editor/VREditor/Public/VRModeSettings.h
@@ -80,12 +80,12 @@ public:
 	UPROPERTY(EditAnywhere, config, Category = "Motion Controllers", meta = (DisplayName = "Trigger Pressed Threshold (Oculus Touch)", ClampMin = 0.01, ClampMax = 1.0))
 	float TriggerPressedThreshold_Rift;
 
-	/** The class we'll create our interactors from */
-	UPROPERTY(EditAnywhere, config, Category = "Motion Controllers")
+	/** The controller to use when UnrealEd is in VR mode. Use VREditorInteractor get default editor behavior, or select a custom controller for special behavior */
+	UPROPERTY(EditAnywhere, config, NoClear, Category = "Motion Controllers")
 	TSoftClassPtr<UVREditorInteractor> InteractorClass;
 
-	/** The class we'll create our interactors from */
-	UPROPERTY( EditAnywhere, config, Category = "Motion Controllers" )
+	/** The teleporter to use when UnrealEd is in VR mode. Use VREditorTeleporter to get default editor behavior, or select a custom teleporter */
+	UPROPERTY( EditAnywhere, config, NoClear, Category = "Motion Controllers" )
 	TSubclassOf<AVREditorTeleporter> TeleporterClass;
 
 private:
diff --git a/Engine/Source/Editor/VREditor/VRScoutingInteractor.h b/Engine/Source/Editor/VREditor/VRScoutingInteractor.h
--- a/Engine/Source/Editor/VREditor/VRScoutingInteractor.h
+++ b/Engine/Source/Editor/VREditor/VRScoutingInteractor.h
@@ -16,7 +16,7 @@ class UStaticMeshSocket;
 /**
  * Represents the interactor in the world
  */
-UCLASS()
+UCLASS(Abstract)
 class VREDITOR_API UVRScoutingInteractor: public UVREditorInteractor
 {
 	GENERATED_BODY()
