Integrating CL#9475625 from dev-enterprise

- Changed VREditor TeleporterClass settings to be a SoftClassPtr instead of a TSubClassof to avoid crashing when pointing to an Editor Utility Blueprint

- Changed VR Scouting default Teleporter class to be an EditorUtilityBlueprint

diff --git a/Engine/Source/Editor/VREditor/Public/VRModeSettings.h b/Engine/Source/Editor/VREditor/Public/VRModeSettings.h
--- a/Engine/Source/Editor/VREditor/Public/VRModeSettings.h
+++ b/Engine/Source/Editor/VREditor/Public/VRModeSettings.h
@@ -86,7 +86,7 @@ public:
 
 	/** The teleporter to use when UnrealEd is in VR mode. Use VREditorTeleporter to get default editor behavior, or select a custom teleporter */
 	UPROPERTY( EditAnywhere, config, NoClear, Category = "Motion Controllers" )
-	TSubclassOf<AVREditorTeleporter> TeleporterClass;
+	TSoftClassPtr<AVREditorTeleporter> TeleporterClass;
 
 private:
 #if WITH_EDITOR
diff --git a/Engine/Source/Editor/VREditor/VREditorMode.cpp b/Engine/Source/Editor/VREditor/VREditorMode.cpp
--- a/Engine/Source/Editor/VREditor/VREditorMode.cpp
+++ b/Engine/Source/Editor/VREditor/VREditorMode.cpp
@@ -331,11 +331,12 @@ void UVREditorMode::Enter()
 		PlacementSystem->Init(this);
 
 		// Setup teleporter
-		const TSubclassOf<AVREditorTeleporter> TeleporterClass = GetDefault<UVRModeSettings>()->TeleporterClass;
+		const TSoftClassPtr<AVREditorTeleporter> TeleporterClassSoft = GetDefault<UVRModeSettings>()->TeleporterClass;
+		TeleporterClassSoft.LoadSynchronous();
 
-		if (TeleporterClass)
+		if (TeleporterClassSoft.IsValid())
 		{
-			TeleportActor = CastChecked<AVREditorTeleporter>( SpawnTransientSceneActor( TeleporterClass, TEXT( "Teleporter" ), true ) );
+			TeleportActor = CastChecked<AVREditorTeleporter>( SpawnTransientSceneActor(TeleporterClassSoft.Get(), TEXT( "Teleporter" ), true ) );
 		}
 
 		if( !TeleportActor )
