Add Exit VR button in VRScouting.
Fix warning when VRScouting is launched.

diff --git a/Engine/Source/Editor/UnrealEd/Private/EditorWorldExtension.cpp b/Engine/Source/Editor/UnrealEd/Private/EditorWorldExtension.cpp
--- a/Engine/Source/Editor/UnrealEd/Private/EditorWorldExtension.cpp
+++ b/Engine/Source/Editor/UnrealEd/Private/EditorWorldExtension.cpp
@@ -40,6 +40,10 @@ bool UEditorWorldExtension::InputAxis( FEditorViewportClient* InViewportClient,
 
 UWorld* UEditorWorldExtension::GetWorld() const
 {
+	if (OwningExtensionsCollection == nullptr)
+	{
+		return nullptr;
+	}
 	return OwningExtensionsCollection->GetWorld();
 }
 
diff --git a/Engine/Source/Editor/VREditor/UI/VREditorDockableCameraWindow.cpp b/Engine/Source/Editor/VREditor/UI/VREditorDockableCameraWindow.cpp
new file mode 100644
--- /dev/null
+++ b/Engine/Source/Editor/VREditor/UI/VREditorDockableCameraWindow.cpp
@@ -0,0 +1,10 @@
+// Copyright 1998-2019 Epic Games, Inc. All Rights Reserved.
+
+#include "VREditorDockableCameraWindow.h"
+#include "VREditorCameraWidgetComponent.h"
+
+AVREditorDockableCameraWindow::AVREditorDockableCameraWindow(const FObjectInitializer& ObjectInitializer) :
+	Super(ObjectInitializer.SetDefaultSubobjectClass<UVREditorCameraWidgetComponent>(TEXT("WidgetComponent")))
+{
+}
+
diff --git a/Engine/Source/Editor/VREditor/UI/VREditorDockableCameraWindow.h b/Engine/Source/Editor/VREditor/UI/VREditorDockableCameraWindow.h
new file mode 100644
--- /dev/null
+++ b/Engine/Source/Editor/VREditor/UI/VREditorDockableCameraWindow.h
@@ -0,0 +1,21 @@
+// Copyright 1998-2019 Epic Games, Inc. All Rights Reserved.
+
+#pragma once
+
+#include "VREditorDockableWindow.h"
+#include "VREditorDockableCameraWindow.generated.h"
+
+/**
+ * A specialized dockable window for camera viewfinders in VR that applies the correct material
+ */
+UCLASS()
+class AVREditorDockableCameraWindow : public AVREditorDockableWindow
+{
+	GENERATED_BODY()
+	
+public:
+
+	/** Default constructor */
+  AVREditorDockableCameraWindow(const FObjectInitializer& ObjectInitializer);
+};
+
diff --git a/Engine/Source/Editor/VREditor/UI/VREditorUISystem.cpp b/Engine/Source/Editor/VREditor/UI/VREditorUISystem.cpp
--- a/Engine/Source/Editor/VREditor/UI/VREditorUISystem.cpp
+++ b/Engine/Source/Editor/VREditor/UI/VREditorUISystem.cpp
@@ -40,6 +40,7 @@
 #include "VREditorFloatingUI.h"
 #include "VREditorRadialFloatingUI.h"
 #include "VREditorDockableWindow.h"
+#include "VREditorDockableCameraWindow.h"
 #include "ViewportInteractionTypes.h"
 #include "IHeadMountedDisplay.h"
 #include "ViewportWorldInteraction.h"
@@ -2597,7 +2598,7 @@ void UVREditorUISystem::UpdateDetachedActorPreviewUI(TSharedRef<SWidget> InWidge
 	{
 		const FName TabIdentifier = NAME_None;	// No tab for us!
 		const bool bWithSceneComponent = false;
-		PreviewPanel = GetOwner().SpawnTransientSceneActor<AVREditorDockableWindow>(TEXT("ActorPreviewUI"), bWithSceneComponent);
+		PreviewPanel = GetOwner().SpawnTransientSceneActor<AVREditorDockableCameraWindow>(TEXT("ActorDetachablePreviewUI"), bWithSceneComponent);
 		PreviewPanel->CreationContext.bNoCloseButton = true;
 		FloatingUIs.Add(PreviewPanelID, PreviewPanel);
 		AActor* Actor = GEditor->GetSelectedActors()->GetBottom<AActor>();
@@ -2665,7 +2666,7 @@ void UVREditorUISystem::UpdateExternalUMGUI(const FVREditorFloatingUICreationCon
 	{
 		const FName TabIdentifier = NAME_None;	// No tab for us!
 		const bool bWithSceneComponent = false;
-  		ExternalPanel = GetOwner().SpawnTransientSceneActor<AVREditorDockableWindow>(TEXT("ActorPreviewUI"), bWithSceneComponent);	
+  		ExternalPanel = GetOwner().SpawnTransientSceneActor<AVREditorDockableWindow>(TEXT("ActorPreviewUMGUI"), bWithSceneComponent);	
 		
 		// AVREditorDockableWindow does most of its setup in PostActorCreated(), but the CreationContext hasn't been set yet. 
 		// We let it set its defaults, then override the custom settings here.
diff --git a/Engine/Source/Editor/ViewportInteraction/ViewportInteractor.cpp b/Engine/Source/Editor/ViewportInteraction/ViewportInteractor.cpp
--- a/Engine/Source/Editor/ViewportInteraction/ViewportInteractor.cpp
+++ b/Engine/Source/Editor/ViewportInteraction/ViewportInteractor.cpp
@@ -694,6 +694,7 @@ FHitResult UViewportInteractor::GetHitResultFromLaserPointer( TArray<AActor*>* O
 			// @todo viewportinteraction: We'll need to device a way to allow volume wire bounds to be selectable using this system
 			static TArray<AActor*> VolumeActors;
 			VolumeActors.Reset();
+			if (WorldInteraction->GetWorld() != nullptr)
 			{
 				for( TActorIterator<AVolume> It( WorldInteraction->GetWorld(), AVolume::StaticClass() ); It; ++It )
 				{
