When an asset editor closing is requested due to force delete, close the asset editor UI if we are in VR Mode.

diff --git a/Engine/Source/Editor/VREditor/UI/VREditorUISystem.cpp b/Engine/Source/Editor/VREditor/UI/VREditorUISystem.cpp
--- a/Engine/Source/Editor/VREditor/UI/VREditorUISystem.cpp
+++ b/Engine/Source/Editor/VREditor/UI/VREditorUISystem.cpp
@@ -215,6 +215,8 @@ void UVREditorUISystem::Init(UVREditorMode* InVRMode)
 	FVREditorActionCallbacks::SelectingCandidateActorsText = FVREditorActionCallbacks::GetSelectingCandidateActorsText();
 
 	GLevelEditorModeTools().OnEditorModeIDChanged().AddUObject(this, &UVREditorUISystem::HandleEditorModeChanged);
+	AssetEditorCloseDelegate = GEditor->GetEditorSubsystem<UAssetEditorSubsystem>()->OnAssetEditorRequestClose().AddUObject(this, &UVREditorUISystem::OnCloseAssetEditor);
+
 }
 
 void UVREditorUISystem::UpdateInteractors( )
@@ -295,7 +297,7 @@ void UVREditorUISystem::Shutdown()
 	// Remove the proxy tab manager, we don't want to steal tabs any more.
 	FGlobalTabmanager::Get()->SetProxyTabManager(TSharedPtr<FProxyTabmanager>());
 	GEditor->GetEditorSubsystem<UAssetEditorSubsystem>()->OnAssetEditorOpened().RemoveAll(this);
-
+	GEditor->GetEditorSubsystem<UAssetEditorSubsystem>()->OnAssetEditorRequestClose().RemoveAll(this);
 	VRMode = nullptr;
 	DraggingUI = nullptr;
 	ColorPickerUI = nullptr;
@@ -2804,4 +2806,14 @@ void UVREditorUISystem::ToggledDebugMode(bool bDebugModeEnabled)
 	}
 }
 
+void UVREditorUISystem::OnCloseAssetEditor(UObject* Asset, EAssetEditorCloseReason CloseReason)
+{
+	AVREditorFloatingUI* AssetEditorPanel = GetPanel(TabManagerPanelID);
+	if(AssetEditorPanel)
+	{
+		AssetEditorPanel->SetSlateWidget(SNullWidget::NullWidget);
+		AssetEditorPanel->ShowUI(false);
+	}
+}
+
 #undef LOCTEXT_NAMESPACE
diff --git a/Engine/Source/Editor/VREditor/UI/VREditorUISystem.h b/Engine/Source/Editor/VREditor/UI/VREditorUISystem.h
--- a/Engine/Source/Editor/VREditor/UI/VREditorUISystem.h
+++ b/Engine/Source/Editor/VREditor/UI/VREditorUISystem.h
@@ -32,6 +32,8 @@ typedef FName FEditorModeID;
 
 typedef FName VREditorPanelID;
 
+enum class EAssetEditorCloseReason : uint8;
+
 /** Stores the animation playback state of a VR UI element */
 enum class EVREditorAnimationState : uint8
 {
@@ -332,6 +334,7 @@ protected:
 	/** When VR Mode debug was toggled. */
 	void ToggledDebugMode(bool bDebugModeEnabled);
 
+	void OnCloseAssetEditor(UObject* Asset, EAssetEditorCloseReason CloseReason);
 protected:
 
 	/** Owning object */
@@ -439,4 +442,6 @@ protected:
 	/** When started dragging from the radial menu we want the analog stick to be reset before the user is allowed to scale the panel. 
 	    Otherwise the panel will immediately start scaling because the user is using the analog stick to aim at the radial menu items. */
 	bool bPanelCanScale;
+
+	FDelegateHandle AssetEditorCloseDelegate;
 };
\ No newline at end of file
