Check if XRSystem is valid before enabling stereo

diff --git a/Engine/Source/Editor/VREditor/VREditorMode.cpp b/Engine/Source/Editor/VREditor/VREditorMode.cpp
--- a/Engine/Source/Editor/VREditor/VREditorMode.cpp
+++ b/Engine/Source/Editor/VREditor/VREditorMode.cpp
@@ -1113,7 +1113,7 @@ void UVREditorMode::StartViewport(TSharedPtr<SLevelViewport> Viewport)
 		VRViewportClient.Invalidate(true);
 	}
 
-	if (bActuallyUsingVR)
+	if (bActuallyUsingVR && GEngine->XRSystem.IsValid())
 	{
 		Viewport->EnableStereoRendering( bActuallyUsingVR );
 		Viewport->SetRenderDirectlyToWindow( bActuallyUsingVR );
