Warn VR editor users that the new editor mode UI is not compatible with VR Editor and offer to revert to the old UI and restart the editor for them

diff --git a/Engine/Source/Editor/VREditor/VREditorModeManager.cpp b/Engine/Source/Editor/VREditor/VREditorModeManager.cpp
--- a/Engine/Source/Editor/VREditor/VREditorModeManager.cpp
+++ b/Engine/Source/Editor/VREditor/VREditorModeManager.cpp
@@ -131,17 +131,43 @@ void FVREditorModeManager::EnableVREditor( const bool bEnable, const bool bForce
 	{
 		if( bEnable && ( IsVREditorAvailable() || bForceWithoutHMD ))
 		{
-			FSuppressableWarningDialog::FSetupInfo SetupInfo(LOCTEXT("VRModeEntry_Message", "VR Mode enables you to work on your project in virtual reality using motion controllers. This feature is still under development, so you may experience bugs or crashes while using it."),
-				LOCTEXT("VRModeEntry_Title", "Entering VR Mode - Experimental"), "Warning_VRModeEntry", GEditorSettingsIni);
+			UEditorStyleSettings* StyleSettings = GetMutableDefault<UEditorStyleSettings>();
+
+			if (StyleSettings && !StyleSettings->bEnableLegacyEditorModeUI)
+			{
+				FSuppressableWarningDialog::FSetupInfo SetupInfo(LOCTEXT("VRModeLegacyModeUIEntry_Message", "VR Mode currently requires that legacy editor mode UI be enabled.  Without this, modes like mesh paint, landscape, and foliage will not function.  Enable Legacy editor mode UI (Requires restart)?"),
+					LOCTEXT("VRModeEntry_Title", "Entering VR Mode - Experimental"), "Warning_VRModeLegacyModeUIEntry", GEditorSettingsIni);
+
+				SetupInfo.ConfirmText = LOCTEXT("VRModeLegacyModeUIEntry_ConfirmText", "Enable and Restart");
+				SetupInfo.CancelText = LOCTEXT("VRModeLegacyModeUIEntry_CancelText", "Don't Enable");
+				SetupInfo.bDefaultToSuppressInTheFuture = false;
+
+				FSuppressableWarningDialog VRModeVRModeLegacyModeUIWarning(SetupInfo);
+
+				if (VRModeEntryWarning.ShowModal() != FSuppressableWarningDialog::Cancel)
+				{
+					StyleSettings->bEnableLegacyEditorModeUI = true;
+					StyleSettings->SaveConfig();
+					FUnrealEdMisc::RestartEditor(true);
+					return;
+				}
+
+			}
 
-			SetupInfo.ConfirmText = LOCTEXT("VRModeEntry_ConfirmText", "Continue");
-			SetupInfo.CancelText = LOCTEXT("VRModeEntry_CancelText", "Cancel");
-			SetupInfo.bDefaultToSuppressInTheFuture = true;
-			FSuppressableWarningDialog VRModeEntryWarning(SetupInfo);
 
-			if (VRModeEntryWarning.ShowModal() != FSuppressableWarningDialog::Cancel)
 			{
-				StartVREditorMode(bForceWithoutHMD);
+				FSuppressableWarningDialog::FSetupInfo SetupInfo(LOCTEXT("VRModeEntry_Message", "VR Mode enables you to work on your project in virtual reality using motion controllers. This feature is still under development, so you may experience bugs or crashes while using it."),
+					LOCTEXT("VRModeEntry_Title", "Entering VR Mode - Experimental"), "Warning_VRModeEntry", GEditorSettingsIni);
+
+				SetupInfo.ConfirmText = LOCTEXT("VRModeEntry_ConfirmText", "Continue");
+				SetupInfo.CancelText = LOCTEXT("VRModeEntry_CancelText", "Cancel");
+				SetupInfo.bDefaultToSuppressInTheFuture = true;
+				FSuppressableWarningDialog VRModeEntryWarning(SetupInfo);
+
+				if (VRModeEntryWarning.ShowModal() != FSuppressableWarningDialog::Cancel)
+				{
+					StartVREditorMode(bForceWithoutHMD);
+				}
 			}
 		}
 		else if( !bEnable )
