Fix CIS

diff --git a/Engine/Source/Editor/VREditor/VREditorModeManager.cpp b/Engine/Source/Editor/VREditor/VREditorModeManager.cpp
--- a/Engine/Source/Editor/VREditor/VREditorModeManager.cpp
+++ b/Engine/Source/Editor/VREditor/VREditorModeManager.cpp
@@ -21,6 +21,7 @@
 #include "ProjectDescriptor.h"
 #include "Interfaces/IProjectManager.h"
 #include "UnrealEdMisc.h"
+#include "Classes/EditorStyleSettings.h"
 
 #define LOCTEXT_NAMESPACE "VREditor"
 
@@ -144,11 +145,11 @@ void FVREditorModeManager::EnableVREditor( const bool bEnable, const bool bForce
 
 				FSuppressableWarningDialog VRModeVRModeLegacyModeUIWarning(SetupInfo);
 
-				if (VRModeEntryWarning.ShowModal() != FSuppressableWarningDialog::Cancel)
+				if (VRModeVRModeLegacyModeUIWarning.ShowModal() != FSuppressableWarningDialog::Cancel)
 				{
 					StyleSettings->bEnableLegacyEditorModeUI = true;
 					StyleSettings->SaveConfig();
-					FUnrealEdMisc::RestartEditor(true);
+					FUnrealEdMisc::Get().RestartEditor(true);
 					return;
 				}
 
