Fix for crash on exiting VR mode from the taskbar.

diff --git a/Engine/Source/Editor/VREditor/UI/VREditorFloatingUI.cpp b/Engine/Source/Editor/VREditor/UI/VREditorFloatingUI.cpp
--- a/Engine/Source/Editor/VREditor/UI/VREditorFloatingUI.cpp
+++ b/Engine/Source/Editor/VREditor/UI/VREditorFloatingUI.cpp
@@ -200,14 +200,6 @@ void AVREditorFloatingUI::TickManually(float DeltaTime)
 	}
 }
 
-void AVREditorFloatingUI::Destroyed()
-{
-	CleanupWidgetReferences();
-
-
-	Super::Destroyed();
-}
-
 void AVREditorFloatingUI::CleanupWidgetReferences()
 {
 	if (WidgetComponent != nullptr)
diff --git a/Engine/Source/Editor/VREditor/UI/VREditorFloatingUI.h b/Engine/Source/Editor/VREditor/UI/VREditorFloatingUI.h
--- a/Engine/Source/Editor/VREditor/UI/VREditorFloatingUI.h
+++ b/Engine/Source/Editor/VREditor/UI/VREditorFloatingUI.h
@@ -129,7 +129,6 @@ public:
 
 	/** AActor overrides */
 	virtual void BeginDestroy() override;
-	virtual void Destroyed() override;
 
 	void CleanupWidgetReferences();
 
