Now that we are not double-clearing the VR widgets, add back the slate widget pointer reset.

diff --git a/Engine/Source/Editor/VREditor/UI/VREditorFloatingUI.cpp b/Engine/Source/Editor/VREditor/UI/VREditorFloatingUI.cpp
--- a/Engine/Source/Editor/VREditor/UI/VREditorFloatingUI.cpp
+++ b/Engine/Source/Editor/VREditor/UI/VREditorFloatingUI.cpp
@@ -210,6 +210,11 @@ void AVREditorFloatingUI::CleanupWidgetReferences()
 		WidgetComponent = nullptr;
 	}
 
+	if (SlateWidget.IsValid())
+	{
+		SlateWidget.Reset();
+	}
+
 	// @todo vreditor unreal: UMG has a bug that prevents you from re-using the user widget for a new widget component
 	// after a previous widget component that was using it was destroyed
 	if (UserWidget != nullptr)
