Adding a check that we are using the default interactors before warning about legacy UI. This prevents unneeded warnings for Virtual Production.

diff --git a/Engine/Source/Editor/VREditor/VREditorModeManager.cpp b/Engine/Source/Editor/VREditor/VREditorModeManager.cpp
--- a/Engine/Source/Editor/VREditor/VREditorModeManager.cpp
+++ b/Engine/Source/Editor/VREditor/VREditorModeManager.cpp
@@ -22,6 +22,7 @@
 #include "Interfaces/IProjectManager.h"
 #include "UnrealEdMisc.h"
 #include "Classes/EditorStyleSettings.h"
+#include "VREditorInteractor.h"
 
 #define LOCTEXT_NAMESPACE "VREditor"
 
@@ -133,8 +134,19 @@ void FVREditorModeManager::EnableVREditor( const bool bEnable, const bool bForce
 		if( bEnable && ( IsVREditorAvailable() || bForceWithoutHMD ))
 		{
 			UEditorStyleSettings* StyleSettings = GetMutableDefault<UEditorStyleSettings>();
+			const UVRModeSettings* VRModeSettings = GetDefault<UVRModeSettings>();
+			bool bUsingDefaultInteractors = true;
+			if (VRModeSettings)
+			{
+				const TSoftClassPtr<UVREditorInteractor> InteractorClassSoft = VRModeSettings->InteractorClass;
+				InteractorClassSoft.LoadSynchronous();
 
-			if (StyleSettings && !StyleSettings->bEnableLegacyEditorModeUI)
+				if (InteractorClassSoft.IsValid())
+				{
+					bUsingDefaultInteractors = (InteractorClassSoft.Get() == UVREditorInteractor::StaticClass());
+				}
+			}
+			if ( bUsingDefaultInteractors && StyleSettings && !StyleSettings->bEnableLegacyEditorModeUI)
 			{
 				FSuppressableWarningDialog::FSetupInfo SetupInfo(LOCTEXT("VRModeLegacyModeUIEntry_Message", "VR Mode currently requires that legacy editor mode UI be enabled.  Without this, modes like mesh paint, landscape, and foliage will not function.  Enable Legacy editor mode UI (Requires restart)?"),
 					LOCTEXT("VRModeEntry_Title", "Entering VR Mode - Experimental"), "Warning_VRModeLegacyModeUIEntry", GEditorSettingsIni);
