Fix crash when in VR and MultiUser, where other machines may crash when one starts VR Mode.

This was happening because success of GetTransformAndForwardVector was not being checked, hence uninitialized location values were being sent in the MultiUser/Concert event to the other clients, and eventually cause NaN values in the bounds of laser spline mesh components of the VR presence actor.

diff --git a/Engine/Source/Editor/VREditor/VREditorInteractor.cpp b/Engine/Source/Editor/VREditor/VREditorInteractor.cpp
--- a/Engine/Source/Editor/VREditor/VREditorInteractor.cpp
+++ b/Engine/Source/Editor/VREditor/VREditorInteractor.cpp
@@ -526,7 +526,7 @@ void UVREditorInteractor::Tick_Implementation( const float DeltaTime )
 		const bool bDraggingWorld = (InteractorData.DraggingMode == EViewportInteractionDraggingMode::World ||
 			(GetOtherInteractor() != nullptr && GetOtherInteractor()->GetInteractorData().DraggingMode == EViewportInteractionDraggingMode::World && InteractorData.DraggingMode == EViewportInteractionDraggingMode::AssistingDrag));
 
-		FVector LaserPointerStart, LaserPointerEnd;
+		FVector LaserPointerStart(0), LaserPointerEnd(0);
 		const bool bHasLaser = GetLaserPointer(/* Out */ LaserPointerStart, /* Out */ LaserPointerEnd, bEvenIfBlocked );
 		if (bForceShowLaser || (bHasLaser && !bDraggingWorld))
 		{
@@ -588,13 +588,16 @@ void UVREditorInteractor::Tick_Implementation( const float DeltaTime )
 			// Get the hand transform and forward vector.
 			FTransform InteractorTransform;
 			FVector InteractorForwardVector;
-			GetTransformAndForwardVector( /* Out */ InteractorTransform, /* Out */ InteractorForwardVector );
-			InteractorForwardVector.Normalize();
 
-			// Offset the start point of the laser.
-			LaserPointerStart = InteractorTransform.GetLocation() + (InteractorForwardVector * LaserPointerStartOffset);
+			if (GetTransformAndForwardVector( /* Out */ InteractorTransform, /* Out */ InteractorForwardVector))
+			{
+				InteractorForwardVector.Normalize();
+
+				// Offset the start point of the laser.
+				LaserPointerStart = InteractorTransform.GetLocation() + (InteractorForwardVector * LaserPointerStartOffset);
 
-			UpdateSplineLaser( LaserPointerStart, LaserPointerEnd, InteractorForwardVector );
+				UpdateSplineLaser(LaserPointerStart, LaserPointerEnd, InteractorForwardVector);
+			}
 		}
 
 		bForceShowLaser = false;
