When the IE_Released Event happens, it also seems to send a VRActionTypes::TrackpadPositionX (and Y) at the same time, and previously that did not happen.  If we avoid processing this if the Delta is zero, it seems to alleviate the issue.

diff --git a/Engine/Source/Editor/VREditor/VREditorInteractor.cpp b/Engine/Source/Editor/VREditor/VREditorInteractor.cpp
--- a/Engine/Source/Editor/VREditor/VREditorInteractor.cpp
+++ b/Engine/Source/Editor/VREditor/VREditorInteractor.cpp
@@ -1035,26 +1035,29 @@ void UVREditorInteractor::HandleInputAxis( FEditorViewportClient& ViewportClient
 
 	if (!bOutWasHandled)
 	{
-		if (Action.ActionType == VRActionTypes::TrackpadPositionX)
+		if (Delta != 0)
 		{
-			LastTrackpadPosition.X = bIsTrackpadPositionValid[0] ? TrackpadPosition.X : Delta;
-			LastTrackpadPositionUpdateTime = FTimespan::FromSeconds( FPlatformTime::Seconds() );
-			TrackpadPosition.X = Delta;
-			bIsTrackpadPositionValid[0] = true;
-		}
-
-		if (Action.ActionType == VRActionTypes::TrackpadPositionY)
-		{
-			float DeltaAxis = Delta;
-			if (VREd::InvertTrackpadVertical->GetInt() != 0)
+			if (Action.ActionType == VRActionTypes::TrackpadPositionX)
 			{
-				DeltaAxis = -DeltaAxis;	// Y axis is inverted from HMD
+				LastTrackpadPosition.X = bIsTrackpadPositionValid[0] ? TrackpadPosition.X : Delta;
+				LastTrackpadPositionUpdateTime = FTimespan::FromSeconds(FPlatformTime::Seconds());
+				TrackpadPosition.X = Delta;
+				bIsTrackpadPositionValid[0] = true;
 			}
 
-			LastTrackpadPosition.Y = bIsTrackpadPositionValid[1] ? TrackpadPosition.Y : DeltaAxis;
-			LastTrackpadPositionUpdateTime = FTimespan::FromSeconds( FPlatformTime::Seconds() );
-			TrackpadPosition.Y = DeltaAxis;
-			bIsTrackpadPositionValid[1] = true;
+			if (Action.ActionType == VRActionTypes::TrackpadPositionY)
+			{
+				float DeltaAxis = Delta;
+				if (VREd::InvertTrackpadVertical->GetInt() != 0)
+				{
+					DeltaAxis = -DeltaAxis;	// Y axis is inverted from HMD
+				}
+
+				LastTrackpadPosition.Y = bIsTrackpadPositionValid[1] ? TrackpadPosition.Y : DeltaAxis;
+				LastTrackpadPositionUpdateTime = FTimespan::FromSeconds(FPlatformTime::Seconds());
+				TrackpadPosition.Y = DeltaAxis;
+				bIsTrackpadPositionValid[1] = true;
+			}
 		}
 	}
 
