Fix for #JIRA UE-91265 Crash when trying to save after exiting vr scouting #rb alejandro.arango lauren.barnes #lockdown simon.tourangeau

diff --git a/Engine/Source/Editor/VREditor/UI/VREditorFloatingUI.cpp b/Engine/Source/Editor/VREditor/UI/VREditorFloatingUI.cpp
--- a/Engine/Source/Editor/VREditor/UI/VREditorFloatingUI.cpp
+++ b/Engine/Source/Editor/VREditor/UI/VREditorFloatingUI.cpp
@@ -200,6 +200,16 @@ void AVREditorFloatingUI::TickManually(float DeltaTime)
 	}
 }
 
+void AVREditorFloatingUI::Destroyed()
+{
+	if (!IsPendingKill())
+	{
+		CleanupWidgetReferences();
+	}
+
+	Super::Destroyed();
+}
+
 void AVREditorFloatingUI::CleanupWidgetReferences()
 {
 	if (WidgetComponent != nullptr)
diff --git a/Engine/Source/Editor/VREditor/UI/VREditorFloatingUI.h b/Engine/Source/Editor/VREditor/UI/VREditorFloatingUI.h
--- a/Engine/Source/Editor/VREditor/UI/VREditorFloatingUI.h
+++ b/Engine/Source/Editor/VREditor/UI/VREditorFloatingUI.h
@@ -129,6 +129,7 @@ public:
 
 	/** AActor overrides */
 	virtual void BeginDestroy() override;
+	virtual void Destroyed() override;
 
 	void CleanupWidgetReferences();
 
