Removes PIE, SIE, Editor Modes (Foliage, Landscape, Mesh Paint) and their associated radial menu options.

diff --git a/Engine/Plugins/Editor/MeshEditor/Source/MeshEditor/MeshEditor.Build.cs b/Engine/Plugins/Editor/MeshEditor/Source/MeshEditor/MeshEditor.Build.cs
--- a/Engine/Plugins/Editor/MeshEditor/Source/MeshEditor/MeshEditor.Build.cs
+++ b/Engine/Plugins/Editor/MeshEditor/Source/MeshEditor/MeshEditor.Build.cs
@@ -21,8 +21,6 @@ namespace UnrealBuildTool.Rules
 					"EditableMesh",
                     "MeshDescription",
 					"StaticMeshDescription",
-					"ViewportInteraction",
-                    "VREditor",
 					"Projects",
                     "RHI",
 					"LevelEditor",
diff --git a/Engine/Plugins/MeshPainting/Source/MeshPaintingToolset/Public/MeshPaintHelpers.h b/Engine/Plugins/MeshPainting/Source/MeshPaintingToolset/Public/MeshPaintHelpers.h
--- a/Engine/Plugins/MeshPainting/Source/MeshPaintingToolset/Public/MeshPaintHelpers.h
+++ b/Engine/Plugins/MeshPainting/Source/MeshPaintingToolset/Public/MeshPaintHelpers.h
@@ -164,10 +164,6 @@ public:
 
 	static void ApplyFillWithMask(FColor& InOutColor, const FColor& MaskColor, const FColor& FillColor);
 
-		
-	/** Checks whether or not the given Viewport Client is a VR editor viewport client */
-	static bool IsInVRMode(const FEditorViewportClient* ViewportClient);
-
 	/** Forces the component to render LOD level at LODIndex instead of the view-based LOD level ( X = 0 means do not force the LOD, X > 0 means force the lod to X - 1 ) */
 	static void ForceRenderMeshLOD(UMeshComponent* Component, int32 LODIndex);
 
diff --git a/Engine/Source/Editor/ContentBrowser/Private/SContentBrowser.cpp b/Engine/Source/Editor/ContentBrowser/Private/SContentBrowser.cpp
--- a/Engine/Source/Editor/ContentBrowser/Private/SContentBrowser.cpp
+++ b/Engine/Source/Editor/ContentBrowser/Private/SContentBrowser.cpp
@@ -69,7 +69,6 @@
 #include "ContentBrowserDataSubsystem.h"
 
 #include "Brushes/SlateColorBrush.h"
-#include "IVREditorModule.h"
 
 
 #define LOCTEXT_NAMESPACE "ContentBrowser"
@@ -2511,11 +2510,6 @@ void SContentBrowser::HandleResaveAllCurrentFolderCommand() const
 
 bool SContentBrowser::HandleDeleteCommandCanExecute() const
 {
-	if (IVREditorModule::Get().IsVREditorModeActive())
-	{
-		return false;
-	}
-
 	// The order of these conditions are carefully crafted to match the logic of the context menu summoning, as this callback 
 	// is shared between the path and asset views, and is given zero context as to which one is making the request
 	// Change this logic at your peril, lest the the dominoes fall like a house of cards (checkmate)
diff --git a/Engine/Source/Editor/FoliageEdit/FoliageEdit.build.cs b/Engine/Source/Editor/FoliageEdit/FoliageEdit.build.cs
--- a/Engine/Source/Editor/FoliageEdit/FoliageEdit.build.cs
+++ b/Engine/Source/Editor/FoliageEdit/FoliageEdit.build.cs
@@ -25,16 +25,6 @@ public class FoliageEdit : ModuleRules
                 "DetailCustomizations",
                 "AssetTools",
                 "Foliage",
-				"ViewportInteraction",
-				"VREditor"
-			}
-		);
-
-		CircularlyReferencedDependentModules.AddRange(
-			new string[]
-			{
-				"ViewportInteraction",
-				"VREditor"
 			}
 		);
 
diff --git a/Engine/Source/Editor/FoliageEdit/Private/FoliageEdMode.cpp b/Engine/Source/Editor/FoliageEdit/Private/FoliageEdMode.cpp
--- a/Engine/Source/Editor/FoliageEdit/Private/FoliageEdMode.cpp
+++ b/Engine/Source/Editor/FoliageEdit/Private/FoliageEdMode.cpp
@@ -48,11 +48,6 @@
 #include "Components/SplineMeshComponent.h"
 #include "Components/BrushComponent.h"
 
-// VR Editor
-#include "EditorWorldExtension.h"
-#include "VREditorMode.h"
-#include "ViewportWorldInteraction.h"
-#include "VREditorInteractor.h"
 #include "LevelUtils.h"
 #include "Widgets/Notifications/SNotificationList.h"
 #include "Framework/Notifications/NotificationManager.h"
@@ -230,7 +225,6 @@ FEdModeFoliage::FEdModeFoliage()
 	, bCanAltDrag(false)
 	, bAdjustBrushRadius(false)
 	, FoliageMeshListSortMode(EColumnSortMode::Ascending)
-	, FoliageInteractor(nullptr)
 	, UpdateSelectionCounter(0)
 	, bHasDeferredSelectionNotification(false)
 	, bMoving(false)
@@ -454,22 +448,6 @@ void FEdModeFoliage::Enter()
 		}
 	}
 
-	// Register for VR input events
-	UViewportWorldInteraction* ViewportWorldInteraction = Cast<UViewportWorldInteraction>( GEditor->GetEditorWorldExtensionsManager()->GetEditorWorldExtensions( GetWorld() )->FindExtension( UViewportWorldInteraction::StaticClass() ) );
-	if (ViewportWorldInteraction != nullptr)
-	{
-		ViewportWorldInteraction->OnViewportInteractionInputAction().RemoveAll(this);
-		ViewportWorldInteraction->OnViewportInteractionInputAction().AddRaw(this, &FEdModeFoliage::OnVRAction);
-
-		ViewportWorldInteraction->OnViewportInteractionHoverUpdate().RemoveAll(this);
-		ViewportWorldInteraction->OnViewportInteractionHoverUpdate().AddRaw(this, &FEdModeFoliage::OnVRHoverUpdate);
-
-		// Hide the VR transform gizmo while we're in foliage mode.  It sort of gets in the way of painting.
-		ViewportWorldInteraction->SetTransformGizmoVisible( false );
-
-		SetBrushOpacity(VREd::FoliageOpacity->GetFloat());
-	}
-
 	// Make sure the brush is visible.
 	SphereBrushComponent->SetVisibility(true);
 }
@@ -487,23 +465,6 @@ void FEdModeFoliage::Exit()
 	}
 	check(!GEditor->IsTransactionActive());
 
-	// Unregister VR mode from event handlers
-	{
-		UViewportWorldInteraction* ViewportWorldInteraction = Cast<UViewportWorldInteraction>( GEditor->GetEditorWorldExtensionsManager()->GetEditorWorldExtensions( GetWorld() )->FindExtension( UViewportWorldInteraction::StaticClass() ) );
-		if (ViewportWorldInteraction != nullptr)
-		{
-			// Restore the transform gizmo visibility
-			ViewportWorldInteraction->SetTransformGizmoVisible( true );
-
-			ViewportWorldInteraction->OnViewportInteractionInputAction().RemoveAll(this);
-			ViewportWorldInteraction->OnViewportInteractionHoverUpdate().RemoveAll(this);
-			FoliageInteractor = nullptr;
-
-			// Reset the brush opacity to default.
-			SetBrushOpacity(DefaultBrushOpacity);
-		}
-	}
-
 	FToolkitManager::Get().CloseToolkit(Toolkit.ToSharedRef());
 	Toolkit.Reset();
 
@@ -603,138 +564,6 @@ void FEdModeFoliage::OnEndPIE(const bool bIsSimulating)
 	}
 }
 
-void FEdModeFoliage::OnVRHoverUpdate(UViewportInteractor* Interactor, FVector& HoverImpactPoint, bool& bWasHandled)
-{
-	UVREditorMode* VREditorMode = Cast<UVREditorMode>( GEditor->GetEditorWorldExtensionsManager()->GetEditorWorldExtensions( GetWorld() )->FindExtension( UVREditorMode::StaticClass() ) );
-	if( VREditorMode != nullptr && VREditorMode->IsFullyInitialized() )
-	{
-		// Check if we're hovering over UI. If so, stop painting so we don't display the preview brush sphere
-		if (FoliageInteractor && (FoliageInteractor->IsHoveringOverPriorityType() || FoliageInteractor->GetDraggingMode() != EViewportInteractionDraggingMode::Nothing))
-		{
-			EndFoliageBrushTrace();
-			FoliageInteractor = nullptr;
-		}
-		// If there isn't currently a foliage interactor and we are hovering over something valid
-		else if (FoliageInteractor == nullptr && !Interactor->IsHoveringOverPriorityType() && Interactor->GetHitResultFromLaserPointer().GetActor() != nullptr)
-		{
-			FoliageInteractor = Interactor;
-		}
-		// If we aren't hovering over something valid and the tool isn't active
-		else if (Interactor->GetHitResultFromLaserPointer().GetActor() == nullptr && !bToolActive)
-		{
-			FoliageInteractor = nullptr;
-		}
-
-		// Skip other interactors if we are painting with one
-		if (FoliageInteractor && Interactor == FoliageInteractor)
-		{
-			// Go ahead and paint immediately
-			FVector LaserPointerStart, LaserPointerEnd;
-			if (FoliageInteractor->GetLaserPointer( /* Out */ LaserPointerStart, /* Out */ LaserPointerEnd))
-			{
-				const FVector LaserPointerDirection = (LaserPointerEnd - LaserPointerStart).GetSafeNormal();
-
-				FoliageBrushTrace(nullptr, LaserPointerStart, LaserPointerDirection);
-
-				if( bBrushTraceValid )
-				{
-					HoverImpactPoint = BrushLocation;
-					bWasHandled = true;
-				}
-			}
-		}
-		const bool bBrushMeshVisible = !(FoliageInteractor == nullptr || Interactor->GetDraggingMode() != EViewportInteractionDraggingMode::Nothing);
-		SphereBrushComponent->SetVisibility(bBrushMeshVisible);
-	}
-}
-
-void FEdModeFoliage::OnVRAction(class FEditorViewportClient& ViewportClient, UViewportInteractor* Interactor, const FViewportActionKeyInput& Action, bool& bOutIsInputCaptured, bool& bWasHandled)
-{
-	UVREditorMode* VREditorMode = Cast<UVREditorMode>( GEditor->GetEditorWorldExtensionsManager()->GetEditorWorldExtensions( GetWorld() )->FindExtension( UVREditorMode::StaticClass() ) );
-	if (VREditorMode != nullptr && Interactor != nullptr && Interactor->GetDraggingMode() == EViewportInteractionDraggingMode::Nothing)
-	{
-		const UVREditorInteractor* VREditorInteractor = Cast<UVREditorInteractor>(Interactor);
-		if (Action.ActionType == ViewportWorldActionTypes::SelectAndMove && (VREditorInteractor == nullptr || !VREditorMode->IsShowingRadialMenu(VREditorInteractor)))
-		{
-			if (Action.Event == IE_Pressed && !Interactor->IsHoveringOverPriorityType())
-			{
-				bWasHandled = true;
-				bOutIsInputCaptured = true;
-
-				// Go ahead and paint immediately
-				FVector LaserPointerStart, LaserPointerEnd;
-				if (Interactor->GetLaserPointer( /* Out */ LaserPointerStart, /* Out */ LaserPointerEnd))
-				{
-					const FVector LaserPointerDirection = (LaserPointerEnd - LaserPointerStart).GetSafeNormal();
-					BrushTraceDirection = LaserPointerDirection;
-
-					// Only start painting if we're not dragging a widget handle
-					if (ViewportClient.GetCurrentWidgetAxis() == EAxisList::None)
-					{
-						if (UISettings.GetPaintToolSelected() || UISettings.GetReapplyToolSelected() || UISettings.GetLassoSelectToolSelected())
-						{
-							StartFoliageBrushTrace(&ViewportClient, Interactor);
-							FoliageBrushTrace(&ViewportClient, LaserPointerStart, LaserPointerDirection);
-						}
-
-						// Fill a static mesh with foliage brush
-						else if (UISettings.GetPaintBucketToolSelected() || UISettings.GetReapplyPaintBucketToolSelected())
-						{
-							FHitResult HitResult = Interactor->GetHitResultFromLaserPointer();
-
-							if (HitResult.Actor.Get() != nullptr)
-							{
-								FScopedTransaction Transaction(NSLOCTEXT("UnrealEd", "FoliageMode_EditTransaction", "Foliage Editing"));
-
-								if (IsModifierButtonPressed(&ViewportClient))
-								{
-									ApplyPaintBucket_Remove(HitResult.Actor.Get());
-								}
-								else
-								{
-									ApplyPaintBucket_Add(HitResult.Actor.Get());
-								}
-							}
-						}
-						// Select an instanced foliage
-						else if (UISettings.GetSelectToolSelected())
-						{
-							FHitResult HitResult = Interactor->GetHitResultFromLaserPointer();
-
-							FScopedTransaction Transaction(NSLOCTEXT("UnrealEd", "FoliageMode_EditTransaction", "Foliage Editing"));
-
-							if (HitResult.GetActor() != nullptr)
-							{
-								// Clear all currently selected instances
-								FEdModeFoliageSelectionUpdate Scope(this);
-								SelectInstances(ViewportClient.GetWorld(), false);
-								for (auto& FoliageMeshUI : FoliageMeshList)
-								{
-									UFoliageType* Settings = FoliageMeshUI->Settings;
-									SelectInstanceAtLocation(ViewportClient.GetWorld(), Settings, HitResult.ImpactPoint, !IsModifierButtonPressed(&ViewportClient));
-								}
-							}
-
-							// @todo vreditor: we currently don't have a key mapping scheme to snap selected instances to ground 
-							// SnapSelectedInstancesToGround(GetWorld());
-
-						}
-					}
-				}
-			}
-			// Stop current tracking if the user is no longer painting
-			else if (Action.Event == IE_Released && FoliageInteractor && (FoliageInteractor == Interactor))
-			{
-				EndFoliageBrushTrace();
-				FoliageInteractor = nullptr;
-
-				bWasHandled = true;
-				bOutIsInputCaptured = false;
-			}
-		}
-	}
-}
-
 void FEdModeFoliage::PostUndo()
 {
 	FEdMode::PostUndo();
@@ -978,10 +807,6 @@ void FEdModeFoliage::StartFoliageBrushTrace(FEditorViewportClient* ViewportClien
 	if (!bToolActive)
 	{
 		GEditor->BeginTransaction(NSLOCTEXT("UnrealEd", "FoliageMode_EditTransaction", "Foliage Editing"));
-		if (Interactor)
-		{
-			FoliageInteractor = Interactor;
-		}
 		PreApplyBrush();
 		ApplyBrush(ViewportClient);
 
@@ -1068,8 +893,7 @@ void FEdModeFoliage::FoliageBrushTrace(FEditorViewportClient* ViewportClient, co
 bool FEdModeFoliage::MouseMove(FEditorViewportClient* ViewportClient, FViewport* Viewport, int32 MouseX, int32 MouseY)
 {
 	// Use mouse capture if there's no other interactor currently tracing brush
-	UVREditorMode* VREditorMode = Cast<UVREditorMode>( GEditor->GetEditorWorldExtensionsManager()->GetEditorWorldExtensions( GetWorld() )->FindExtension( UVREditorMode::StaticClass() ) );
-	if (IsEditingEnabled() && (VREditorMode == nullptr || !VREditorMode->IsActive()))
+	if (IsEditingEnabled())
 	{
 		// Compute a world space ray from the screen space mouse coordinates
 		FSceneViewFamilyContext ViewFamily(FSceneViewFamily::ConstructionValues(
@@ -1105,29 +929,26 @@ bool FEdModeFoliage::MouseMove(FEditorViewportClient* ViewportClient, FViewport*
 */
 bool FEdModeFoliage::CapturedMouseMove(FEditorViewportClient* ViewportClient, FViewport* Viewport, int32 MouseX, int32 MouseY)
 {
-	// Use mouse capture if there's no other interactor currently tracing brush
-	UVREditorMode* VREditorMode = Cast<UVREditorMode>( GEditor->GetEditorWorldExtensionsManager()->GetEditorWorldExtensions( GetWorld() )->FindExtension( UVREditorMode::StaticClass() ) );
-	if (VREditorMode == nullptr || !VREditorMode->IsActive())
-	{
-		//Compute a world space ray from the screen space mouse coordinates
-		FSceneViewFamilyContext ViewFamily(FSceneViewFamily::ConstructionValues(
-			ViewportClient->Viewport,
-			ViewportClient->GetScene(),
-			ViewportClient->EngineShowFlags)
-			.SetRealtimeUpdate(ViewportClient->IsRealtime()));
 
-		FSceneView* View = ViewportClient->CalcSceneView(&ViewFamily);
-		FViewportCursorLocation MouseViewportRay(View, ViewportClient, MouseX, MouseY);
-		BrushTraceDirection = MouseViewportRay.GetDirection();
+	//Compute a world space ray from the screen space mouse coordinates
+	FSceneViewFamilyContext ViewFamily(FSceneViewFamily::ConstructionValues(
+		ViewportClient->Viewport,
+		ViewportClient->GetScene(),
+		ViewportClient->EngineShowFlags)
+		.SetRealtimeUpdate(ViewportClient->IsRealtime()));
 
-		FVector BrushTraceStart = MouseViewportRay.GetOrigin();
-		if (ViewportClient->IsOrtho())
-		{
-			BrushTraceStart += -WORLD_MAX * BrushTraceDirection;
-		}
+	FSceneView* View = ViewportClient->CalcSceneView(&ViewFamily);
+	FViewportCursorLocation MouseViewportRay(View, ViewportClient, MouseX, MouseY);
+	BrushTraceDirection = MouseViewportRay.GetDirection();
 
-		FoliageBrushTrace(ViewportClient, BrushTraceStart, BrushTraceDirection);
+	FVector BrushTraceStart = MouseViewportRay.GetOrigin();
+	if (ViewportClient->IsOrtho())
+	{
+		BrushTraceStart += -WORLD_MAX * BrushTraceDirection;
 	}
+
+	FoliageBrushTrace(ViewportClient, BrushTraceStart, BrushTraceDirection);
+
 	return false;
 }
 
@@ -2806,7 +2627,8 @@ void FEdModeFoliage::PreApplyBrush()
 
 void FEdModeFoliage::ApplyBrush(FEditorViewportClient* ViewportClient)
 {
-	if (!bBrushTraceValid || ViewportClient != GCurrentLevelEditingViewportClient)
+	FEditorViewportClient* CurrentViewportClient = StaticCast<FEditorViewportClient*>(GEditor->GetActiveViewport()->GetClient());
+	if (!bBrushTraceValid || ViewportClient != CurrentViewportClient)
 	{
 		return;
 	}
@@ -2814,8 +2636,7 @@ void FEdModeFoliage::ApplyBrush(FEditorViewportClient* ViewportClient)
 	float BrushArea = PI * FMath::Square(UISettings.GetRadius());
 
 	// Tablet pressure or motion controller pressure
-	const UVREditorInteractor* VREditorInteractor = Cast<UVREditorInteractor>(FoliageInteractor);
-	const float Pressure = VREditorInteractor ? VREditorInteractor->GetSelectAndMoveTriggerValue() : ViewportClient->Viewport->IsPenActive() ? ViewportClient->Viewport->GetTabletPressure() : 1.f;
+	const float Pressure = ViewportClient->Viewport->IsPenActive() ? ViewportClient->Viewport->GetTabletPressure() : 1.f;
 
 	// Cache a copy of the world pointer
 	UWorld* World = ViewportClient->GetWorld();
@@ -3525,14 +3346,7 @@ bool FEdModeFoliage::CanPaint(const UFoliageType* FoliageType, const ULevel* InL
 
 bool FEdModeFoliage::IsModifierButtonPressed(const FEditorViewportClient* ViewportClient) const
 {
-	const UVREditorInteractor* VREditorInteractor = Cast<UVREditorInteractor>(FoliageInteractor);
-	bool bIsModifierPressed = false;
-	if (VREditorInteractor != nullptr)
-	{
-		bIsModifierPressed = VREditorInteractor->IsModifierPressed();
-	}
-
-	return IsShiftDown(ViewportClient->Viewport) || bIsModifierPressed;
+	return IsShiftDown(ViewportClient->Viewport);
 }
 
 bool FEdModeFoliage::CanMoveSelectedFoliageToLevel(ULevel* InTargetLevel) const
@@ -3834,7 +3648,7 @@ bool FEdModeFoliage::InputKey(FEditorViewportClient* ViewportClient, FViewport*
 	}
 
 	bool bHandled = false;
-	if ((UISettings.GetPaintToolSelected() || UISettings.GetReapplyToolSelected() || UISettings.GetLassoSelectToolSelected()) && FoliageInteractor == nullptr)
+	if ((UISettings.GetPaintToolSelected() || UISettings.GetReapplyToolSelected() || UISettings.GetLassoSelectToolSelected()))
 	{
 		// Require Ctrl or not as per user preference
 		ELandscapeFoliageEditorControlType FoliageEditorControlType = GetDefault<ULevelEditorViewportSettings>()->FoliageEditorControlType;
@@ -3903,7 +3717,7 @@ bool FEdModeFoliage::InputKey(FEditorViewportClient* ViewportClient, FViewport*
 		}
 	}
 
-	if (!bHandled && (UISettings.GetLassoSelectToolSelected() || UISettings.GetSelectToolSelected()) && FoliageInteractor == nullptr)
+	if (!bHandled && (UISettings.GetLassoSelectToolSelected() || UISettings.GetSelectToolSelected()))
 	{
 		if (Event == IE_Pressed)
 		{
diff --git a/Engine/Source/Editor/FoliageEdit/Private/FoliageEdMode.h b/Engine/Source/Editor/FoliageEdit/Private/FoliageEdMode.h
--- a/Engine/Source/Editor/FoliageEdit/Private/FoliageEdMode.h
+++ b/Engine/Source/Editor/FoliageEdit/Private/FoliageEdMode.h
@@ -20,7 +20,6 @@ class UFoliageType;
 class ULandscapeComponent;
 class UPrimitiveComponent;
 class UStaticMeshComponent;
-class UViewportInteractor;
 struct FViewportClick;
 
 //
@@ -524,12 +523,6 @@ public:
 	/** Add desired instances. Uses foliage settings to determine location/scale/rotation and whether instances should be ignored */
 	static void AddInstances(UWorld* InWorld, const TArray<FDesiredFoliageInstance>& DesiredInstances, const FFoliagePaintingGeometryFilter& OverrideGeometryFilter, bool InRebuildFoliageTree = true);
 
-	/** Called when the user presses a button on their motion controller device */
-	void OnVRAction(class FEditorViewportClient& ViewportClient, class UViewportInteractor* Interactor, const struct FViewportActionKeyInput& Action, bool& bOutIsInputCaptured, bool& bWasHandled);
-
-	/** Called on VR hovering */
-	void OnVRHoverUpdate(UViewportInteractor* Interactor, FVector& HoverImpactPoint, bool& bWasHandled);
-
 	/** Called as PIE ends */
 	void OnEndPIE(const bool bIsSimulating);
 
@@ -702,9 +695,6 @@ private:
 
 	FDelegateHandle OnActorSpawnedHandle;
 
-	/** When painting in VR, this is the hand index that we're painting with.  Otherwise INDEX_NONE. */
-	class UViewportInteractor* FoliageInteractor;
-
 	int32 UpdateSelectionCounter;
 	bool bHasDeferredSelectionNotification;
 	friend class FEdModeFoliageSelectionUpdate;
diff --git a/Engine/Source/Editor/LandscapeEditor/LandscapeEditor.Build.cs b/Engine/Source/Editor/LandscapeEditor/LandscapeEditor.Build.cs
--- a/Engine/Source/Editor/LandscapeEditor/LandscapeEditor.Build.cs
+++ b/Engine/Source/Editor/LandscapeEditor/LandscapeEditor.Build.cs
@@ -25,18 +25,9 @@ public class LandscapeEditor : ModuleRules
 				"ImageWrapper",
                 "EditorWidgets",
                 "Foliage",
-				"ViewportInteraction",
-				"VREditor",
 			}
 			);
 
-		CircularlyReferencedDependentModules.AddRange(
-			new string[]
-			{
-				"ViewportInteraction",
-				"VREditor"
-			}
-		);
 
 		PrivateIncludePathModuleNames.AddRange(
 			new string[] {
diff --git a/Engine/Source/Editor/LandscapeEditor/Private/LandscapeEdMode.cpp b/Engine/Source/Editor/LandscapeEditor/Private/LandscapeEdMode.cpp
--- a/Engine/Source/Editor/LandscapeEditor/Private/LandscapeEdMode.cpp
+++ b/Engine/Source/Editor/LandscapeEditor/Private/LandscapeEdMode.cpp
@@ -27,8 +27,6 @@
 #include "LandscapeHeightfieldCollisionComponent.h"
 #include "InstancedFoliageActor.h"
 #include "EditorWorldExtension.h"
-#include "ViewportWorldInteraction.h"
-#include "VREditorInteractor.h"
 #include "LandscapeEdModeTools.h"
 #include "LandscapeInfoMap.h"
 
@@ -38,15 +36,11 @@
 #include "SLandscapeEditor.h"
 #include "Framework/Application/SlateApplication.h"
 
-// VR Editor
-#include "VREditorMode.h"
-
 // Classes
 #include "LandscapeMaterialInstanceConstant.h"
 #include "LandscapeSplinesComponent.h"
 #include "ComponentReregisterContext.h"
 #include "EngineUtils.h"
-#include "IVREditorModule.h"
 #include "Misc/ScopedSlowTask.h"
 #include "LandscapeEditorCommands.h"
 #include "Framework/Commands/InputBindingManager.h"
@@ -598,19 +592,6 @@ void FEdModeLandscape::Enter()
 	{
 		GizmoBrush->EnterBrush();
 	}
-
-	// Register to find out about VR input events
-	UViewportWorldInteraction* ViewportWorldInteraction = Cast<UViewportWorldInteraction>(GEditor->GetEditorWorldExtensionsManager()->GetEditorWorldExtensions(GetWorld())->FindExtension(UViewportWorldInteraction::StaticClass()));
-	if (ViewportWorldInteraction != nullptr)
-	{
-
-			ViewportWorldInteraction->OnViewportInteractionInputAction().RemoveAll(this);
-			ViewportWorldInteraction->OnViewportInteractionInputAction().AddRaw(this, &FEdModeLandscape::OnVRAction);
-
-			ViewportWorldInteraction->OnViewportInteractionHoverUpdate().RemoveAll(this);
-			ViewportWorldInteraction->OnViewportInteractionHoverUpdate().AddRaw(this, &FEdModeLandscape::OnVRHoverUpdate);
-
-	}
 }
 
 /** FEdMode: Called when the mode is exited */
@@ -630,14 +611,6 @@ void FEdModeLandscape::Exit()
 		}
 	}
 
-	// Unregister VR mode from event handlers
-	UViewportWorldInteraction* ViewportWorldInteraction = Cast<UViewportWorldInteraction>(GEditor->GetEditorWorldExtensionsManager()->GetEditorWorldExtensions(GetWorld())->FindExtension(UViewportWorldInteraction::StaticClass()));
-	if (ViewportWorldInteraction != nullptr)
-	{
-		ViewportWorldInteraction->OnViewportInteractionInputAction().RemoveAll(this);
-		ViewportWorldInteraction->OnViewportInteractionHoverUpdate().RemoveAll(this);
-	}
-
 	GEngine->OnLevelActorDeleted().Remove(OnLevelActorDeletedDelegateHandle);
 	GEngine->OnLevelActorAdded().Remove(OnLevelActorAddedDelegateHandle);
 	
@@ -716,42 +689,6 @@ void FEdModeLandscape::Exit()
 	FEdMode::Exit();
 }
 
-
-void FEdModeLandscape::OnVRHoverUpdate(UViewportInteractor* Interactor, FVector& HoverImpactPoint, bool& bWasHandled)
-{
-	if (InteractorPainting != nullptr && InteractorPainting == Interactor && IVREditorModule::Get().IsVREditorModeActive())
-	{
-		UVREditorMode* VREditorMode = Cast<UVREditorMode>( GEditor->GetEditorWorldExtensionsManager()->GetEditorWorldExtensions( GetWorld() )->FindExtension( UVREditorMode::StaticClass() ) );
-		if( VREditorMode != nullptr && VREditorMode->IsActive() && Interactor != nullptr && Interactor->GetDraggingMode() == EViewportInteractionDraggingMode::Nothing )
-		{
-			const UVREditorInteractor* VREditorInteractor = Cast<UVREditorInteractor>(Interactor);
-
-			if (VREditorInteractor != nullptr && !VREditorInteractor->IsHoveringOverPriorityType() && CurrentTool && (CurrentTool->GetSupportedTargetTypes() == ELandscapeToolTargetTypeMask::NA || CurrentToolTarget.TargetType != ELandscapeToolTargetType::Invalid))
-			{
-				FVector HitLocation;
-				FVector LaserPointerStart, LaserPointerEnd;
-				if (Interactor->GetLaserPointer( /* Out */ LaserPointerStart, /* Out */ LaserPointerEnd))
-				{
-					if( LandscapeTrace( LaserPointerStart, LaserPointerEnd, HitLocation ) )
-					{
-						if (CurrentTool && CurrentTool->IsToolActive())
-						{
-							CurrentTool->SetExternalModifierPressed(Interactor->IsModifierPressed());
-							CurrentTool->MouseMove(nullptr, nullptr, HitLocation.X, HitLocation.Y);
-						}
-
-						if (CurrentBrush)
-						{
-							// Inform the brush of the current location, to update the cursor
-							CurrentBrush->MouseMove(HitLocation.X, HitLocation.Y);
-						}
-					}
-				}
-			}
-		}
-	}
-}
-
 void FEdModeLandscape::OnPreSaveWorld(uint32 InSaveFlags, const class UWorld* InWorld)
 {
 	// Avoid doing this during cooking to keep determinism and we don't want to do this on GameWorlds.
@@ -771,70 +708,6 @@ void FEdModeLandscape::OnPreSaveWorld(uint32 InSaveFlags, const class UWorld* In
 	}
 }
 
-void FEdModeLandscape::OnVRAction(FEditorViewportClient& ViewportClient, UViewportInteractor* Interactor, const struct FViewportActionKeyInput& Action, bool& bOutIsInputCaptured, bool& bWasHandled)
-{
-	UVREditorMode* VREditorMode = Cast<UVREditorMode>( GEditor->GetEditorWorldExtensionsManager()->GetEditorWorldExtensions( GetWorld() )->FindExtension( UVREditorMode::StaticClass() ) );
-	// Never show the traditional Unreal transform widget.  It doesn't work in VR because we don't have hit proxies.
-	ViewportClient.EngineShowFlags.SetModeWidgets(false);
-
-	if (VREditorMode != nullptr && VREditorMode->IsActive() && Interactor != nullptr && Interactor->GetDraggingMode() == EViewportInteractionDraggingMode::Nothing)
-	{
-		if (Action.ActionType == ViewportWorldActionTypes::SelectAndMove)
-		{
-			const UVREditorInteractor* VREditorInteractor = Cast<UVREditorInteractor>(Interactor);
-
-			// Begin landscape brush
-			if (Action.Event == IE_Pressed && !VREditorInteractor->IsHoveringOverUI() && !VREditorInteractor->IsHoveringOverPriorityType() && CurrentTool)
-			{
-				if (ViewportClient.Viewport != nullptr && ViewportClient.Viewport == ToolActiveViewport)
-				{
-					CurrentTool->EndTool(&ViewportClient);
-					ToolActiveViewport = nullptr;
-				}
-
-				if (CurrentTool->GetSupportedTargetTypes() == ELandscapeToolTargetTypeMask::NA || CurrentToolTarget.TargetType != ELandscapeToolTargetType::Invalid)
-				{
-					FVector HitLocation;
-					FVector LaserPointerStart, LaserPointerEnd;
-					if (Interactor->GetLaserPointer( /* Out */ LaserPointerStart, /* Out */ LaserPointerEnd))
-					{
-						if (LandscapeTrace(LaserPointerStart, LaserPointerEnd, HitLocation))
-						{
-							if (!(CurrentToolTarget.TargetType == ELandscapeToolTargetType::Weightmap && CurrentToolTarget.LayerInfo == NULL))
-							{
-								CurrentTool->SetExternalModifierPressed(Interactor->IsModifierPressed());
-								if( CurrentTool->BeginTool(&ViewportClient, CurrentToolTarget, HitLocation))
-								{
-									ToolActiveViewport = ViewportClient.Viewport;
-								}
-							}
-
-							bIsPaintingInVR = true;
-							bWasHandled = true;
-							bOutIsInputCaptured = false;
-
-							InteractorPainting = Interactor;
-						}
-					}
-				}
-			}
-
-			// End landscape brush
-			else if (Action.Event == IE_Released)
-			{
-				if (CurrentTool && ViewportClient.Viewport != nullptr && ViewportClient.Viewport == ToolActiveViewport)
-				{
-					CurrentTool->EndTool(&ViewportClient);
-					ToolActiveViewport = nullptr;
-				}
-
-				bIsPaintingInVR = false;
-			}
-		}
-	}
-
-}
-
 /** FEdMode: Called once per frame */
 void FEdModeLandscape::Tick(FEditorViewportClient* ViewportClient, float DeltaTime)
 {
@@ -3384,17 +3257,6 @@ void FEdModeLandscape::ForceRealTimeViewports(const bool bEnable)
 				if (bEnable)
 				{
 					Viewport.SetRealtimeOverride(bEnable, LOCTEXT("RealtimeOverrideMessage_Landscape", "Landscape Mode"));
-
-					// @todo vreditor: Force game view to true in VREditor since we can't use hitproxies and debug objects yet
-					UVREditorMode* VREditorMode = Cast<UVREditorMode>( GEditor->GetEditorWorldExtensionsManager()->GetEditorWorldExtensions( GetWorld() )->FindExtension( UVREditorMode::StaticClass() ) );
-					if( VREditorMode != nullptr && VREditorMode->IsActive())
-					{
-						Viewport.SetVREditView(true);
-					} 
-					else
-					{
-						Viewport.SetVREditView(false);
-					}
 				}
 				else
 				{
diff --git a/Engine/Source/Editor/LandscapeEditor/Private/LandscapeEdMode.h b/Engine/Source/Editor/LandscapeEditor/Private/LandscapeEdMode.h
--- a/Engine/Source/Editor/LandscapeEditor/Private/LandscapeEdMode.h
+++ b/Engine/Source/Editor/LandscapeEditor/Private/LandscapeEdMode.h
@@ -594,11 +594,6 @@ public:
 
 	void OnPreSaveWorld(uint32 InSaveFlags, const class UWorld* InWorld);
 
-	/** Called when the user presses a button on their motion controller device */
-	void OnVRAction(FEditorViewportClient& ViewportClient, UViewportInteractor* Interactor, const FViewportActionKeyInput& Action, bool& bOutIsInputCaptured, bool& bWasHandled);
-
-	void OnVRHoverUpdate(UViewportInteractor* Interactor, FVector& HoverImpactPoint, bool& bWasHandled);
-
 	/** Handle notification that visible levels may have changed and we should update the editable landscapes list */
 	void HandleLevelsChanged(bool ShouldExitMode);
 
diff --git a/Engine/Source/Editor/LandscapeEditor/Private/LandscapeEdModeTools.h b/Engine/Source/Editor/LandscapeEditor/Private/LandscapeEdModeTools.h
--- a/Engine/Source/Editor/LandscapeEditor/Private/LandscapeEdModeTools.h
+++ b/Engine/Source/Editor/LandscapeEditor/Private/LandscapeEdModeTools.h
@@ -13,13 +13,12 @@
 #include "LandscapeDataAccess.h"
 #include "LandscapeHeightfieldCollisionComponent.h"
 #include "InstancedFoliageActor.h"
-#include "VREditorInteractor.h"
 #include "AI/NavigationSystemBase.h"
 #include "Landscape.h"
 #include "Logging/LogMacros.h"
 
 DECLARE_LOG_CATEGORY_EXTERN(LogLandscapeTools, Log, All);
-// VR Editor
+
 
 //
 //	FNoiseParameter - Perlin noise
@@ -1354,7 +1353,6 @@ public:
 		, EdMode(InEdMode)
 		, bCanToolBeActivated(true)
 		, ToolStroke()
-		, bExternalModifierPressed(false)
 	{
 	}
 
@@ -1456,7 +1454,6 @@ public:
 		ToolStroke.Reset();
 		EdMode->CurrentBrush->EndStroke();
 		EdMode->UpdateLayerUsageInformation(&EdMode->CurrentToolTarget.LayerInfo);
-		bExternalModifierPressed = false;
 
 		if (ShouldUpdateEditingLayer())
 		{
@@ -1515,11 +1512,6 @@ public:
 	virtual void SetCanToolBeActivated(bool Value) { bCanToolBeActivated = Value; }
 	virtual bool CanToolBeActivated() const {	return bCanToolBeActivated; }
 
-	virtual void SetExternalModifierPressed(const bool bPressed) override
-	{
-		bExternalModifierPressed = bPressed;
-	}
-
 protected:
 	TArray<FLandscapeToolInteractorPosition> InteractorPositions;
 	FVector2D LastInteractorPosition;
@@ -1528,13 +1520,10 @@ protected:
 	bool bCanToolBeActivated;
 	TOptional<TStrokeClass> ToolStroke;
 
-	/** Whether a modifier was pressed in another system (VREditor). */
-	bool bExternalModifierPressed;
-
 	bool IsModifierPressed(const class FEditorViewportClient* ViewportClient = nullptr)
 	{
-		UE_LOG(LogLandscapeTools, VeryVerbose, TEXT("bExternalModifierPressed = %d, ViewportClient = %d, IsShiftDown = %d"), bExternalModifierPressed, (ViewportClient != nullptr), (ViewportClient != nullptr && IsShiftDown(ViewportClient->Viewport)));
-		return bExternalModifierPressed || (ViewportClient != nullptr && IsShiftDown(ViewportClient->Viewport));
+		UE_LOG(LogLandscapeTools, VeryVerbose, TEXT("ViewportClient = %d, IsShiftDown = %d"), (ViewportClient != nullptr), (ViewportClient != nullptr && IsShiftDown(ViewportClient->Viewport)));
+		return ViewportClient != nullptr && IsShiftDown(ViewportClient->Viewport);
 	}
 };
 
diff --git a/Engine/Source/Editor/LandscapeEditor/Public/LandscapeToolInterface.h b/Engine/Source/Editor/LandscapeEditor/Public/LandscapeToolInterface.h
--- a/Engine/Source/Editor/LandscapeEditor/Public/LandscapeToolInterface.h
+++ b/Engine/Source/Editor/LandscapeEditor/Public/LandscapeToolInterface.h
@@ -230,7 +230,6 @@ public:
 
 	virtual void SetCanToolBeActivated(bool Value) { }
 	virtual bool CanToolBeActivated() const { return true;  }
-	virtual void SetExternalModifierPressed(const bool bPressed) {};
 
 	virtual EEditAction::Type GetActionEditDuplicate() { return EEditAction::Skip; }
 	virtual EEditAction::Type GetActionEditDelete() { return EEditAction::Skip; }
diff --git a/Engine/Source/Editor/MeshPaint/MeshPaint.Build.cs b/Engine/Source/Editor/MeshPaint/MeshPaint.Build.cs
--- a/Engine/Source/Editor/MeshPaint/MeshPaint.Build.cs
+++ b/Engine/Source/Editor/MeshPaint/MeshPaint.Build.cs
@@ -30,21 +30,11 @@ public class MeshPaint : ModuleRules
                 "MeshDescription",
 				"StaticMeshDescription",
                 "SourceControl",
-                "ViewportInteraction",
-                "VREditor",
                 "PropertyEditor",
                 "MainFrame",
             }
         );
 
-		CircularlyReferencedDependentModules.AddRange(
-			new string[]
-			{
-				"ViewportInteraction",
-				"VREditor"
-			}
-		);
-
 		PrivateIncludePathModuleNames.AddRange(
 			new string[]
 			{
diff --git a/Engine/Source/Editor/MeshPaint/Private/IMeshPaintMode.cpp b/Engine/Source/Editor/MeshPaint/Private/IMeshPaintMode.cpp
--- a/Engine/Source/Editor/MeshPaint/Private/IMeshPaintMode.cpp
+++ b/Engine/Source/Editor/MeshPaint/Private/IMeshPaintMode.cpp
@@ -12,9 +12,6 @@
 #include "EditorModeManager.h"
 #include "Toolkits/ToolkitManager.h"
 
-#include "VREditorMode.h"
-#include "IVREditorModule.h"
-
 #include "AssetToolsModule.h"
 #include "AssetRegistryModule.h"
 #include "EditorSupportDelegates.h"
@@ -27,11 +24,6 @@
 
 #include "MeshPaintAdapterFactory.h"
 
-#include "ViewportWorldInteraction.h"
-#include "ViewportInteractableInterface.h"
-#include "VREditorInteractor.h"
-#include "EditorWorldExtension.h"
-
 #include "Components/PrimitiveComponent.h"
 
 #include "IMeshPainter.h"
@@ -44,15 +36,13 @@ DEFINE_LOG_CATEGORY_STATIC(LogMeshPaintEdMode, Log, All);
 
 /** Constructor */
 IMeshPaintEdMode::IMeshPaintEdMode() 
-	: FEdMode(),
-	  PaintingWithInteractorInVR( nullptr )
+	: FEdMode()
 {
 	GEditor->OnEditorClose().AddRaw(this, &IMeshPaintEdMode::OnResetViewMode);
 }
 
 IMeshPaintEdMode::~IMeshPaintEdMode()
 {
-	GEditor->OnEditorClose().RemoveAll(this);
 }
 
 /** FGCObject interface */
@@ -103,21 +93,6 @@ void IMeshPaintEdMode::Enter()
 	// highlight effect distorting the appearance.
 	GEngine->OverrideSelectedMaterialColor( FLinearColor::Black );
 
-	UEditorWorldExtensionCollection* ExtensionCollection = GEditor->GetEditorWorldExtensionsManager()->GetEditorWorldExtensions(GetWorld());
-	if (ExtensionCollection != nullptr)
-	{
-		// Register to find out about VR input events
-		UViewportWorldInteraction* ViewportWorldInteraction = Cast<UViewportWorldInteraction>(ExtensionCollection->FindExtension(UViewportWorldInteraction::StaticClass()));
-		if (ViewportWorldInteraction != nullptr)
-		{
-			ViewportWorldInteraction->OnViewportInteractionInputAction().RemoveAll(this);
-			ViewportWorldInteraction->OnViewportInteractionInputAction().AddRaw(this, &IMeshPaintEdMode::OnVRAction);
-
-			// Hide the VR transform gizmo while we're in mesh paint mode.  It sort of gets in the way of painting.
-			ViewportWorldInteraction->SetTransformGizmoVisible(false);
-		}
-	}
-
 	if (UsesToolkits())
 	{
 		MeshPainter->RegisterCommands(Toolkit->GetToolkitCommands());
@@ -141,20 +116,6 @@ void IMeshPaintEdMode::Exit()
 		MeshPainter->UnregisterCommands(Toolkit->GetToolkitCommands());
 	}
 
-	UEditorWorldExtensionCollection* ExtensionCollection = GEditor->GetEditorWorldExtensionsManager()->GetEditorWorldExtensions(GetWorld());
-	if (ExtensionCollection != nullptr)
-	{
-		UViewportWorldInteraction* ViewportWorldInteraction = Cast<UViewportWorldInteraction>(ExtensionCollection->FindExtension(UViewportWorldInteraction::StaticClass()));
-		if (ViewportWorldInteraction != nullptr)
-		{
-			// Restore the transform gizmo visibility
-			ViewportWorldInteraction->SetTransformGizmoVisible(true);
-
-			// Unregister from event handlers
-			ViewportWorldInteraction->OnViewportInteractionInputAction().RemoveAll(this);
-		}
-	}
-
 	// The user can manipulate the editor selection lock flag in paint mode so we make sure to restore it here
 	GEdSelectionLock = bWasSelectionLockedOnStart;
 	
@@ -223,7 +184,6 @@ bool IMeshPaintEdMode::EndTracking(FEditorViewportClient* InViewportClient, FVie
 	if (MeshPainter->IsPainting())
 	{
 		MeshPainter->FinishPainting();
-		PaintingWithInteractorInVR = nullptr;
 	}
 
 	return true;
@@ -257,7 +217,6 @@ bool IMeshPaintEdMode::InputKey( FEditorViewportClient* InViewportClient, FViewp
 		{
 			bHandled = true;
 			MeshPainter->FinishPainting();
-			PaintingWithInteractorInVR = nullptr;
 		}
 		else if( !MeshPainter->IsPainting() && bUserWantsPaint && !InViewportClient->IsMovingCamera())
 		{	
@@ -389,40 +348,24 @@ void IMeshPaintEdMode::Render( const FSceneView* View, FViewport* Viewport, FPri
 	MeshPainter->Render(View, Viewport, PDI);
 
 	// Flow painting
-	if (MeshPainter->IsPainting())
+	if (MeshPainter->IsPainting() && MeshPainter->GetBrushSettings()->bEnableFlow)
 	{
-		/** If we are currently painting with a VR interactor, apply paint for the current vr interactor state/position */
-		if (PaintingWithInteractorInVR != nullptr)
-		{
-			UVREditorInteractor* VREditorInteractor = Cast<UVREditorInteractor>(PaintingWithInteractorInVR);
-			FVector LaserPointerStart, LaserPointerEnd;
-			if (VREditorInteractor->GetLaserPointer( /* Out */ LaserPointerStart, /* Out */ LaserPointerEnd))
+		// Make sure the cursor is visible OR we're flood filling.  No point drawing a paint cue when there's no cursor.
+		if (Viewport->IsCursorVisible())
+		{				
+			// Grab the mouse cursor position
+			FIntPoint MousePosition;
+			Viewport->GetMousePos(MousePosition);
+
+			// Is the mouse currently over the viewport? or flood filling
+			if ((MousePosition.X >= 0 && MousePosition.Y >= 0 && MousePosition.X < (int32)Viewport->GetSizeXY().X && MousePosition.Y < (int32)Viewport->GetSizeXY().Y))
 			{
-				const FVector LaserPointerDirection = (LaserPointerEnd - LaserPointerStart).GetSafeNormal();
-				UVREditorMode* VREditorMode = Cast<UVREditorMode>(GEditor->GetEditorWorldExtensionsManager()->GetEditorWorldExtensions(GetWorld())->FindExtension(UVREditorMode::StaticClass()));
-				MeshPainter->PaintVR(Viewport, VREditorMode->GetHeadTransform().GetLocation(), LaserPointerStart, LaserPointerDirection, VREditorInteractor);
-			}
-		}
-		else if (MeshPainter->GetBrushSettings()->bEnableFlow)
-		{
-			// Make sure the cursor is visible OR we're flood filling.  No point drawing a paint cue when there's no cursor.
-			if (Viewport->IsCursorVisible())
-			{				
-				// Grab the mouse cursor position
-				FIntPoint MousePosition;
-				Viewport->GetMousePos(MousePosition);
-
-				// Is the mouse currently over the viewport? or flood filling
-				if ((MousePosition.X >= 0 && MousePosition.Y >= 0 && MousePosition.X < (int32)Viewport->GetSizeXY().X && MousePosition.Y < (int32)Viewport->GetSizeXY().Y))
-				{
-					// Compute a world space ray from the screen space mouse coordinates
-					FViewportCursorLocation MouseViewportRay(View, static_cast<FEditorViewportClient*>(Viewport->GetClient()), MousePosition.X, MousePosition.Y);
-					MeshPainter->Paint(Viewport, View->ViewMatrices.GetViewOrigin(), MouseViewportRay.GetOrigin(), MouseViewportRay.GetDirection());
-				}
+				// Compute a world space ray from the screen space mouse coordinates
+				FViewportCursorLocation MouseViewportRay(View, static_cast<FEditorViewportClient*>(Viewport->GetClient()), MousePosition.X, MousePosition.Y);
+				MeshPainter->Paint(Viewport, View->ViewMatrices.GetViewOrigin(), MouseViewportRay.GetOrigin(), MouseViewportRay.GetDirection());
 			}
 		}
 	}
-
 }
 
 /** FEdMode: Handling SelectActor */
@@ -465,85 +408,4 @@ bool IMeshPaintEdMode::ProcessEditDelete()
 	return false;
 }
 
-void IMeshPaintEdMode::OnVRAction( class FEditorViewportClient& ViewportClient, UViewportInteractor* Interactor, 
-	const FViewportActionKeyInput& Action, bool& bOutIsInputCaptured, bool& bWasHandled )
-{
-	UVREditorMode* VREditorMode = Cast<UVREditorMode>(GEditor->GetEditorWorldExtensionsManager()->GetEditorWorldExtensions(GetWorld())->FindExtension(UVREditorMode::StaticClass()));
-	UVREditorInteractor* VREditorInteractor = Cast<UVREditorInteractor>(Interactor);
-	if (VREditorMode != nullptr && VREditorMode->IsActive() && VREditorInteractor != nullptr && VREditorInteractor->GetDraggingMode() == EViewportInteractionDraggingMode::Nothing)
-	{
-		if (Action.ActionType == ViewportWorldActionTypes::SelectAndMove)
-		{
-			if (!MeshPainter->IsPainting() && Action.Event == IE_Pressed && !VREditorInteractor->IsHoveringOverPriorityType())
-			{
-				// Check to see that we're clicking on a selected object.  You can only paint on selected things.  Otherwise,
-				// we'll fall through to the normal interaction code which might cause the object to become selected.
-				bool bIsClickingOnSelectedObject = false;
-				{
-					FHitResult HitResult = VREditorInteractor->GetHitResultFromLaserPointer();
-					if( HitResult.Actor.IsValid() )
-					{
-						UViewportWorldInteraction& WorldInteraction = VREditorMode->GetWorldInteraction();
-						
-						if( WorldInteraction.IsInteractableComponent( HitResult.GetComponent() ) )
-						{
-							AActor* Actor = HitResult.Actor.Get();
-
-							// Make sure we're not hovering over some other viewport interactable, such as a dockable window selection bar or close button
-							IViewportInteractableInterface* ActorInteractable = Cast<IViewportInteractableInterface>( Actor );
-							if( ActorInteractable == nullptr )
-							{
-								if( Actor != WorldInteraction.GetTransformGizmoActor() )  // Don't change any actor selection state if the user clicked on a gizmo
-								{
-									if( Actor->IsSelected() )
-									{
-										bIsClickingOnSelectedObject = true;
-									}
-								}
-							}
-						}
-					}
-				}
-
-				if( bIsClickingOnSelectedObject )
-				{
-					bWasHandled = true;
-					bOutIsInputCaptured = true;
-
-					// Go ahead and paint immediately
-					FVector LaserPointerStart, LaserPointerEnd;
-					if( Interactor->GetLaserPointer( /* Out */ LaserPointerStart, /* Out */ LaserPointerEnd ) )
-					{
-						const FVector LaserPointerDirection = ( LaserPointerEnd - LaserPointerStart ).GetSafeNormal();
-
-						/** Apply painting using the current state/position of the VR interactor */
-						const bool bAnyPaintableActorsUnderCursor = MeshPainter->PaintVR( ViewportClient.Viewport, VREditorMode->GetHeadTransform().GetLocation(), LaserPointerStart, LaserPointerDirection, VREditorInteractor);
-						if (bAnyPaintableActorsUnderCursor)
-						{
-							PaintingWithInteractorInVR = VREditorInteractor;
-							ViewportClient.bLockFlightCamera = true;
-						}
-					}
-				}
-			}
-			// Stop current tracking if the user is no longer painting
-			else if( MeshPainter->IsPainting() && Action.Event == IE_Released && PaintingWithInteractorInVR && PaintingWithInteractorInVR == Interactor )
-			{
-				MeshPainter->FinishPainting();
-				ViewportClient.bLockFlightCamera = false;
-				PaintingWithInteractorInVR = nullptr;
-
-				bWasHandled = true;
-				bOutIsInputCaptured = false;
-			}
-
-			else if (MeshPainter->IsPainting())
-			{
-				// A different hand might be painting, so absorb the input
-				bOutIsInputCaptured = bWasHandled = (Action.Event == IE_Released ? false : true);
-			}
-		}
-	}
-}
-
 #undef LOCTEXT_NAMESPACE
diff --git a/Engine/Source/Editor/MeshPaint/Private/IMeshPainter.cpp b/Engine/Source/Editor/MeshPaint/Private/IMeshPainter.cpp
--- a/Engine/Source/Editor/MeshPaint/Private/IMeshPainter.cpp
+++ b/Engine/Source/Editor/MeshPaint/Private/IMeshPainter.cpp
@@ -9,21 +9,12 @@
 #include "MeshPaintHelpers.h"
 #include "ScopedTransaction.h"
 
-#include "VREditorMode.h"
-#include "IVREditorModule.h"
-#include "ViewportWorldInteraction.h"
-#include "ViewportInteractableInterface.h"
-#include "ViewportInteractor.h"
-#include "VREditorInteractor.h"
-#include "EditorWorldExtension.h"
-
 #include "Components/MeshComponent.h"
 #include "Components/PrimitiveComponent.h"
 #include "Framework/Application/SlateApplication.h"
 #include "Framework/Commands/UICommandList.h"
 
 IMeshPainter::IMeshPainter() :
-	CurrentViewportInteractor(nullptr),
 	bArePainting(false),
 	TimeSinceStartedPainting(0.0f),
 	Time(0.0f),
@@ -47,9 +38,7 @@ void IMeshPainter::RenderInteractors(const FSceneView* View, FViewport* Viewport
 	// Apply paint pressure and start painting (or if not currently painting, draw a preview of where paint will be applied)
 	for (const MeshPaintHelpers::FPaintRay& PaintRay : PaintRays)
 	{
-		const UVREditorInteractor* VREditorInteractor = Cast<UVREditorInteractor>(PaintRay.ViewportInteractor);
-		EMeshPaintAction RayPaintAction = VREditorInteractor ? (VREditorInteractor->IsModifierPressed() ? EMeshPaintAction::Erase : EMeshPaintAction::Paint)
-			: (Viewport->KeyState(EKeys::LeftControl) || Viewport->KeyState(EKeys::RightControl)) ? EMeshPaintAction::Erase : EMeshPaintAction::Paint;
+		EMeshPaintAction RayPaintAction = (Viewport->KeyState(EKeys::LeftControl) || Viewport->KeyState(EKeys::RightControl)) ? EMeshPaintAction::Erase : EMeshPaintAction::Paint;
 
 		RenderInteractorWidget(PaintRay.CameraLocation, PaintRay.RayStart, PaintRay.RayDirection, PDI, RayPaintAction, bRenderVertices, DepthGroup);
 	}
@@ -149,32 +138,6 @@ bool IMeshPainter::Paint(FViewport* Viewport, const FVector& InCameraOrigin, con
 	return PaintInternal(InCameraOrigin, Rays, PaintAction, PaintStrength);
 }
 
-bool IMeshPainter::PaintVR(FViewport* Viewport, const FVector& InCameraOrigin, const FVector& InRayOrigin, const FVector& InRayDirection, UVREditorInteractor* VREditorInteractor)
-{
-	bool bPaintApplied = false;
-	// When painting using VR, allow the modifier button to activate Erase mode
-	if (VREditorInteractor)
-	{
-		// Determine custom paint strength from trigger
-		const float StrengthScale = VREditorInteractor->GetSelectAndMoveTriggerValue();
-		EMeshPaintAction PaintAction;
-		const bool bIsModifierPressed = VREditorInteractor->IsModifierPressed();
-		// Determine paint action according to whether or not modifier button is pressed
-		PaintAction = bIsModifierPressed ? EMeshPaintAction::Erase : EMeshPaintAction::Paint;
-
-		// Handle internal painting functionality
-		TPair<FVector, FVector> Ray(InRayOrigin, InRayDirection);
-		bPaintApplied = PaintInternal(InCameraOrigin, MakeArrayView(&Ray, 1), PaintAction, StrengthScale);
-
-		if (bPaintApplied)
-		{
-			CurrentViewportInteractor = VREditorInteractor;
-		}
-	}
-
-	return bPaintApplied;
-}
-
 bool IMeshPainter::InputKey(FEditorViewportClient* InViewportClient, FViewport* InViewport, FKey InKey, EInputEvent InEvent)
 {
 	bool bHandled = false;
@@ -197,8 +160,6 @@ void IMeshPainter::FinishPainting()
 			EndTransaction();
 		}
 	}
-
-	CurrentViewportInteractor = nullptr;
 }
 
 void IMeshPainter::RenderInteractorWidget(const FVector& InCameraOrigin, const FVector& InRayOrigin, const FVector& InRayDirection, FPrimitiveDrawInterface* PDI, EMeshPaintAction PaintAction, bool bRenderVertices, ESceneDepthPriorityGroup DepthGroup /*= SDPG_World*/)
diff --git a/Engine/Source/Editor/MeshPaint/Private/MeshPaintHelpers.cpp b/Engine/Source/Editor/MeshPaint/Private/MeshPaintHelpers.cpp
--- a/Engine/Source/Editor/MeshPaint/Private/MeshPaintHelpers.cpp
+++ b/Engine/Source/Editor/MeshPaint/Private/MeshPaintHelpers.cpp
@@ -39,13 +39,6 @@
 #include "EditorViewportClient.h"
 #include "LevelEditorViewport.h"
 
-#include "VREditorMode.h"
-#include "IVREditorModule.h"
-#include "ViewportWorldInteraction.h"
-#include "ViewportInteractableInterface.h"
-#include "VREditorInteractor.h"
-#include "EditorWorldExtension.h"
-
 #include "Factories/FbxSkeletalMeshImportData.h"
 
 #include "Async/ParallelFor.h"
@@ -572,76 +565,27 @@ bool MeshPaintHelpers::RetrieveViewportPaintRays(const FSceneView* View, FViewpo
 
 	if (ViewportClient->IsPerspective())
 	{
-		// If in VR mode retrieve possible viewport interactors and render widgets for them
-		UVREditorMode* VREditorMode = nullptr;
-		if (MeshPaintHelpers::IsInVRMode(ViewportClient))
-		{
-			VREditorMode = Cast<UVREditorMode>(GEditor->GetEditorWorldExtensionsManager()->GetEditorWorldExtensions(ViewportClient->GetWorld())->FindExtension(UVREditorMode::StaticClass()));
-
-			TArray<UViewportInteractor*> Interactors = VREditorMode->GetWorldInteraction().GetInteractors();
-
-			for (UViewportInteractor* Interactor : Interactors)
-			{
-				checkf(Interactor, TEXT("Invalid VR Interactor"));
-				
-				// Don't draw visual cue if we're hovering over a viewport interact able, such as a dockable window selection bar
-				bool bShouldDrawInteractor = false;
-				FHitResult HitResult = Interactor->GetHitResultFromLaserPointer();
-				if (HitResult.Actor.IsValid())
-				{
-					UViewportWorldInteraction& WorldInteraction = VREditorMode->GetWorldInteraction();
-
-					if (WorldInteraction.IsInteractableComponent(HitResult.GetComponent()))
-					{
-						AActor* Actor = HitResult.Actor.Get();
-
-						// Make sure we're not hovering over some other viewport interactable, such as a dockable window selection bar or close button
-						IViewportInteractableInterface* ActorInteractable = Cast<IViewportInteractableInterface>(Actor);
-						bShouldDrawInteractor = (ActorInteractable == nullptr);
-					}
-				}
-				
-				// Don't draw visual cue for paint brush when the interactor is hovering over UI
-				if (bShouldDrawInteractor && !Interactor->IsHoveringOverPriorityType())
-				{
-					FVector LaserPointerStart, LaserPointerEnd;
-					if (Interactor->GetLaserPointer( /* Out */ LaserPointerStart, /* Out */ LaserPointerEnd))
-					{
-						const FVector LaserPointerDirection = (LaserPointerEnd - LaserPointerStart).GetSafeNormal();
 
-						FPaintRay& NewPaintRay = *new(OutPaintRays) FPaintRay();
-						NewPaintRay.CameraLocation = VREditorMode->GetHeadTransform().GetLocation();
-						NewPaintRay.RayStart = LaserPointerStart;
-						NewPaintRay.RayDirection = LaserPointerDirection;
-						NewPaintRay.ViewportInteractor = Interactor;
-					}
-				}
-			}
-		}
-		else
+		// Make sure the cursor is visible OR we're flood filling.  No point drawing a paint cue when there's no cursor.
+		if (Viewport->IsCursorVisible())
 		{
-			// Else we're painting with mouse
-			// Make sure the cursor is visible OR we're flood filling.  No point drawing a paint cue when there's no cursor.
-			if (Viewport->IsCursorVisible())
+			if (!PDI->IsHitTesting())
 			{
-				if (!PDI->IsHitTesting())
-				{
-					// Grab the mouse cursor position
-					FIntPoint MousePosition;
-					Viewport->GetMousePos(MousePosition);
+				// Grab the mouse cursor position
+				FIntPoint MousePosition;
+				Viewport->GetMousePos(MousePosition);
 
-					// Is the mouse currently over the viewport? or flood filling
-					if ((MousePosition.X >= 0 && MousePosition.Y >= 0 && MousePosition.X < (int32)Viewport->GetSizeXY().X && MousePosition.Y < (int32)Viewport->GetSizeXY().Y))
-					{
-						// Compute a world space ray from the screen space mouse coordinates
-						FViewportCursorLocation MouseViewportRay(View, ViewportClient, MousePosition.X, MousePosition.Y);
-
-						FPaintRay& NewPaintRay = *new(OutPaintRays) FPaintRay();
-						NewPaintRay.CameraLocation = View->ViewMatrices.GetViewOrigin();
-						NewPaintRay.RayStart = MouseViewportRay.GetOrigin();
-						NewPaintRay.RayDirection = MouseViewportRay.GetDirection();
-						NewPaintRay.ViewportInteractor = nullptr;
-					}
+				// Is the mouse currently over the viewport? or flood filling
+				if ((MousePosition.X >= 0 && MousePosition.Y >= 0 && MousePosition.X < (int32)Viewport->GetSizeXY().X && MousePosition.Y < (int32)Viewport->GetSizeXY().Y))
+				{
+					// Compute a world space ray from the screen space mouse coordinates
+					FViewportCursorLocation MouseViewportRay(View, ViewportClient, MousePosition.X, MousePosition.Y);
+
+					FPaintRay& NewPaintRay = *new(OutPaintRays) FPaintRay();
+					NewPaintRay.CameraLocation = View->ViewMatrices.GetViewOrigin();
+					NewPaintRay.RayStart = MouseViewportRay.GetOrigin();
+					NewPaintRay.RayDirection = MouseViewportRay.GetDirection();
+					NewPaintRay.ViewportInteractor = nullptr;
 				}
 			}
 		}
@@ -1172,21 +1116,6 @@ void MeshPaintHelpers::SetRealtimeViewport(bool bRealtime)
 }
 
 
-bool MeshPaintHelpers::IsInVRMode(const FEditorViewportClient* ViewportClient)
-{
-	bool bIsInVRMode = false;
-	if (IVREditorModule::IsAvailable())
-	{
-		UVREditorMode* VREditorMode = Cast<UVREditorMode>(GEditor->GetEditorWorldExtensionsManager()->GetEditorWorldExtensions(ViewportClient->GetWorld())->FindExtension(UVREditorMode::StaticClass()));
-		if (VREditorMode != nullptr && VREditorMode->IsFullyInitialized() && VREditorMode->IsActive())
-		{
-			bIsInVRMode = true;
-		}
-	}
-
-	return bIsInVRMode;
-}
-
 void MeshPaintHelpers::ForceRenderMeshLOD(UMeshComponent* Component, int32 LODIndex)
 {
 	if (UStaticMeshComponent* StaticMeshComponent = Cast<UStaticMeshComponent>(Component))
diff --git a/Engine/Source/Editor/MeshPaint/Public/IMeshPainter.h b/Engine/Source/Editor/MeshPaint/Public/IMeshPainter.h
--- a/Engine/Source/Editor/MeshPaint/Public/IMeshPainter.h
+++ b/Engine/Source/Editor/MeshPaint/Public/IMeshPainter.h
@@ -10,8 +10,6 @@
 class UMeshPaintSettings;
 class UPaintBrushSettings;
 class UMeshComponent;
-class UViewportInteractor;
-class UVREditorInteractor;
 class IMeshPaintGeometryAdapter;
 class FScopedTransaction;
 class FUICommandList;
@@ -51,9 +49,6 @@ public:
 	/** 'Normal' painting functionality, called when the user tries to paint on a mesh using the mouse */
 	virtual bool Paint(FViewport* Viewport, const FVector& InCameraOrigin, const TArrayView<TPair<FVector, FVector>>& Rays);
 
-	/** VR painting functionality, called when the user tries to paint on a mesh using a VR controller */
-	virtual bool PaintVR(FViewport* Viewport, const FVector& InCameraOrigin, const FVector& InRayOrigin, const FVector& InRayDirection, UVREditorInteractor* VREditorInteractor);
-
 	/** Allows painter to act on specific key actions */
 	virtual bool InputKey(FEditorViewportClient* InViewportClient, FViewport* InViewport, FKey InKey, EInputEvent InEvent);
 
@@ -97,9 +92,7 @@ protected:
 
 	/** Functionality for rendering the interactor widget given its positional information and paint action */
 	virtual void RenderInteractorWidget(const FVector& InCameraOrigin, const FVector& InRayOrigin, const FVector& InRayDirection, FPrimitiveDrawInterface* PDI, EMeshPaintAction PaintAction, bool bRenderVertices, ESceneDepthPriorityGroup DepthGroup = SDPG_World);
-protected:
-	/** When painting in VR, this is viewport interactor we are using */
-	UViewportInteractor* CurrentViewportInteractor;	
+
 protected:
 	/** Flag for whether or not we are currently painting */
 	bool bArePainting;
diff --git a/Engine/Source/Editor/MeshPaintMode/MeshPaintMode.Build.cs b/Engine/Source/Editor/MeshPaintMode/MeshPaintMode.Build.cs
--- a/Engine/Source/Editor/MeshPaintMode/MeshPaintMode.Build.cs
+++ b/Engine/Source/Editor/MeshPaintMode/MeshPaintMode.Build.cs
@@ -35,22 +35,12 @@ public class MeshPaintMode : ModuleRules
                 "UnrealEd",
                 "RawMesh",
                 "SourceControl",
-                "ViewportInteraction",
-                "VREditor",
                 "PropertyEditor",
                 "MainFrame",
 				"MeshPaint",
             }
         );
 
-        CircularlyReferencedDependentModules.AddRange(
-            new string[]
-            {
-                "ViewportInteraction",
-                "VREditor"
-            }
-        );
-
         PrivateIncludePathModuleNames.AddRange(
 			new string[]
 			{
diff --git a/Engine/Source/Editor/MeshPaintMode/Private/PaintModePainter.cpp b/Engine/Source/Editor/MeshPaintMode/Private/PaintModePainter.cpp
--- a/Engine/Source/Editor/MeshPaintMode/Private/PaintModePainter.cpp
+++ b/Engine/Source/Editor/MeshPaintMode/Private/PaintModePainter.cpp
@@ -26,13 +26,6 @@
 #include "MeshPaintTypes.h"
 #include "TexturePaintHelpers.h"
 
-#include "VREditorMode.h"
-#include "IVREditorModule.h"
-#include "ViewportWorldInteraction.h"
-#include "ViewportInteractableInterface.h"
-#include "VREditorInteractor.h"
-#include "EditorWorldExtension.h"
-
 #include "Framework/Commands/UICommandList.h"
 #include "PaintModeCommands.h"
 #include "PaintModeSettings.h"
diff --git a/Engine/Source/Editor/UnrealEd/Private/PlayLevel.cpp b/Engine/Source/Editor/UnrealEd/Private/PlayLevel.cpp
--- a/Engine/Source/Editor/UnrealEd/Private/PlayLevel.cpp
+++ b/Engine/Source/Editor/UnrealEd/Private/PlayLevel.cpp
@@ -482,9 +482,8 @@ void UEditorEngine::EndPlayMap()
 		GEngine->PendingDroppedNotes.Empty();
 	}
 
-	//ensure stereo rendering is disabled in case we need to re-enable next PIE run (except when the editor is running in VR)
-	bool bInVRMode = IVREditorModule::Get().IsVREditorModeActive();
-	if (GEngine->StereoRenderingDevice && !bInVRMode)
+	//ensure stereo rendering is disabled in case we need to re-enable next PIE run 
+	if (GEngine->StereoRenderingDevice)
 	{
 		GEngine->StereoRenderingDevice->EnableStereo(false);
 	}
diff --git a/Engine/Source/Editor/VREditor/Public/VREditorMode.h b/Engine/Source/Editor/VREditor/Public/VREditorMode.h
--- a/Engine/Source/Editor/VREditor/Public/VREditorMode.h
+++ b/Engine/Source/Editor/VREditor/Public/VREditorMode.h
@@ -266,12 +266,6 @@ public:
 
 	void SaveSequencerSettings(bool bInKeyAllEnabled, EAutoChangeMode InAutoChangeMode, const class USequencerSettings& InSequencerSettings);
 
-	/** Start or stop simulate-in-editor mode */
-	void ToggleSIEAndVREditor();
-
-	/** Start or stop play-in-editor mode */
-	void TogglePIEAndVREditor();
-
 	/** Create a static motion controller mesh for the current HMD platform */
 	UStaticMeshComponent* CreateMotionControllerMesh( AActor* OwningActor, USceneComponent* AttachmentToComponent, UStaticMesh* OptionalControllerMesh = nullptr );
 
@@ -279,21 +273,6 @@ public:
 	UStaticMeshComponent* CreateMesh( AActor* OwningActor, const FString& MeshName, USceneComponent* AttachmentToComponent /*= nullptr */ );
 	UStaticMeshComponent* CreateMesh(AActor* OwningActor, UStaticMesh* Mesh, USceneComponent* AttachmentToComponent /*= nullptr */);
 
-	/** Sets a delegate for the context-specific actions menu */
-	void SetActionsMenuGenerator(const FOnRadialMenuGenerated NewMenuGenerator, const FText NewLabel);
-
-	/** Resets the delegate and button for the context-specific actions menu */
-	void ResetActionsMenuGenerator();
-
-	/** Gets access to VREditorWorldInteraction */
-	class UVREditorPlacement* GetPlacementSystem()
-	{
-		return PlacementSystem;
-	}
-
-	/** Returns true if we started the play in editor session from this VR Editor */
-	bool GetStartedPlayFromVREditor() const;
-
 	/** Gets the container for all the assets of VREditor. */
 	const class UVREditorAssetContainer& GetAssetContainer() const;
 
@@ -311,10 +290,6 @@ public:
 	DECLARE_EVENT_OneParam(UVREditorPlacement, FOnPlacePreviewActor, bool);
 	FOnPlacePreviewActor& OnPlacePreviewActor() { return OnPlacePreviewActorEvent; };
 
-	/** Call this to force the 'Actions' radial menu to refresh.  This is useful if the menu generator that you've bound
-		needs to be re-run (usually because it switches on something that has changed since the last time it ran.) */
-	void RefreshRadialMenuActionsSubmenu();
-
 	/** Return true if currently aiming to teleport. */
 	bool IsAimingTeleport() const;
 	bool IsTeleporting() const;
@@ -349,12 +324,6 @@ private:
 	/** Called when someone closes a standalone VR Editor window */
 	void OnVREditorWindowClosed( const TSharedRef<SWindow>& ClosedWindow );
 
-	/** FEditorDelegates callbacks */
-	void PostPIEStarted( bool bIsSimulatingInEditor );
-	void PrePIEEnded( bool bWasSimulatingInEditor );
-	void OnEndPIE( bool bWasSimulatingInEditor );
-	void OnPreSwitchPIEAndSIE(bool bIsSimulatingInEditor);
-	void OnSwitchPIEAndSIE(bool bIsSimulatingInEditor);
 
 	/** Start using the viewport passed */
 	void StartViewport( TSharedPtr<SLevelViewport> Viewport );
@@ -362,9 +331,6 @@ private:
 	/** Close the current viewport */
 	void CloseViewport( const bool bShouldDisableStereo );
 
-	/** Resets all the settings when exiting PIE to VR Editor. */
-	void RestoreFromPIE();
-
 	/** Restore the world to meters to the saved one when entering VR Editor */
 	void RestoreWorldToMeters();
 
@@ -524,14 +490,6 @@ private:
 	/** Pointer to the current Sequencer */
 	class ISequencer* CurrentSequencer;
 
-	/** The world to meters scale when leaving PIE simulate to restore when back in the editor world. */
-	float SavedWorldToMetersScaleForPIE;
-
-	/** If we started play in editor from the VR Editor*/
-	bool bStartedPlayFromVREditor;
-
-	/** If we started play in editor from the VR Editor while in simulate. */
-	bool bStartedPlayFromVREditorSimulate;
 
 	/** Container of assets */
 	UPROPERTY()
diff --git a/Engine/Source/Editor/VREditor/UI/VREditorUISystem.cpp b/Engine/Source/Editor/VREditor/UI/VREditorUISystem.cpp
--- a/Engine/Source/Editor/VREditor/UI/VREditorUISystem.cpp
+++ b/Engine/Source/Editor/VREditor/UI/VREditorUISystem.cpp
@@ -28,10 +28,7 @@
 
 #include "Engine/EngineTypes.h"
 #include "Editor/UnrealEdEngine.h"
-#include "GameFramework/WorldSettings.h"
 #include "Sound/SoundCue.h"
-#include "EditorModeManager.h"
-#include "EditorModes.h"
 #include "UnrealEdGlobals.h"
 
 #include "VREditorMode.h"
@@ -58,15 +55,6 @@
 #include "Widgets/Layout/SUniformGridPanel.h"
 #include "Widgets/Layout/SScrollBox.h"
 
-// Content Browser
-#include "CollectionManagerTypes.h"
-#include "IContentBrowserSingleton.h"
-#include "ContentBrowserModule.h"
-
-// World Outliner
-#include "Editor/SceneOutliner/Public/ISceneOutliner.h"
-#include "Editor/SceneOutliner/Public/SceneOutlinerPublicTypes.h"
-#include "Editor/SceneOutliner/Public/SceneOutlinerModule.h"
 
 // Actor Details, Modes
 #include "SEditorViewport.h"
@@ -75,8 +63,6 @@
 #include "LevelEditor.h"
 
 // World Settings
-#include "PropertyEditorModule.h"
-#include "IDetailsView.h"
 
 #include "SLevelViewport.h"
 #include "Widgets/Layout/SScaleBox.h"
@@ -106,13 +92,12 @@
 #include "VREditorFloatingCameraUI.h"
 #include "ScopedTransaction.h"
 #include "Subsystems/AssetEditorSubsystem.h"
+#include "EditorModes.h"
 
 #define LOCTEXT_NAMESPACE "VREditor"
 
 namespace VREd
 {
-	static FAutoConsoleVariable ContentBrowserUIResolutionX( TEXT( "VREd.ContentBrowserUIResolutionX" ), 1920, TEXT( "Horizontal resolution to use for content browser UI render targets" ) );
-	static FAutoConsoleVariable ContentBrowserUIResolutionY( TEXT( "VREd.ContentBrowserUIResolutionY" ), 1200, TEXT( "Vertical resolution to use for content browser UI render targets" ) );
 	static FAutoConsoleVariable SequencerUIResolutionX(TEXT("VREd.SequencerUIResolutionX"), 960, TEXT("Horizontal resolution to use for Sequencer UI render targets"));
 	static FAutoConsoleVariable SequencerUIResolutionY(TEXT("VREd.SequencerUIResolutionY"), 600, TEXT("Vertical resolution to use for Sequencer UI render targets"));
 	static FAutoConsoleVariable DefaultEditorUIResolutionX( TEXT( "VREd.DefaultEditorUIResolutionX" ), 1024, TEXT( "Horizontal resolution to use for VR editor UI render targets" ) );
@@ -121,9 +106,7 @@ namespace VREd
 	static FAutoConsoleVariable DefaultRadialElementResolutionY(TEXT("VREd.DefaultRadialElementResolutionY"), 350, TEXT("Vertical resolution to use for VR editor radial UI render targets"));
 	static FAutoConsoleVariable QuickMenuUIResolutionX( TEXT( "VREd.QuickMenuUIResolutionX" ), 1200, TEXT( "Horizontal resolution to use for Quick Menu VR UI render targets" ) );
 	static FAutoConsoleVariable QuickMenuUIResolutionY( TEXT( "VREd.QuickMenuUIResolutionY" ), 1075, TEXT( "Vertical resolution to use for Quick Menu VR UI render targets" ) );
-	static FAutoConsoleVariable ContentBrowserUISize( TEXT( "VREd.ContentBrowserUISize" ), 80.0f, TEXT( "How big content browser UIs should be" ) );
 	static FAutoConsoleVariable EditorUISize( TEXT( "VREd.EditorUISize" ), 70.0f, TEXT( "How big editor UIs should be" ) );
-	static FAutoConsoleVariable ContentBrowserUIScale( TEXT( "VREd.ContentBrowserUIScale" ), 2.0f, TEXT( "How much to scale up (or down) the content browser for VR" ) );
 	static FAutoConsoleVariable EditorUIScale( TEXT( "VREd.EditorUIScale" ), 2.0f, TEXT( "How much to scale up (or down) editor UIs for VR" ) );
 	static FAutoConsoleVariable AssetEditorUIResolutionX( TEXT( "VREd.AssetEditorUIResolutionX" ), 1920, TEXT( "Horizontal resolution to use for VR editor asset editor UI render targets" ) );
 	static FAutoConsoleVariable AssetEditorUIResolutionY( TEXT( "VREd.AssetEditorUIResolutionY" ), 1080, TEXT( "Vertical resolution to use for VR editor asset editor UI render targets" ) );
@@ -137,7 +120,6 @@ namespace VREd
 	static FAutoConsoleVariable MinDockWindowSize( TEXT( "VREd.MinDockWindowSize" ), 40, TEXT( "Minimum size for dockable windows" ) );
 	static FAutoConsoleVariable UIPanelOpenDistance( TEXT( "VREd.UIPanelOpenDistance" ), 20.f, TEXT( "Distance to spawn a panel from the hand in centimeters" ) );
 	static FAutoConsoleVariable UIPanelOpenRotationPitchOffset( TEXT( "VREd.UIPanelOpenRotationPitchOffset" ), 45.0f, TEXT( "The pitch rotation offset in degrees when spawning a panel in front of the motioncontroller" ) );
-	static FAutoConsoleVariable ColorPickerDockSpawnOffset( TEXT( "VREd.ColorPickerDockSpawnOffset" ),  3.f, TEXT( "Offset of where the color picker spawns" ) );
 	static FAutoConsoleVariable SteamVRTrackpadDeadzone(TEXT("VREd.SteamVRTrackpadDeadzone"), 0.3f, TEXT("The deadzone for the Vive motion controller trackpad"));
 	static const FTransform DefaultColorPickerTransform(	FRotator( -10, 180, 0), FVector( 30, 35, 0 ), FVector( 1 ) );
 	static const FTransform DefaultActorPreviewTransform(FRotator(-20, 120, 0), FVector(50, -60, 0), FVector(1));
@@ -146,13 +128,7 @@ namespace VREd
 	static FAutoConsoleVariable CameraPreviewUISize( TEXT( "VREd.CameraPreviewUISize" ), 50.0f, TEXT( "How big camera preview UIs should be" ) );
 }
 
-const VREditorPanelID UVREditorUISystem::ContentBrowserPanelID = VREditorPanelID("ContentBrowser");
-const VREditorPanelID UVREditorUISystem::WorldOutlinerPanelID = VREditorPanelID("WorldOutliner");
-const VREditorPanelID UVREditorUISystem::DetailsPanelID = VREditorPanelID("Details");
-const VREditorPanelID UVREditorUISystem::ModesPanelID = VREditorPanelID("Modes");
-const VREditorPanelID UVREditorUISystem::TutorialPanelID = VREditorPanelID("Tutorial");
-const VREditorPanelID UVREditorUISystem::WorldSettingsPanelID = VREditorPanelID("WorldSettings");
-const VREditorPanelID UVREditorUISystem::ColorPickerPanelID = VREditorPanelID("ColorPicker");
+
 const VREditorPanelID UVREditorUISystem::SequencerPanelID = VREditorPanelID("SequencerUI");
 const VREditorPanelID UVREditorUISystem::InfoDisplayPanelID = VREditorPanelID("InfoDisplay");
 const VREditorPanelID UVREditorUISystem::RadialMenuPanelID = VREditorPanelID("RadialMenu");
@@ -202,10 +178,6 @@ void UVREditorUISystem::Init(UVREditorMode* InVRMode)
 
 	CreateUIs();
 
-	// Bind the color picker creation & destruction overrides
-	SColorPicker::OnColorPickerNonModalCreateOverride.BindUObject(this, &UVREditorUISystem::CreateVRColorPicker);
-	SColorPicker::OnColorPickerDestroyOverride.BindUObject(this, &UVREditorUISystem::DestroyVRColorPicker);
-
 	// Bind the global tab manager's dockable area restore override
 	FGlobalTabmanager::Get()->OnOverrideDockableAreaRestore_Handler.BindUObject(this, &UVREditorUISystem::DockableAreaRestored);
 
@@ -214,7 +186,6 @@ void UVREditorUISystem::Init(UVREditorMode* InVRMode)
 	FVREditorActionCallbacks::UpdateSelectingCandidateActorsText(VRMode);
 	FVREditorActionCallbacks::SelectingCandidateActorsText = FVREditorActionCallbacks::GetSelectingCandidateActorsText();
 
-	GLevelEditorModeTools().OnEditorModeIDChanged().AddUObject(this, &UVREditorUISystem::HandleEditorModeChanged);
 	AssetEditorCloseDelegate = GEditor->GetEditorSubsystem<UAssetEditorSubsystem>()->OnAssetEditorRequestClose().AddUObject(this, &UVREditorUISystem::OnCloseAssetEditor);
 
 }
@@ -938,180 +909,6 @@ void UVREditorUISystem::CreateUIs()
 	}
 	// Make some editor UIs!
 	{
-		{
-			const FIntPoint Resolution = FIntPoint( VREd::ContentBrowserUIResolutionX->GetInt(), VREd::ContentBrowserUIResolutionY->GetInt() );
-
-			IContentBrowserSingleton& ContentBrowserSingleton = FModuleManager::Get().LoadModuleChecked<FContentBrowserModule>( "ContentBrowser" ).Get();;
-
-			// @todo vreditor UI: We've turned off a LOT of content browser features that are clunky to use in VR right (pop ups, text fields, etc.)
-			FContentBrowserConfig Config;
-			Config.bCanShowClasses = false;
-			Config.bCanShowRealTimeThumbnails = false;
-			Config.InitialAssetViewType = EAssetViewType::Tile;
-			Config.bCanShowDevelopersFolder = false;
-			Config.bCanShowFolders = false;
-			Config.bUseSourcesView = true;
-			Config.bExpandSourcesView = true;
-			Config.ThumbnailLabel = EThumbnailLabel::NoLabel;
-			Config.ThumbnailScale = 0.4f;
-			Config.bCanShowFilters = true;
-			Config.bUsePathPicker = true;
-			Config.bShowBottomToolbar = true;
-			Config.bCanShowLockButton = false;
-
-			if( GIsDemoMode )
-			{
-				Config.SelectedCollectionName.Name = TEXT( "Demo" );
-				Config.SelectedCollectionName.Type = ECollectionShareType::CST_Shared;
-			}
-
-			TSharedRef<SWidget> ContentBrowser = ContentBrowserSingleton.CreateContentBrowser( "VRContentBrowser", nullptr, &Config );
-
-			TSharedRef<SWidget> WidgetToDraw =
-				SNew( SDPIScaler )
-				.DPIScale( VREd::ContentBrowserUIScale->GetFloat() )
-				[
-					ContentBrowser
-				]
-			;
-
-			const bool bWithSceneComponent = false;
-			AVREditorFloatingUI* ContentBrowserUI = GetOwner().SpawnTransientSceneActor<AVREditorDockableWindow>(TEXT( "ContentBrowserUI" ), bWithSceneComponent);
-			ContentBrowserUI->SetSlateWidget( *this, ContentBrowserPanelID, WidgetToDraw, Resolution, VREd::ContentBrowserUISize->GetFloat(), AVREditorFloatingUI::EDockedTo::Nothing );
-			ContentBrowserUI->ShowUI( bShowUI );
-			FloatingUIs.Add(ContentBrowserPanelID, ContentBrowserUI);
-		}
-
-		{
-			FSceneOutlinerModule& SceneOutlinerModule = FModuleManager::Get().LoadModuleChecked<FSceneOutlinerModule>( "SceneOutliner" );
-
-			SceneOutliner::FInitializationOptions InitOptions;
-			InitOptions.Mode = ESceneOutlinerMode::ActorBrowsing;
-
-			TSharedRef<ISceneOutliner> SceneOutliner = SceneOutlinerModule.CreateSceneOutliner( InitOptions, FOnActorPicked() /* Not used for outliner when in browsing mode */ );
-			TSharedRef<SWidget> WidgetToDraw =
-				SNew( SDPIScaler )
-				.DPIScale( VREd::EditorUIScale->GetFloat() )
-				[
-					SNew(SBorder)
-					.BorderImage(FEditorStyle::Get().GetBrush("ToolPanel.GroupBorder"))
-					[
-						SceneOutliner
-					]
-				]
-			;
-
-			const bool bWithSceneComponent = false;
-			AVREditorFloatingUI* WorldOutlinerUI = GetOwner().SpawnTransientSceneActor<AVREditorDockableWindow>(TEXT("WorldOutlinerUI"), bWithSceneComponent);
-			WorldOutlinerUI->SetSlateWidget( *this, WorldOutlinerPanelID, WidgetToDraw, DefaultResolution, VREd::EditorUISize->GetFloat(), AVREditorFloatingUI::EDockedTo::Nothing );
-			WorldOutlinerUI->ShowUI( bShowUI );
-			FloatingUIs.Add(WorldOutlinerPanelID, WorldOutlinerUI);
-		}
-
-		{
-			const TSharedRef< ILevelEditor >& LevelEditor = FModuleManager::GetModuleChecked<FLevelEditorModule>( "LevelEditor" ).GetFirstLevelEditor().ToSharedRef();
-
-			const FName TabIdentifier = NAME_None;	// No tab for us!
-			TSharedRef<SWidget> ActorDetails = LevelEditor->CreateActorDetails( TabIdentifier );
-
-			TSharedRef<SWidget> WidgetToDraw =
-				SNew( SDPIScaler )
-				.DPIScale( VREd::EditorUIScale->GetFloat() )
-				[
-					ActorDetails
-				]
-			;
-
-			const bool bWithSceneComponent = false;
-			AVREditorFloatingUI* ActorDetailsUI = GetOwner().SpawnTransientSceneActor<AVREditorDockableWindow>(TEXT("ActorDetailsUI"), bWithSceneComponent);
-			ActorDetailsUI->SetSlateWidget( *this, DetailsPanelID, WidgetToDraw, DefaultResolution, VREd::EditorUISize->GetFloat(), AVREditorFloatingUI::EDockedTo::Nothing );
-			ActorDetailsUI->ShowUI( bShowUI );
-			FloatingUIs.Add(DetailsPanelID, ActorDetailsUI);
-		}
-
-		{
-			const TSharedRef< ILevelEditor >& LevelEditor = FModuleManager::GetModuleChecked<FLevelEditorModule>( "LevelEditor" ).GetFirstLevelEditor().ToSharedRef();
-
-			TSharedRef<SWidget> Modes = LevelEditor->CreateToolBox();
-
-			TSharedRef<SWidget> WidgetToDraw =
-				SNew( SDPIScaler )
-				.DPIScale( VREd::EditorUIScale->GetFloat() )
-				[
-					SNew(SBorder)
-					.BorderImage(FEditorStyle::Get().GetBrush("ToolPanel.GroupBorder"))
-					[
-						Modes
-					]
-				]
-			;
-
-			const bool bWithSceneComponent = false;
-			AVREditorFloatingUI* ModesUI = GetOwner().SpawnTransientSceneActor<AVREditorDockableWindow>(TEXT("ModesUI"), bWithSceneComponent);
-			ModesUI->SetSlateWidget( *this, ModesPanelID, WidgetToDraw, DefaultResolution, VREd::EditorUISize->GetFloat(), AVREditorFloatingUI::EDockedTo::Nothing );
-			ModesUI->ShowUI( bShowUI );
-			FloatingUIs.Add(ModesPanelID, ModesUI);
-
-			// @todo vreditor placement: This is required to force the modes UI to refresh -- otherwise it looks empty
-			GLevelEditorModeTools().ActivateDefaultMode();
-		}
-
-		{
-			const FIntPoint Resolution = FIntPoint(VREd::AssetEditorUIResolutionX->GetInt(), VREd::AssetEditorUIResolutionY->GetInt());
-
-			const bool bWithSceneComponent = false;
-			AVREditorDockableWindow* TabManagerUI = GetOwner().SpawnTransientSceneActor<AVREditorDockableWindow>(TEXT("AssetEditor"), bWithSceneComponent);
-			TabManagerUI->SetSlateWidget( *this, TabManagerPanelID,  SNullWidget::NullWidget, Resolution, VREd::EditorUISize->GetFloat(), AVREditorFloatingUI::EDockedTo::Nothing );
-			TabManagerUI->ShowUI( bShowUI );
-
-			// @todo vreditor: Could use "Hovering" instead for better performance with many UIs, but needs to be manually refreshed in some cases
-			TabManagerUI->GetWidgetComponent()->SetDrawingPolicy( EVREditorWidgetDrawingPolicy::Always );
-
-			FloatingUIs.Add(TabManagerPanelID, TabManagerUI);
-
-			TSharedPtr<SWindow> TabManagerWindow = TabManagerUI->GetWidgetComponent()->GetSlateWindow();
-			TSharedRef<SWindow> TabManagerWindowRef = TabManagerWindow.ToSharedRef();
-			ProxyTabManager = MakeShareable( new FProxyTabmanager( TabManagerWindowRef ) );
-
-			ProxyTabManager->OnIsTabSupported.AddUObject( this, &UVREditorUISystem::IsProxyTabSupported );
-			ProxyTabManager->OnTabOpened.AddUObject( this, &UVREditorUISystem::OnProxyTabLaunched );
-			ProxyTabManager->OnAttentionDrawnToTab.AddUObject( this, &UVREditorUISystem::OnAttentionDrawnToTab );
-
-			// We're going to start stealing tabs from the global tab manager inserting them into the world instead.
-			FGlobalTabmanager::Get()->SetProxyTabManager( ProxyTabManager );
-
-			GEditor->GetEditorSubsystem<UAssetEditorSubsystem>()->OnAssetEditorOpened().AddUObject( this, &UVREditorUISystem::OnAssetEditorOpened );
-		}
-		// create the world settings menu
-		{
-			FPropertyEditorModule& PropertyEditorModule = FModuleManager::Get().LoadModuleChecked< FPropertyEditorModule >( "PropertyEditor" );
-
-			FDetailsViewArgs DetailsViewArgs( false, false, true, FDetailsViewArgs::HideNameArea, false, GUnrealEd );
-			DetailsViewArgs.bShowActorLabel = false;
-
-			TSharedRef<IDetailsView> WorldSettings = PropertyEditorModule.CreateDetailView( DetailsViewArgs );
-
-			UWorld* World = GetOwner().GetWorld();
-			if ( World != NULL )
-			{
-				WorldSettings->SetObject( World->GetWorldSettings() );
-			}
-
-			TSharedRef<SWidget> WidgetToDraw = 
-				SNew( SDPIScaler )
-				.DPIScale( VREd::EditorUIScale->GetFloat() )
-				[
-					WorldSettings
-				]
-			;
-
-			const bool bWithSceneComponent = false;
-			AVREditorFloatingUI* WorldSettingsUI = GetOwner().SpawnTransientSceneActor<AVREditorDockableWindow>(TEXT("WorldSettingsUI"), bWithSceneComponent);
-			WorldSettingsUI->SetSlateWidget( *this, WorldSettingsPanelID, WidgetToDraw, DefaultResolution, VREd::EditorUISize->GetFloat(), AVREditorFloatingUI::EDockedTo::Nothing );
-			WorldSettingsUI->ShowUI( bShowUI );
-			FloatingUIs.Add(WorldSettingsPanelID, WorldSettingsUI);
-		}
-
 		{
 			const FName TabIdentifier = NAME_None;	// No tab for us!
 			const TSharedRef<SWidget> SequencerWidget = SNullWidget::NullWidget;
@@ -1126,7 +923,7 @@ void UVREditorUISystem::CreateUIs()
 
 			const bool bWithSceneComponent = false;
 			AVREditorFloatingUI* SequencerUI = GetOwner().SpawnTransientSceneActor<AVREditorDockableWindow>(TEXT("SequencerUI"), bWithSceneComponent);
-			SequencerUI->SetSlateWidget(*this, SequencerPanelID, WidgetToDraw, FIntPoint(VREd::SequencerUIResolutionX->GetFloat(), VREd::SequencerUIResolutionY->GetFloat()), VREd::ContentBrowserUISize->GetFloat(), AVREditorFloatingUI::EDockedTo::Nothing);
+			SequencerUI->SetSlateWidget(*this, SequencerPanelID, WidgetToDraw, FIntPoint(VREd::SequencerUIResolutionX->GetFloat(), VREd::SequencerUIResolutionY->GetFloat()), VREd::EditorUISize->GetFloat(), AVREditorFloatingUI::EDockedTo::Nothing);
 			SequencerUI->ShowUI(false);
 			FloatingUIs.Add(SequencerPanelID, SequencerUI);
 		}
@@ -1302,11 +1099,6 @@ void UVREditorUISystem::ShowEditorUIPanel(AVREditorFloatingUI* Panel, UVREditorI
 					FVREditorActionCallbacks::CloseSequencer(GetOwner().GetCurrentSequencer()->GetRootMovieSceneSequence());
 				}
 			}
-			else if (PanelID == ModesPanelID)
-			{
-				// Quit active mode and go back to Placement Mode when closing the Modes panel.
-				GLevelEditorModeTools().DeactivateAllModes();
-			}
 			else if (PanelID.ToString().Contains(ActorPreviewPrefix))
 			{
 				// if we are closing a preview panel, make sure to unpin it also
@@ -1623,94 +1415,7 @@ void UVREditorUISystem::DockableAreaRestored()
 {
 }
 
-void UVREditorUISystem::CreateVRColorPicker(const TSharedRef<SColorPicker>& ColorPicker)
-{
-	// Check that the Color Picker Panel isn't currently showing. Also handles the EEditorUIPanel::ColorPicker being null.
-	const FIntPoint DefaultResolution(VREd::DefaultEditorUIResolutionX->GetInt(), VREd::DefaultEditorUIResolutionY->GetInt());
-
-	TSharedRef<SWidget> WidgetToDraw =
-		SNew(SDPIScaler)
-		.DPIScale(VREd::EditorUIScale->GetFloat())
-		[
-			ColorPicker
-		]
-	;
-
-	// Only spawn the Color picker dockable window if there is none yet
-	bool bJustSpawned = false;
-	if( ColorPickerUI == nullptr )
-	{
-		const bool bWithSceneComponent = false;
-		ColorPickerUI = GetOwner().SpawnTransientSceneActor<AVREditorDockableWindow>(TEXT("ColorPickerUI"), bWithSceneComponent);
-		FloatingUIs.Add(ColorPickerPanelID, ColorPickerUI);
-		bJustSpawned = true;
-	}
-
-	ColorPickerUI->SetSlateWidget( *this, ColorPickerPanelID, WidgetToDraw, DefaultResolution, VREd::EditorUISize->GetFloat(), AVREditorFloatingUI::EDockedTo::Room );
-
-	// Always spawn based on the location of the menu you are hovering over. To get this information, find the hand hovering over the UI.
-	UVREditorInteractor* VREditorInteractor = GetOwner().GetHandInteractor(EControllerHand::Left);
-	if (!VREditorInteractor->IsHoveringOverUI())
-	{
-		VREditorInteractor = GetOwner().GetHandInteractor(EControllerHand::Right);
-	}
-
-	// Spawn the color picker close to the pressed UI it was opened using UI. Otherwise spawn it in front of the user.
-	UWidgetComponent* HoveredWidgetComponent = VREditorInteractor->GetLastHoveredWidgetComponent();
-	if( HoveredWidgetComponent != nullptr )
-	{
-		FVector LaserPointerStart, LaserPointerEnd;
-		if( VREditorInteractor->GetLaserPointer( LaserPointerStart, LaserPointerEnd ) )
-		{
-			//Calculate new location using the direction of the laser with an offset towards the user
-			FVector LaserPointerDirection = LaserPointerEnd - LaserPointerStart;
-			LaserPointerDirection.Normalize();
-			const FVector HoverOverUILocation = VREditorInteractor->GetHoverLocation();
-			const FVector Offset = LaserPointerDirection * VREd::ColorPickerDockSpawnOffset->GetFloat() * VRMode->GetWorldScaleFactor();
-
-			FTransform NewTransform = FTransform( HoveredWidgetComponent->GetComponentQuat(), HoverOverUILocation - Offset );
-			
-			// Move the color picker if it was already spawned, or instantly set the transform if it was just spawned. We don't want to move the UI from location zero
-			if( bJustSpawned )
-			{
-				ColorPickerUI->SetDockedTo( AVREditorBaseActor::EDockedTo::Room );
-				ColorPickerUI->SetTransform( NewTransform );
-			}
-			else
-			{
-				ColorPickerUI->MoveTo( NewTransform, 0.2f, AVREditorBaseActor::EDockedTo::Room );
-			}
-
-			ColorPickerUI->SetActorTransform( NewTransform );
-			ColorPickerUI->UpdateRelativeRoomTransform();
-		}
-	}
-	else
-	{
-		// Make sure the UIs are centered around the direction your head is looking (yaw only!)
-		const FVector RoomSpaceHeadLocation = VRMode->GetRoomSpaceHeadTransform().GetLocation() / VRMode->GetWorldScaleFactor();
-		FRotator RoomSpaceHeadYawRotator = VRMode->GetRoomSpaceHeadTransform().GetRotation().Rotator();
-		RoomSpaceHeadYawRotator.Pitch = 0.0f;
-		RoomSpaceHeadYawRotator.Roll = 0.0f;
-	
-		FTransform NewTransform = VREd::DefaultColorPickerTransform;
-		NewTransform *= FTransform(RoomSpaceHeadYawRotator.Quaternion(), FVector::ZeroVector);
-		ColorPickerUI->SetLocalRotation(NewTransform.GetRotation().Rotator());
-		ColorPickerUI->SetRelativeOffset(RoomSpaceHeadLocation + NewTransform.GetTranslation());
-	}
-
-	ColorPickerUI->ShowUI( true );
-}
 
-void UVREditorUISystem::DestroyVRColorPicker()
-{
-	UVREditorInteractor* VREditorInteractor = nullptr;
-	const bool bShouldShow = false;
-	const bool bSpawnInFront = false;
-	const bool bDragFromOpen = false;
-	const bool bPlaySound = true;
-	ShowEditorUIPanel(ColorPickerPanelID, VREditorInteractor, bShouldShow, bSpawnInFront, bDragFromOpen, bPlaySound);
-}
 
 void UVREditorUISystem::MakeUniformGridMenu(const TSharedRef<FMultiBox>& MultiBox, const TSharedRef<SMultiBoxWidget>& MultiBoxWidget, int32 Columns)
 {
@@ -2038,40 +1743,7 @@ void UVREditorUISystem::BuildNumPadWidget()
 
 	// Create the menu widget
 	MenuBuilder.MakeWidget(&VREditorMenuBuilderOverride);
-	/**
-	TSharedRef<SWidget> FinalWidget = SNew( SOverlay )
-		+ SOverlay::Slot()
-		[
-			MenuWidget
-		]
-		+ SOverlay::Slot()
-		.HAlign( HAlign_Center )
-		.VAlign( VAlign_Center )
-		[
-			SNew(SVerticalBox)
-			+ SVerticalBox::Slot()
-			.HAlign( HAlign_Fill )
-			.VAlign( VAlign_Center )
-			.AutoHeight()
-			[
-				SNew( STextBlock )
-				.Justification( ETextJustify::Center )
-				.Text( LOCTEXT( "BackspaceHelp", "Modifier: Backspace" ) )
-				.TextStyle( FVREditorStyle::Get(), "VREditorStyle.HelperText" )
-			]
-		+ SVerticalBox::Slot()
-			.HAlign( HAlign_Fill )
-			.VAlign( VAlign_Center )
-			.AutoHeight()
-			[
-				SNew( STextBlock )
-				.Justification( ETextJustify::Center )
-				.Text( LOCTEXT( "EnterHelp", "Grip: Enter" ) )
-			.	TextStyle( FVREditorStyle::Get(), "VREditorStyle.HelperText" )
-			]
-		]
-		;
-		*/
+
 }
 
 const TSharedRef<SWidget>& UVREditorUISystem::FindWidgetOfType(const TSharedRef<SWidget>& Content, FName WidgetType)
@@ -2175,118 +1847,6 @@ AVREditorFloatingUI* UVREditorUISystem::GetPanel(const VREditorPanelID& InPanelI
 	return Result;
 }
 
-void UVREditorUISystem::SequencerRadialMenuGenerator(FMenuBuilder& MenuBuilder, TSharedPtr<FUICommandList> CommandList, UVREditorMode* InVRMode, float& RadiusOverride)
-{
-	RadiusOverride = 1.0f;
-	// First menu entry is at 90 degrees 
-	MenuBuilder.AddMenuEntry(
-		LOCTEXT("Play", "Play"),
-		FText(),
-		FSlateIcon(FVREditorStyle::GetStyleSetName(), "VREditorStyle.SequencerPlay"),
-		FUIAction
-		(
-			FExecuteAction::CreateStatic(&FVREditorActionCallbacks::PlaySequenceAtRate, InVRMode, 1.0f),
-			FCanExecuteAction::CreateUObject(this->GetRadialMenuHandler(), &UVRRadialMenuHandler::IsActionMenuBound)
-		)
-		);
-	MenuBuilder.AddMenuEntry(
-		LOCTEXT("Reverse", "Reverse"),
-		FText(),
-		FSlateIcon(FVREditorStyle::GetStyleSetName(), "VREditorStyle.SequencerReverse"),
-		FUIAction
-		(
-			FExecuteAction::CreateStatic(&FVREditorActionCallbacks::PlaySequenceAtRate, InVRMode, -1.0f),
-			FCanExecuteAction::CreateUObject(this->GetRadialMenuHandler(), &UVRRadialMenuHandler::IsActionMenuBound)
-		)
-		);
-	MenuBuilder.AddMenuEntry(
-		LOCTEXT("Stop", "Stop"),
-		FText(),
-		FSlateIcon(FVREditorStyle::GetStyleSetName(), "VREditorStyle.SequencerStop"),
-		FUIAction
-		(
-			FExecuteAction::CreateStatic(&FVREditorActionCallbacks::PauseSequencePlayback, InVRMode),
-			FCanExecuteAction::CreateUObject(this->GetRadialMenuHandler(), &UVRRadialMenuHandler::IsActionMenuBound)
-		)
-		);
-	MenuBuilder.AddMenuEntry(
-		LOCTEXT("PlayFromStart", "Play From Start"),
-		FText(),
-		FSlateIcon(FVREditorStyle::GetStyleSetName(), "VREditorStyle.PlayFromStart"),
-		FUIAction
-		(
-			FExecuteAction::CreateStatic(&FVREditorActionCallbacks::PlayFromBeginning, InVRMode),
-			FCanExecuteAction::CreateUObject(this->GetRadialMenuHandler(), &UVRRadialMenuHandler::IsActionMenuBound)
-			)
-		);
-	MenuBuilder.AddMenuEntry(
-		LOCTEXT("ScrubTime", "Scrub Time"),
-		FText(),
-		FSlateIcon(FVREditorStyle::GetStyleSetName(), "VREditorStyle.Scrub"),
-		FUIAction
-		(
-			FExecuteAction::CreateStatic(&FVREditorActionCallbacks::ToggleSequencerScrubbing, InVRMode, UIInteractor),
-			FCanExecuteAction::CreateUObject(this->GetRadialMenuHandler(), &UVRRadialMenuHandler::IsActionMenuBound),
-			FGetActionCheckState::CreateStatic(&FVREditorActionCallbacks::GetSequencerScrubState, UIInteractor)
-			),
-		NAME_None,
-		EUserInterfaceActionType::ToggleButton
-		);
-	TAttribute<FText> LoopingLabel;
-	LoopingLabel.BindStatic(&FVREditorActionCallbacks::GetSequencerLoopingText);
-	MenuBuilder.AddMenuEntry(
-		LoopingLabel,
-		FText(),
-		FSlateIcon(FVREditorStyle::GetStyleSetName(), "VREditorStyle.ToggleLooping"),
-		FUIAction
-		(
-			FExecuteAction::CreateStatic(&FVREditorActionCallbacks::ToggleLooping, InVRMode),
-			FCanExecuteAction::CreateUObject(this->GetRadialMenuHandler(), &UVRRadialMenuHandler::IsActionMenuBound)
-		),
-		NAME_None,
-		EUserInterfaceActionType::ToggleButton
-		);
-	MenuBuilder.AddMenuEntry(
-		LOCTEXT("SetSelectionRangeStart", "Set Range Start"),
-		FText(),
-		FSlateIcon(),
-		FUIAction
-		(
-			FExecuteAction::CreateStatic(&FVREditorActionCallbacks::SetSelectionRangeStart, InVRMode),
-			FCanExecuteAction::CreateUObject(this->GetRadialMenuHandler(), &UVRRadialMenuHandler::IsActionMenuBound)
-		)
-	);
-	MenuBuilder.AddMenuEntry(
-		LOCTEXT("SetSelectionRangeEnd", "Set Range End"),
-		FText(),
-		FSlateIcon(),
-		FUIAction
-		(
-			FExecuteAction::CreateStatic(&FVREditorActionCallbacks::SetSelectionRangeEnd, InVRMode),
-			FCanExecuteAction::CreateUObject(this->GetRadialMenuHandler(), &UVRRadialMenuHandler::IsActionMenuBound)
-		)
-	);
-}
-
-void UVREditorUISystem::HandleEditorModeChanged(const FEditorModeID& EditorModeID, bool IsEnabled)
-{
-	if (IsEnabled == true &&
-		(EditorModeID == FBuiltinEditorModes::EM_Foliage ||
-		 EditorModeID == FBuiltinEditorModes::EM_Landscape ||
-		 EditorModeID == FBuiltinEditorModes::EM_MeshPaint))
-	{
-		AVREditorFloatingUI* Panel = GetPanel(ModesPanelID);
-		if (Panel != nullptr && UIInteractor != nullptr && !Panel->IsUIVisible())
-		{
-			const bool bShouldShow = true;
-			const bool bSpawnInFront = true;
-			bool bDragFromOpen = false;
-			bDragFromOpen = ShouldPreviewPanel();
-			ShowEditorUIPanel(Panel, UIInteractor, bShouldShow, bSpawnInFront, bDragFromOpen);
-		}
-	}
-}
-
 void UVREditorUISystem::ResetAll()
 {
 	RadialMenuHandler->Home();
@@ -2309,11 +1869,6 @@ void UVREditorUISystem::ResetAll()
 	}
 	FVREditorActionCallbacks::ChangeEditorModes(FBuiltinEditorModes::EM_Placement);
 
-	if( GEditor->bIsSimulatingInEditor )
-	{
-		VRMode->ToggleSIEAndVREditor();
-	}
-
 	// Turn off all snapping
 	if (FVREditorActionCallbacks::GetTranslationSnapState() == ECheckBoxState::Checked)
 	{
@@ -2409,47 +1964,6 @@ void UVREditorUISystem::UpdateSequencerUI()
 				ShowEditorUIPanel(SequencerPanel, UIInteractor, bShouldShow, bSpawnInFront, bDragFromOpen);
 				bSequencerOpenedFromRadialMenu = false;
 			}
-			else
-			{
-				// Spawn the opened sequencer just in front of the content browser
-				AVREditorFloatingUI* ContentBrowser = GetPanel(ContentBrowserPanelID);
-				if (ContentBrowser && ContentBrowser->IsUIVisible())	
-				{
-					ShowEditorUIPanel(ContentBrowser, UIInteractor, false);
-
-					const float WorldScaleFactor = VRMode->GetWorldScaleFactor();
-					const FTransform WindowToWorld = ContentBrowser->GetActorTransform();
-					const FTransform WorldToRoom = VRMode->GetRoomTransform().Inverse();
-					const FTransform WindowToRoom = WindowToWorld * WorldToRoom;
-					const FVector RoomSpaceWindowLocation = WindowToRoom.GetLocation() / WorldScaleFactor;
-					const FQuat RoomSpaceWindowRotation = WindowToRoom.GetRotation();
-
-					SequencerPanel->SetScale(ContentBrowser->GetScale());
-					SequencerPanel->SetDockedTo(AVREditorBaseActor::EDockedTo::Room);
-					SequencerPanel->SetRelativeOffset(RoomSpaceWindowLocation);
-					SequencerPanel->SetLocalRotation(RoomSpaceWindowRotation.Rotator());
-					SequencerPanel->ShowUI(true);
-
-					// Play haptic effect so user knows to look at their hand that now has UI on it!
-					UIInteractor->PlayHapticEffect(VREd::UIAssetEditorSummonedOnHandHapticFeedbackStrength->GetFloat());
-				}
-			}
-
-			FOnRadialMenuGenerated SequencerRadialMenu;
-			SequencerRadialMenu.BindUObject(this, &UVREditorUISystem::SequencerRadialMenuGenerator);
-			if (RadialMenuHandler != nullptr)
-			{
-				if (RadialMenuHandler->IsActionMenuBound())
-				{
-					ExistingActionsMenu = RadialMenuHandler->GetActionsMenuGenerator();
-					ExistingActionsMenuLabel = RadialMenuHandler->GetActionMenuLabel();
-				}
-				RadialMenuHandler->SetActionsMenuGenerator(SequencerRadialMenu, LOCTEXT("Sequencer", "Sequencer"));
-				FVREditorActionCallbacks::UpdateSequencerLoopingText(VRMode);
-				Sequencer->GetSequencerSettings()->GetOnLoopStateChanged().AddStatic(FVREditorActionCallbacks::UpdateSequencerLoopingText, VRMode);
-				FVREditorActionCallbacks::SequencerLoopText = FVREditorActionCallbacks::GetSequencerLoopingText();
-			}
-
 
 			if (InfoDisplayPanel != nullptr)
 			{
@@ -2471,16 +1985,6 @@ void UVREditorUISystem::UpdateSequencerUI()
 		}
 		else
 		{
-			if (RadialMenuHandler != nullptr)
-			{
-				RadialMenuHandler->Home();
-				RadialMenuHandler->ResetActionsMenuGenerator();
-				if (ExistingActionsMenu.IsBound())
-				{
-					RadialMenuHandler->SetActionsMenuGenerator(ExistingActionsMenu, ExistingActionsMenuLabel);
-				}
-			}
-
 			// Hide the info display when finished with sequencer.
 			if (InfoDisplayPanel != nullptr)
 			{
@@ -2608,7 +2112,7 @@ void UVREditorUISystem::UpdateDetachedActorPreviewUI(TSharedRef<SWidget> InWidge
 		{
 			PreviewWindowInfo.Add(PreviewPanelID, Actor);
 		}
-		PreviewPanel->SetSlateWidget(*this, PreviewPanelID, SNullWidget::NullWidget, ViewfinderSize, VREd::ContentBrowserUISize->GetFloat(), AVREditorFloatingUI::EDockedTo::Nothing);
+		PreviewPanel->SetSlateWidget(*this, PreviewPanelID, SNullWidget::NullWidget, ViewfinderSize, VREd::EditorUISize->GetFloat(), AVREditorFloatingUI::EDockedTo::Nothing);
 		PreviewPanel->ShowUI(true);
 	}
 
@@ -2623,7 +2127,7 @@ void UVREditorUISystem::UpdateDetachedActorPreviewUI(TSharedRef<SWidget> InWidge
 			]
 		;
 	
-		PreviewPanel->SetSlateWidget(*this, PreviewPanelID, WidgetToDraw, ViewfinderSize, VREd::ContentBrowserUISize->GetFloat(), AVREditorFloatingUI::EDockedTo::Nothing);
+		PreviewPanel->SetSlateWidget(*this, PreviewPanelID, WidgetToDraw, ViewfinderSize, VREd::EditorUISize->GetFloat(), AVREditorFloatingUI::EDockedTo::Nothing);
 		const bool bDragFromOpen = false;
 		const bool bShouldShow = true;
 		const bool bSpawnInFront = true;
diff --git a/Engine/Source/Editor/VREditor/UI/VREditorUISystem.h b/Engine/Source/Editor/VREditor/UI/VREditorUISystem.h
--- a/Engine/Source/Editor/VREditor/UI/VREditorUISystem.h
+++ b/Engine/Source/Editor/VREditor/UI/VREditorUISystem.h
@@ -311,10 +311,6 @@ protected:
 	void BuildNumPadWidget();
 	/** Swaps the content of the radial menu between the radial menu and the numpad */
 	void SwapRadialMenu();
-
-	/** Creates the sequencer radial menu to pass to the radial menu generator */
-	void SequencerRadialMenuGenerator(FMenuBuilder& MenuBuilder, TSharedPtr<FUICommandList> CommandList, class UVREditorMode* InVRMode, float& RadiusOverride);
-
 	/**
 	* Handles being notified when any editor mode changes to see if any VR Editor UI needs to change.
 	* @param Mode The mode that changed.
diff --git a/Engine/Source/Editor/VREditor/UI/VRRadialMenuHandler.cpp b/Engine/Source/Editor/VREditor/UI/VRRadialMenuHandler.cpp
--- a/Engine/Source/Editor/VREditor/UI/VRRadialMenuHandler.cpp
+++ b/Engine/Source/Editor/VREditor/UI/VRRadialMenuHandler.cpp
@@ -120,10 +120,7 @@ void UVRRadialMenuHandler::Init(UVREditorUISystem* InUISystem)
 	HomeMenu.BindUObject(this, &UVRRadialMenuHandler::HomeMenuGenerator);
 	SnapMenu.BindUObject(this, &UVRRadialMenuHandler::SnapMenuGenerator);
 	GizmoMenu.BindUObject(this, &UVRRadialMenuHandler::GizmoMenuGenerator);
-	UIMenu.BindUObject(this, &UVRRadialMenuHandler::UIMenuGenerator);
 	EditMenu.BindUObject(this, &UVRRadialMenuHandler::EditMenuGenerator);
-	ToolsMenu.BindUObject(this, &UVRRadialMenuHandler::ToolsMenuGenerator);
-	ModesMenu.BindUObject(this, &UVRRadialMenuHandler::ModesMenuGenerator);
 	SystemMenu.BindUObject(this, &UVRRadialMenuHandler::SystemMenuGenerator);
 
 	Home();
@@ -199,18 +196,6 @@ void UVRRadialMenuHandler::HomeMenuGenerator(FMenuBuilder& MenuBuilder, TSharedP
 		NAME_None,
 		EUserInterfaceActionType::CollapsedButton
 	);
-	MenuBuilder.AddMenuEntry(
-		LOCTEXT("Windows", "Windows"),
-		FText(),
-		FSlateIcon(FVREditorStyle::GetStyleSetName(), "VREditorStyle.WindowsMenu"),
-		FUIAction
-		(
-			FExecuteAction::CreateUObject(this, &UVRRadialMenuHandler::RegisterMenuGenerator, UIMenu, true),
-			FCanExecuteAction::CreateStatic(&FLevelEditorActionCallbacks::DefaultCanExecuteAction)
-		),
-		NAME_None,
-		EUserInterfaceActionType::CollapsedButton
-	);
 	MenuBuilder.AddMenuEntry(
 		LOCTEXT("Edit", "Edit"),
 		FText(),
@@ -223,45 +208,6 @@ void UVRRadialMenuHandler::HomeMenuGenerator(FMenuBuilder& MenuBuilder, TSharedP
 		NAME_None,
 		EUserInterfaceActionType::CollapsedButton
 	);
-	MenuBuilder.AddMenuEntry(
-		LOCTEXT("Tools", "Tools"),
-		FText(),
-		FSlateIcon(FVREditorStyle::GetStyleSetName(), "VREditorStyle.ToolsMenu"),
-		FUIAction
-		(
-			FExecuteAction::CreateUObject(this, &UVRRadialMenuHandler::RegisterMenuGenerator, ToolsMenu, true),
-			FCanExecuteAction::CreateStatic(&FLevelEditorActionCallbacks::DefaultCanExecuteAction)
-		),
-		NAME_None,
-		EUserInterfaceActionType::CollapsedButton
-	);
-	MenuBuilder.AddMenuEntry(
-		LOCTEXT("Modes", "Modes"),
-		FText(),
-		FSlateIcon(FVREditorStyle::GetStyleSetName(), "VREditorStyle.ModesMenu"),
-		FUIAction
-		(
-			FExecuteAction::CreateUObject(this, &UVRRadialMenuHandler::RegisterMenuGenerator, ModesMenu, true),
-			FCanExecuteAction::CreateStatic(&FLevelEditorActionCallbacks::DefaultCanExecuteAction)
-		),
-		NAME_None,
-		EUserInterfaceActionType::CollapsedButton
-		);
-	TAttribute<FText> DynamicActionsLabel;
-	DynamicActionsLabel.BindStatic(&UVRRadialMenuHandler::GetActionMenuLabel);
-	MenuBuilder.AddMenuEntry(
-		DynamicActionsLabel,
-		FText(),
-		FSlateIcon(FVREditorStyle::GetStyleSetName(), "VREditorStyle.ActionsMenu"),
-		FUIAction
-		(
-			FExecuteAction::CreateUObject(this, &UVRRadialMenuHandler::RegisterMenuGenerator, ActionsMenu, true),
-			FCanExecuteAction::CreateUObject(this, &UVRRadialMenuHandler::IsActionMenuBound)
-		),
-		NAME_None,
-		EUserInterfaceActionType::CollapsedButton
-		);
-
 	MenuBuilder.AddMenuEntry(
 		LOCTEXT("System", "System"),
 		FText(),
@@ -482,91 +428,6 @@ void UVRRadialMenuHandler::GizmoMenuGenerator(FMenuBuilder& MenuBuilder, TShared
 	MenuBuilder.EndSection();
 }
 
-void UVRRadialMenuHandler::UIMenuGenerator(FMenuBuilder& MenuBuilder, TSharedPtr<FUICommandList> CommandList, UVREditorMode* VRMode, float& RadiusOverride)
-{
-	MenuBuilder.BeginSection("Windows");
-
-	MenuBuilder.AddMenuEntry(
-		LOCTEXT("ActorDetails", "Details"),
-		LOCTEXT("ActorDetailsTooltip", "Details"),
-		FSlateIcon(FVREditorStyle::GetStyleSetName(), "VREditorStyle.Details"),
-		FUIAction
-		(
-			FExecuteAction::CreateStatic(&FVREditorActionCallbacks::OnUIToggleButtonClicked, VRMode, UVREditorUISystem::DetailsPanelID),
-			FCanExecuteAction::CreateStatic(&FLevelEditorActionCallbacks::DefaultCanExecuteAction),
-			FGetActionCheckState::CreateStatic(&FVREditorActionCallbacks::GetUIToggledState, VRMode, UVREditorUISystem::DetailsPanelID)
-		),
-		NAME_None,
-		EUserInterfaceActionType::ToggleButton
-	);
-	MenuBuilder.AddMenuEntry(
-		LOCTEXT("ContentBrowser", "Content Browser"),
-		FText(),
-		FSlateIcon(FVREditorStyle::GetStyleSetName(), "VREditorStyle.ContentBrowser"),
-		FUIAction
-		(
-			FExecuteAction::CreateStatic(&FVREditorActionCallbacks::OnUIToggleButtonClicked, VRMode, UVREditorUISystem::ContentBrowserPanelID),
-			FCanExecuteAction::CreateStatic(&FLevelEditorActionCallbacks::DefaultCanExecuteAction),
-			FGetActionCheckState::CreateStatic(&FVREditorActionCallbacks::GetUIToggledState, VRMode, UVREditorUISystem::ContentBrowserPanelID)
-		),
-		NAME_None,
-		EUserInterfaceActionType::ToggleButton
-	);
-	MenuBuilder.AddMenuEntry(
-		LOCTEXT("ModesPanel", "Modes Panel"),
-		FText(),
-		FSlateIcon(FVREditorStyle::GetStyleSetName(), "VREditorStyle.ModesPanel"),
-		FUIAction
-		(
-			FExecuteAction::CreateStatic(&FVREditorActionCallbacks::OnUIToggleButtonClicked, VRMode, UVREditorUISystem::ModesPanelID),
-			FCanExecuteAction::CreateStatic(&FLevelEditorActionCallbacks::DefaultCanExecuteAction),
-			FGetActionCheckState::CreateStatic(&FVREditorActionCallbacks::GetUIToggledState, VRMode, UVREditorUISystem::ModesPanelID)
-		),
-		NAME_None,
-		EUserInterfaceActionType::ToggleButton
-	);
-	MenuBuilder.AddMenuEntry(
-		LOCTEXT("WorldOutliner", "World Outliner"),
-		FText(),
-		FSlateIcon(FVREditorStyle::GetStyleSetName(), "VREditorStyle.WorldOutliner"),
-		FUIAction
-		(
-			FExecuteAction::CreateStatic(&FVREditorActionCallbacks::OnUIToggleButtonClicked, VRMode, UVREditorUISystem::WorldOutlinerPanelID),
-			FCanExecuteAction::CreateStatic(&FLevelEditorActionCallbacks::DefaultCanExecuteAction),
-			FGetActionCheckState::CreateStatic(&FVREditorActionCallbacks::GetUIToggledState, VRMode, UVREditorUISystem::WorldOutlinerPanelID)
-		),
-		NAME_None,
-		EUserInterfaceActionType::ToggleButton
-	);
-	MenuBuilder.AddMenuEntry(
-		LOCTEXT("WorldSettings", "World Settings"),
-		FText(),
-		FSlateIcon(FVREditorStyle::GetStyleSetName(), "VREditorStyle.WorldSettings"),
-		FUIAction
-		(
-			FExecuteAction::CreateStatic(&FVREditorActionCallbacks::OnUIToggleButtonClicked, VRMode, UVREditorUISystem::WorldSettingsPanelID),
-			FCanExecuteAction::CreateStatic(&FLevelEditorActionCallbacks::DefaultCanExecuteAction),
-			FGetActionCheckState::CreateStatic(&FVREditorActionCallbacks::GetUIToggledState, VRMode, UVREditorUISystem::WorldSettingsPanelID)
-		),
-		NAME_None,
-		EUserInterfaceActionType::ToggleButton
-	);
-	MenuBuilder.AddMenuEntry(
-		LOCTEXT("CreateNewSequence", "Create Sequence"),
-		FText(),
-		FSlateIcon(FVREditorStyle::GetStyleSetName(), "VREditorStyle.Sequencer"),
-		FUIAction
-		(
-			FExecuteAction::CreateStatic(&FVREditorActionCallbacks::CreateNewSequence, VRMode),
-			FCanExecuteAction::CreateStatic(&FLevelEditorActionCallbacks::DefaultCanExecuteAction)
-		),
-		NAME_None,
-		EUserInterfaceActionType::ToggleButton
-	);
-
-	MenuBuilder.EndSection();
-}
-
 void UVRRadialMenuHandler::EditMenuGenerator(FMenuBuilder& MenuBuilder, TSharedPtr<FUICommandList> CommandList, UVREditorMode* VRMode, float& RadiusOverride)
 {
 	MenuBuilder.BeginSection("Edit");
@@ -646,170 +507,6 @@ void UVRRadialMenuHandler::EditMenuGenerator(FMenuBuilder& MenuBuilder, TSharedP
 
 	MenuBuilder.EndSection();
 }
-
-void UVRRadialMenuHandler::ToolsMenuGenerator(FMenuBuilder& MenuBuilder, TSharedPtr<FUICommandList> CommandList, UVREditorMode* VRMode, float& RadiusOverride)
-{
-	MenuBuilder.BeginSection("Tools");
-
-	TAttribute<FText> DynamicSimulateLabel;
-	DynamicSimulateLabel.BindStatic(&FVREditorActionCallbacks::GetSimulateText);
-
-	MenuBuilder.AddMenuEntry(
-		DynamicSimulateLabel,
-		FText(),
-		FSlateIcon(FVREditorStyle::GetStyleSetName(), "VREditorStyle.Simulate"),
-		FUIAction
-		(
-			FExecuteAction::CreateStatic(&FVREditorActionCallbacks::OnSimulateButtonClicked, VRMode),
-			FCanExecuteAction::CreateStatic(&FLevelEditorActionCallbacks::DefaultCanExecuteAction)
-		),
-		NAME_None,
-		EUserInterfaceActionType::CollapsedButton
-	);
-
-
-	MenuBuilder.AddMenuEntry(
-		LOCTEXT("SaveActors", "Save Actors"),
-		FText(),
-		FSlateIcon(FVREditorStyle::GetStyleSetName(), "VREditorStyle.SaveSimulation"),
-		FUIAction
-		(
-			FExecuteAction::CreateStatic(&FLevelEditorActionCallbacks::OnKeepSimulationChanges),
-			FCanExecuteAction::CreateStatic(&FLevelEditorActionCallbacks::CanExecuteKeepSimulationChanges)
-		),
-		NAME_None,
-		EUserInterfaceActionType::CollapsedButton
-	);
-
-
-	MenuBuilder.AddMenuEntry(
-		LOCTEXT("PauseSimulation", "Pause"),
-		FText(),
-		FSlateIcon(FVREditorStyle::GetStyleSetName(), "VREditorStyle.Pause"),
-		FUIAction
-		(
-			FExecuteAction::CreateStatic(&FPlayWorldCommandCallbacks::PausePlaySession_Clicked),
-			FCanExecuteAction::CreateStatic(&FPlayWorldCommandCallbacks::HasPlayWorldAndRunning)
-		),
-		NAME_None,
-		EUserInterfaceActionType::CollapsedButton
-	);
-
-	MenuBuilder.AddMenuEntry(
-		LOCTEXT("ResumeSimulation", "Resume"),
-		FText(),
-		FSlateIcon(FVREditorStyle::GetStyleSetName(), "VREditorStyle.Resume"),
-		FUIAction
-		(
-			FExecuteAction::CreateStatic(&FPlayWorldCommandCallbacks::ResumePlaySession_Clicked),
-			FCanExecuteAction::CreateStatic(&FPlayWorldCommandCallbacks::HasPlayWorldAndPaused)
-		),
-		NAME_None,
-		EUserInterfaceActionType::CollapsedButton
-	);
-
-	MenuBuilder.AddMenuEntry(
-		LOCTEXT("PlayInEditor", "Play"),
-		FText(),
-		FSlateIcon(FVREditorStyle::GetStyleSetName(), "VREditorStyle.Play"),
-		FUIAction
-		(
-			FExecuteAction::CreateStatic(&FVREditorActionCallbacks::OnPlayButtonClicked, VRMode),
-			FCanExecuteAction::CreateStatic(&FVREditorActionCallbacks::CanPlay, VRMode)
-		),
-		NAME_None,
-		EUserInterfaceActionType::CollapsedButton
-	);
-
-	MenuBuilder.AddMenuEntry(
-		LOCTEXT("Screenshot", "Screenshot"),
-		FText(),
-		FSlateIcon(FVREditorStyle::GetStyleSetName(), "VREditorStyle.Screenshot"),
-		FUIAction
-		(
-			FExecuteAction::CreateStatic(&FVREditorActionCallbacks::OnScreenshotButtonClicked, VRMode),
-			FCanExecuteAction::CreateStatic(&FLevelEditorActionCallbacks::DefaultCanExecuteAction)
-		),
-		NAME_None,
-		EUserInterfaceActionType::CollapsedButton
-	);
-
-	MenuBuilder.AddMenuEntry(
-		LOCTEXT("Flashlight", "Flashlight"),
-		FText(),
-		FSlateIcon(FVREditorStyle::GetStyleSetName(), "VREditorStyle.Flashlight"),
-		FUIAction
-		(
-			FExecuteAction::CreateStatic(&FVREditorActionCallbacks::OnLightButtonClicked, VRMode),
-			FCanExecuteAction::CreateStatic(&FLevelEditorActionCallbacks::DefaultCanExecuteAction)
-		),
-		NAME_None,
-		EUserInterfaceActionType::CollapsedButton
-	);
-
-	MenuBuilder.EndSection();
-}
-
-void UVRRadialMenuHandler::ModesMenuGenerator(FMenuBuilder& MenuBuilder, TSharedPtr<FUICommandList> CommandList, UVREditorMode* VRMode, float& RadiusOverride)
-{
-	MenuBuilder.BeginSection("Modes");
-
-	MenuBuilder.AddMenuEntry(
-		LOCTEXT("Actors", "Actors"),
-		FText(),
-		FSlateIcon(FVREditorStyle::GetStyleSetName(), "VREditorStyle.ActorsMode"),
-		FUIAction
-		(
-			FExecuteAction::CreateStatic(&FVREditorActionCallbacks::ChangeEditorModes, FBuiltinEditorModes::EM_Placement),
-			FCanExecuteAction::CreateStatic(&FLevelEditorActionCallbacks::DefaultCanExecuteAction),
-			FGetActionCheckState::CreateStatic(&FVREditorActionCallbacks::EditorModeActive, FBuiltinEditorModes::EM_Placement)
-		),
-		NAME_None,
-		EUserInterfaceActionType::ToggleButton
-	);
-	MenuBuilder.AddMenuEntry(
-		LOCTEXT("Foliage", "Foliage"),
-		FText(),
-		FSlateIcon(FVREditorStyle::GetStyleSetName(), "VREditorStyle.FoliageMode"),
-		FUIAction
-		(
-			FExecuteAction::CreateStatic(&FVREditorActionCallbacks::ChangeEditorModes, FBuiltinEditorModes::EM_Foliage),
-			FCanExecuteAction::CreateStatic(&FLevelEditorActionCallbacks::DefaultCanExecuteAction),
-			FGetActionCheckState::CreateStatic(&FVREditorActionCallbacks::EditorModeActive, FBuiltinEditorModes::EM_Foliage)
-		),
-		NAME_None,
-		EUserInterfaceActionType::ToggleButton
-	);
-	MenuBuilder.AddMenuEntry(
-		LOCTEXT("Landscape", "Landscape"),
-		FText(),
-		FSlateIcon(FVREditorStyle::GetStyleSetName(), "VREditorStyle.LandscapeMode"),
-		FUIAction
-		(
-			FExecuteAction::CreateStatic(&FVREditorActionCallbacks::ChangeEditorModes, FBuiltinEditorModes::EM_Landscape),
-			FCanExecuteAction::CreateStatic(&FLevelEditorActionCallbacks::DefaultCanExecuteAction),
-			FGetActionCheckState::CreateStatic(&FVREditorActionCallbacks::EditorModeActive, FBuiltinEditorModes::EM_Landscape)
-		),
-		NAME_None,
-		EUserInterfaceActionType::ToggleButton
-	);
-	MenuBuilder.AddMenuEntry(
-		LOCTEXT("MeshPaint", "Paint"),
-		FText(),
-		FSlateIcon(FVREditorStyle::GetStyleSetName(), "VREditorStyle.MeshPaintMode"),
-		FUIAction
-		(
-			FExecuteAction::CreateStatic(&FVREditorActionCallbacks::ChangeEditorModes, FBuiltinEditorModes::EM_MeshPaint),
-			FCanExecuteAction::CreateStatic(&FLevelEditorActionCallbacks::DefaultCanExecuteAction),
-			FGetActionCheckState::CreateStatic(&FVREditorActionCallbacks::EditorModeActive, FBuiltinEditorModes::EM_MeshPaint)
-		),
-		NAME_None,
-		EUserInterfaceActionType::ToggleButton
-	);
-
-	MenuBuilder.EndSection();
-}
-
 void UVRRadialMenuHandler::SystemMenuGenerator(FMenuBuilder& MenuBuilder, TSharedPtr<FUICommandList> CommandList, UVREditorMode* VRMode, float& RadiusOverride)
 {
 	MenuBuilder.BeginSection("System");
@@ -859,28 +556,5 @@ void UVRRadialMenuHandler::RegisterMenuGenerator(const FOnRadialMenuGenerated Ne
 	}
 }
 
-void UVRRadialMenuHandler::SetActionsMenuGenerator(const FOnRadialMenuGenerated NewMenuGenerator, const FText NewLabel)
-{
-	UVRRadialMenuHandler::ActionMenuLabel = NewLabel;
-	ActionsMenu = NewMenuGenerator;
-}
-
-
-void UVRRadialMenuHandler::ResetActionsMenuGenerator()
-{
-	UVRRadialMenuHandler::ActionMenuLabel = LOCTEXT("DefaultActions", "Actions");
-	FOnRadialMenuGenerated EmptyRadialMenu;
-	ActionsMenu = EmptyRadialMenu;
-}
-
-FText UVRRadialMenuHandler::GetActionMenuLabel()
-{
-	return  UVRRadialMenuHandler::ActionMenuLabel;
-}
-
-bool UVRRadialMenuHandler::IsActionMenuBound()
-{
-	return ActionsMenu.IsBound();
-}
 
 #undef LOCTEXT_NAMESPACE
diff --git a/Engine/Source/Editor/VREditor/UI/VRRadialMenuHandler.h b/Engine/Source/Editor/VREditor/UI/VRRadialMenuHandler.h
--- a/Engine/Source/Editor/VREditor/UI/VRRadialMenuHandler.h
+++ b/Engine/Source/Editor/VREditor/UI/VRRadialMenuHandler.h
@@ -33,12 +33,6 @@ public:
 	/** Registers a new menu generator and replaces the currently displayed menu if the radial menu is open */
 	void RegisterMenuGenerator(const FOnRadialMenuGenerated NewMenuGenerator, const bool bShouldAddToStack = true);
 
-	/** Sets a delegate for the context-specific actions menu */
-	void SetActionsMenuGenerator(const FOnRadialMenuGenerated NewMenuGenerator, const FText NewLabel);
-
-	/** Resets the delegate and button for the context-specific actions menu */
-	void ResetActionsMenuGenerator();
-
 	FOnRadialMenuGenerated GetCurrentMenuGenerator()
 	{
 		return OnRadialMenuGenerated;
@@ -49,23 +43,12 @@ public:
 		return HomeMenu;
 	};
 
-	FOnRadialMenuGenerated GetActionsMenuGenerator()
-	{
-		return ActionsMenu;
-	};
-
 	/** Returns to the previous radial menu */
 	void BackOutMenu();
 
 	/** Returns to the home menu */
 	void Home();
 
-	/** Allows disabling buttons in the action menu if it's not currently bound */
-	bool IsActionMenuBound();
-
-	/** Allows other systems to read and save the title of existing action menus */
-	static FText GetActionMenuLabel();
-
 	void Init(class UVREditorUISystem* InUISystem);
 
 protected:
@@ -76,14 +59,8 @@ protected:
 
 	void GizmoMenuGenerator(FMenuBuilder& MenuBuilder, TSharedPtr<FUICommandList> CommandList, UVREditorMode* VRMode, float& RadiusOverride);
 
-	void UIMenuGenerator(FMenuBuilder& MenuBuilder, TSharedPtr<FUICommandList> CommandList, UVREditorMode* VRMode, float& RadiusOverride);
-
 	void EditMenuGenerator(FMenuBuilder& MenuBuilder, TSharedPtr<FUICommandList> CommandList, UVREditorMode* VRMode, float& RadiusOverride);
 
-	void ToolsMenuGenerator(FMenuBuilder& MenuBuilder, TSharedPtr<FUICommandList> CommandList, UVREditorMode* VRMode, float& RadiusOverride);
-
-	void ModesMenuGenerator(FMenuBuilder& MenuBuilder, TSharedPtr<FUICommandList> CommandList, UVREditorMode* VRMode, float& RadiusOverride);
-
 	void SystemMenuGenerator(FMenuBuilder& MenuBuilder, TSharedPtr<FUICommandList> CommandList, UVREditorMode* VRMode, float& RadiusOverride);
 
 protected:
@@ -96,16 +73,8 @@ protected:
 
 	FOnRadialMenuGenerated GizmoMenu;
 
-	FOnRadialMenuGenerated UIMenu;
-
 	FOnRadialMenuGenerated EditMenu;
 
-	FOnRadialMenuGenerated ToolsMenu;
-
-	FOnRadialMenuGenerated ModesMenu;
-
-	FOnRadialMenuGenerated ActionsMenu;
-
 	FOnRadialMenuGenerated SystemMenu;
 
 	static FText ActionMenuLabel;
diff --git a/Engine/Source/Editor/VREditor/VREditorActions.cpp b/Engine/Source/Editor/VREditor/VREditorActions.cpp
--- a/Engine/Source/Editor/VREditor/VREditorActions.cpp
+++ b/Engine/Source/Editor/VREditor/VREditorActions.cpp
@@ -233,95 +233,6 @@ ECheckBoxState FVREditorActionCallbacks::GetUIToggledState(UVREditorMode* InVRMo
 	return InVRMode->GetUISystem().IsShowingEditorUIPanel(PanelToCheck) ? ECheckBoxState::Checked : ECheckBoxState::Unchecked;
 }
 
-
-void FVREditorActionCallbacks::OnLightButtonClicked(UVREditorMode* InVRMode)
-{
-	// Always spawn the flashlight on the hand clicking on the UI
-	UVREditorInteractor* VREditorInteractor = InVRMode->GetHandInteractor(EControllerHand::Left);
-	if (!VREditorInteractor->IsHoveringOverUI())
-	{
-		VREditorInteractor = InVRMode->GetHandInteractor(EControllerHand::Right);
-
-	}
-	InVRMode->ToggleFlashlight(VREditorInteractor);
-
-}
-
-ECheckBoxState FVREditorActionCallbacks::GetFlashlightState(UVREditorMode* InVRMode)
-{
-	return InVRMode->IsFlashlightOn() ? ECheckBoxState::Checked : ECheckBoxState::Unchecked;
-}
-
-
-void FVREditorActionCallbacks::OnScreenshotButtonClicked(UVREditorMode* InVRMode)
-{
-	// @todo vreditor: update after direct buffer grab changes
-
-	FString GeneratedFilename = TEXT("");
-	FString Filename = TEXT("");
-	FScreenshotRequest::CreateViewportScreenShotFilename(GeneratedFilename);
-	const bool bRemovePath = false;
-	GeneratedFilename = FPaths::GetBaseFilename(GeneratedFilename, bRemovePath);
-	FFileHelper::GenerateNextBitmapFilename(GeneratedFilename, TEXT("png"), Filename);
-
-	TSharedRef<SWidget> WindowRef = InVRMode->GetLevelViewportPossessedForVR().AsWidget();
-
-
-	TArray<FColor> OutImageData;
-	FIntVector OutImageSize;
-
-	if (FSlateApplication::Get().TakeScreenshot(WindowRef, OutImageData, OutImageSize))
-	{
-		int32 Width = OutImageSize.X;
-		int32 Height = OutImageSize.Y;
-
-		TArray<FColor> ScaledBitmap;
-
-
-		ScaledBitmap = OutImageData;
-		//Clear the alpha channel before saving
-		for (int32 Index = 0; Index < Width*Height; Index++)
-		{
-			ScaledBitmap[Index].A = 255;
-		}
-
-
-		TArray<uint8> CompressedBitmap;
-		FImageUtils::CompressImageArray(Width, Height, ScaledBitmap, CompressedBitmap);
-
-
-		//Save locally
-		const bool bTree = true;
-		const FString FinalFileName = Filename;
-		IFileManager::Get().MakeDirectory(*FPaths::GetPath(FinalFileName), bTree);
-		FFileHelper::SaveArrayToFile(CompressedBitmap, *FinalFileName);
-
-	}
-}
-
-void FVREditorActionCallbacks::OnPlayButtonClicked(UVREditorMode* InVRMode)
-{
-	InVRMode->TogglePIEAndVREditor();
-}
-
-bool FVREditorActionCallbacks::CanPlay(UVREditorMode* InVRMode)
-{
-	static const FName OculusSystemName(TEXT("OculusHMD"));
-	return FLevelEditorActionCallbacks::DefaultCanExecuteAction() 
-		&& VREd::AllowPlay->GetInt() == 1 
-		&& (GEditor != nullptr && !GEditor->bIsSimulatingInEditor);
-}
-
-void FVREditorActionCallbacks::OnSimulateButtonClicked(UVREditorMode* InVRMode)
-{
-	InVRMode->ToggleSIEAndVREditor();
-}
-
-FText FVREditorActionCallbacks::GetSimulateText()
-{
-	return GEditor->bIsSimulatingInEditor ? LOCTEXT("SimulateStopButton", "Stop") : LOCTEXT("SimulateStartButton", "Simulate");
-}
-
 void FVREditorActionCallbacks::OnSnapActorsToGroundClicked(UVREditorMode* InVRMode)
 {
 	InVRMode->SnapSelectedActorsToGround();
diff --git a/Engine/Source/Editor/VREditor/VREditorActions.h b/Engine/Source/Editor/VREditor/VREditorActions.h
--- a/Engine/Source/Editor/VREditor/VREditorActions.h
+++ b/Engine/Source/Editor/VREditor/VREditorActions.h
@@ -99,44 +99,6 @@ public:
 	 */
 	static ECheckBoxState GetUIToggledState(UVREditorMode* InVRMode, VREditorPanelID PanelToCheck);
 
-	/**
-	 * Toggles a flashlight on and off on the interactor that clicked on the UI button
-	 * @param InVRMode Currently active VRMode
-	 */
-	static void OnLightButtonClicked(UVREditorMode* InVRMode);
-
-	/**
-	 * Returns whether or not the flashlight is enabled - used for check box on the flashlight button
-	 * @param InVRMode Currently active VRMode
-	 */
-	static ECheckBoxState GetFlashlightState(UVREditorMode* InVRMode);
-
-	/**
-	 * Takes a screenshot of the mirror viewport
-	 * @param InVRMode Currently active VRMode
-	 */
-	static void OnScreenshotButtonClicked(UVREditorMode* InVRMode);
-
-	/**
-	 * Enters Play In Editor mode for testing of gameplay
-	 * @param InVRMode Currently active VRMode
-	 */
-	static void OnPlayButtonClicked(UVREditorMode* InVRMode);
-
-	/**
-	 * If we are allowed to enter play.
-	 */
-	static bool CanPlay(UVREditorMode* InVRMode);
-
-	/**
-	 * Enters Simulate mode for physics and animation playback
-	 * @param InVRMode Currently active VRMode
-	 */
-	static void OnSimulateButtonClicked(UVREditorMode* InVRMode);
-
-	/** Returns the simulate button text to display */
-	static FText GetSimulateText();
-
 	/**
 	 * Snaps currently selected Actors to the ground
 	 * @param InVRMode Currently active VRMode
diff --git a/Engine/Source/Editor/VREditor/VREditorInteractor.cpp b/Engine/Source/Editor/VREditor/VREditorInteractor.cpp
--- a/Engine/Source/Editor/VREditor/VREditorInteractor.cpp
+++ b/Engine/Source/Editor/VREditor/VREditorInteractor.cpp
@@ -884,39 +884,6 @@ void UVREditorInteractor::PreviewInputKey( class FEditorViewportClient& Viewport
 		bIsPressingTrackpad = Event == IE_Released ? false : true;
 	}
 
-	if (GetControllerType() == EControllerType::Laser)
-	{
-		// Are we holding "up" on the trackpad?
-		const bool bIsHoldingUpOnTrackpad =
-			bIsTrackpadPositionValid[0] && bIsTrackpadPositionValid[1] &&
-			TrackpadPosition.Y >= VREd::MinTrackpadOffsetBeforeRadialMenu->GetFloat() &&
-			(GetHMDDeviceType() == OculusDeviceType || bIsPressingTrackpad);
-
-		if (bIsHoldingUpOnTrackpad && Action.ActionType == ViewportWorldActionTypes::SelectAndMove && Event == IE_Pressed)
-		{
-			bOutWasHandled = true;
-
-			// Try to place the object currently selected
-			TArray<UObject*> SelectedObjects;
-			{
-				FEditorDelegates::LoadSelectedAssetsIfNeeded.Broadcast();
-				GEditor->GetSelectedObjects()->GetSelectedObjects( /* Out */ SelectedObjects );
-			}
-
-			if (SelectedObjects.Num() > 0)
-			{
-				TArray<UObject*> ObjectToPlace;
-				ObjectToPlace.Add( SelectedObjects[0] );
-
-				Action.bIsInputCaptured = true;
-
-				const bool bShouldInterpolateFromDragLocation = false;
-				UActorFactory* FactoryToUse = nullptr;	// Use default factory
-				GetVRMode().GetPlacementSystem()->StartPlacingObjects( ObjectToPlace, FactoryToUse, this, bShouldInterpolateFromDragLocation );
-			}
-		}
-	}
-
 	// Update modifier state
 	if (Action.ActionType == VRActionTypes::Modifier)
 	{
diff --git a/Engine/Source/Editor/VREditor/VREditorMode.cpp b/Engine/Source/Editor/VREditor/VREditorMode.cpp
--- a/Engine/Source/Editor/VREditor/VREditorMode.cpp
+++ b/Engine/Source/Editor/VREditor/VREditorMode.cpp
@@ -97,9 +97,6 @@ UVREditorMode::UVREditorMode() :
 	AutoScalerSystem( nullptr ),
 	WorldInteraction( nullptr ),
 	bFirstTick( true ),
-	SavedWorldToMetersScaleForPIE( 100.f ),
-	bStartedPlayFromVREditor( false ),
-	bStartedPlayFromVREditorSimulate( false ),
 	AssetContainer( nullptr )
 { 
 }
@@ -179,7 +176,6 @@ void UVREditorMode::Shutdown()
 	// @todo vreditor urgent: Disable global editor hacks for VR Editor mode
 	GEnableVREditorHacks = false;
 
-	FEditorDelegates::EndPIE.RemoveAll(this);
 }
 
 void UVREditorMode::AllocateInteractors()
@@ -233,12 +229,6 @@ void UVREditorMode::Enter()
 		WorldInteraction->OnPostWorldInteractionTick().AddUObject( this, &UVREditorMode::PostTick );
 	}
 
-	FEditorDelegates::PostPIEStarted.AddUObject( this, &UVREditorMode::PostPIEStarted );
-	FEditorDelegates::PrePIEEnded.AddUObject( this, &UVREditorMode::PrePIEEnded );
-	FEditorDelegates::EndPIE.AddUObject(this, &UVREditorMode::OnEndPIE);
-	FEditorDelegates::OnPreSwitchBeginPIEAndSIE.AddUObject(this, &UVREditorMode::OnPreSwitchPIEAndSIE);
-	FEditorDelegates::OnSwitchBeginPIEAndSIE.AddUObject(this, &UVREditorMode::OnSwitchPIEAndSIE);
-
 
 	// @todo vreditor: We need to make sure the user can never switch to orthographic mode, or activate settings that
 	// would disrupt the user's ability to view the VR scene.
@@ -365,8 +355,6 @@ void UVREditorMode::Enter()
 
 	bFirstTick = true;
 	SetActive(true);
-	bStartedPlayFromVREditor = false;
-	bStartedPlayFromVREditorSimulate = false;
 }
 
 void UVREditorMode::Exit(const bool bShouldDisableStereo)
@@ -474,11 +462,6 @@ void UVREditorMode::Exit(const bool bShouldDisableStereo)
 
 	AssetContainer = nullptr;
 
-	FEditorDelegates::PostPIEStarted.RemoveAll( this );
-	FEditorDelegates::PrePIEEnded.RemoveAll( this );
-	FEditorDelegates::EndPIE.RemoveAll( this );
-	FEditorDelegates::OnPreSwitchBeginPIEAndSIE.RemoveAll(this);
-	FEditorDelegates::OnSwitchBeginPIEAndSIE.RemoveAll(this);
 
 	GEditor->OnEditorClose().RemoveAll( this );
 
@@ -895,79 +878,6 @@ void UVREditorMode::SaveSequencerSettings(bool bInKeyAllEnabled, EAutoChangeMode
 	SavedEditorState.AutoChangeMode = InAutoChangeMode;
 }
 
-void UVREditorMode::ToggleSIEAndVREditor()
-{
-	if (GEditor->EditorWorld == nullptr && !GEditor->bIsSimulatingInEditor)
-	{
-		FRequestPlaySessionParams SessionParams;
-		SessionParams.DestinationSlateViewport = VREditorLevelViewportWeakPtr;
-		SessionParams.WorldType = EPlaySessionWorldType::SimulateInEditor;
-
-		GEditor->RequestPlaySession(SessionParams);
-	}
-	else if (GEditor->PlayWorld != nullptr && GEditor->bIsSimulatingInEditor)
-	{
-		GEditor->RequestEndPlayMap();
-	}
-}
-
-void UVREditorMode::TogglePIEAndVREditor()
-{
-	bool bRequestedPIE = false;
-	if (GEditor->EditorWorld == nullptr && GEditor->PlayWorld == nullptr && !GEditor->bIsSimulatingInEditor)
-	{
-		FRequestPlaySessionParams SessionParams;
-		SessionParams.DestinationSlateViewport = VREditorLevelViewportWeakPtr;
-		SessionParams.WorldType = EPlaySessionWorldType::PlayInEditor;
-		
-		const bool bHMDIsReady = (GEngine && GEngine->XRSystem.IsValid() && GEngine->XRSystem->GetHMDDevice() && GEngine->XRSystem->GetHMDDevice()->IsHMDConnected());
-		if (bHMDIsReady)
-		{
-			SessionParams.SessionPreviewTypeOverride = EPlaySessionPreviewType::VRPreview;
-		}
-		GEditor->RequestPlaySession(SessionParams);
-		bRequestedPIE = true;
-	}
-	else if (GEditor->PlayWorld != nullptr)
-	{
-		// Since we are already in simulate, we want to toggle to PIE.
-		if (GEditor->bIsSimulatingInEditor)
-		{
-			bStartedPlayFromVREditorSimulate = true;
-			bRequestedPIE = true;
-
-			GEditor->RequestToggleBetweenPIEandSIE();
-		}
-		else
-		{
-			// If this play started while in simulate, then toggle back to simulate.
-			if (bStartedPlayFromVREditorSimulate)
-			{
-				GEditor->RequestToggleBetweenPIEandSIE();
-			}
-			else
-			{
-				GEditor->RequestEndPlayMap();
-			}
-		}
-	}
-
-	if (bRequestedPIE)
-	{
-		// Turn off input processing while in PIE.  We don't want any input events until the user comes back to the editor
-		WorldInteraction->SetUseInputPreprocessor(false);
-
-		SavedWorldToMetersScaleForPIE = GetWorld()->GetWorldSettings()->WorldToMeters;
-
-		// Restore the world to meters before entering play
-		RestoreWorldToMeters();
-
-		SetActive(false);
-		WorldInteraction->SetActive(false);
-		bStartedPlayFromVREditor = true;
-	}
-}
-
 void UVREditorMode::TransitionWorld(UWorld* NewWorld, EEditorWorldExtensionTransitionState TransitionState)
 {
 	Super::TransitionWorld(NewWorld, TransitionState);
@@ -1093,7 +1003,6 @@ void UVREditorMode::StartViewport(TSharedPtr<SLevelViewport> Viewport)
 			const float DefaultWorldToMeters = VREd::DefaultWorldToMeters->GetFloat();
 			const float SavedWorldToMeters = DefaultWorldToMeters != 0.0f ? DefaultWorldToMeters : VRViewportClient.GetWorld()->GetWorldSettings()->WorldToMeters;
 			SavedEditorState.WorldToMetersScale = SavedWorldToMeters;
-			SavedWorldToMetersScaleForPIE = SavedWorldToMeters;
 		}
 
 		if (bActuallyUsingVR)
@@ -1192,25 +1101,6 @@ void UVREditorMode::CloseViewport( const bool bShouldDisableStereo )
 	}
 }
 
-void UVREditorMode::RestoreFromPIE()
-{
-	SetActive(true);
-	bStartedPlayFromVREditorSimulate = false;
-
-	GetWorld()->GetWorldSettings()->WorldToMeters = SavedWorldToMetersScaleForPIE;
-	WorldInteraction->SetWorldToMetersScale(SavedWorldToMetersScaleForPIE);
-
-	// Re-enable input pre-processing
-	WorldInteraction->SetUseInputPreprocessor(true);
-	WorldInteraction->SetActive(true);
-
-	UVREditorInteractor* UIInteractor = UISystem->GetUIInteractor();
-	if (UIInteractor != nullptr)
-	{
-		UIInteractor->ResetTrackpad();
-		UISystem->HideRadialMenu(false, false);
-	}
-}
 
 void UVREditorMode::RestoreWorldToMeters()
 {
@@ -1268,26 +1158,6 @@ UStaticMeshComponent* UVREditorMode::CreateMesh(AActor* OwningActor, UStaticMesh
 	return CreatedMeshComponent;
 }
 
-void UVREditorMode::SetActionsMenuGenerator(const FOnRadialMenuGenerated NewMenuGenerator, const FText NewLabel)
-{
-	GetUISystem().GetRadialMenuHandler()->SetActionsMenuGenerator(NewMenuGenerator, NewLabel);
-}
-
-void UVREditorMode::ResetActionsMenuGenerator()
-{
-	GetUISystem().GetRadialMenuHandler()->ResetActionsMenuGenerator();
-}
-
-void UVREditorMode::RefreshRadialMenuActionsSubmenu()
-{
-	GetUISystem().GetRadialMenuHandler()->RegisterMenuGenerator( GetUISystem().GetRadialMenuHandler()->GetActionsMenuGenerator() );
-}
-
-bool UVREditorMode::GetStartedPlayFromVREditor() const
-{
-	return bStartedPlayFromVREditor;
-}
-
 const UVREditorAssetContainer& UVREditorMode::GetAssetContainer() const
 {
 	return *AssetContainer;
@@ -1314,78 +1184,6 @@ bool UVREditorMode::IsAimingTeleport() const
 	return TeleportActor->IsAiming();
 }
 
-void UVREditorMode::PostPIEStarted( bool bIsSimulatingInEditor )
-{
-	if (!bIsSimulatingInEditor)
-	{
-		GEnableVREditorHacks = false;
-	}
-}
-
-
-void UVREditorMode::PrePIEEnded( bool bWasSimulatingInEditor )
-{
-	if (!bWasSimulatingInEditor && !bStartedPlayFromVREditorSimulate)
-	{
-		GEnableVREditorHacks = true;
-	}
-	else if (bStartedPlayFromVREditorSimulate)
-	{
-		// Pre PIE to SIE. When exiting play with escape, the delegate toggle PIE and SIE won't be called. We know that we started PIE from simulate. However simulate will also be closed.
-		GEnableVREditorHacks = true;
-	}
-}
-
-void UVREditorMode::OnEndPIE(bool bWasSimulatingInEditor)
-{
-	if (!bWasSimulatingInEditor && !bStartedPlayFromVREditorSimulate)
-	{
-		RestoreFromPIE();
-	}
-	else if (bStartedPlayFromVREditorSimulate)
-	{
-		// Post PIE to SIE
-		RestoreFromPIE();
-		GetOwningCollection()->ShowAllActors(true);
-	}
-}
-
-void UVREditorMode::OnPreSwitchPIEAndSIE(bool bIsSimulatingInEditor)
-{
-	if (bStartedPlayFromVREditorSimulate)
-	{
-		if (bIsSimulatingInEditor)
-		{
-			// Pre SIE to PIE
-			GetOwningCollection()->ShowAllActors(false);
-		}
-		else
-		{
-			// Pre PIE to SIE
-			GEnableVREditorHacks = true;
-		}
-	}
-}
-
-void UVREditorMode::OnSwitchPIEAndSIE(bool bIsSimulatingInEditor)
-{
-	if (bStartedPlayFromVREditorSimulate)
-	{
-		if (bIsSimulatingInEditor)
-		{
-			// Post PIE to SIE
-			RestoreFromPIE();
-			GetOwningCollection()->ShowAllActors(true);
-		}
-		else
-		{
-			// Post SIE to PIE
-			GEnableVREditorHacks = false;
-			FSlateApplication::Get().SetAllUserFocusToGameViewport();
-		}
-	}
-}
-
 void UVREditorMode::ToggleDebugMode()
 {
 	UVREditorMode::bDebugModeEnabled = !UVREditorMode::bDebugModeEnabled;
diff --git a/Engine/Source/Editor/VREditor/VREditorModeManager.cpp b/Engine/Source/Editor/VREditor/VREditorModeManager.cpp
--- a/Engine/Source/Editor/VREditor/VREditorModeManager.cpp
+++ b/Engine/Source/Editor/VREditor/VREditorModeManager.cpp
@@ -45,7 +45,6 @@ void FVREditorModeManager::Tick( const float DeltaTime )
 	IHeadMountedDisplay * const HMD = GEngine != nullptr && GEngine->XRSystem.IsValid() ? GEngine->XRSystem->GetHMDDevice() : nullptr;
 	if (GetDefault<UVRModeSettings>()->bEnableAutoVREditMode
 		&& HMD
-		&& (GEditor->PlayWorld == nullptr || (CurrentVREditorMode != nullptr && CurrentVREditorMode->GetStartedPlayFromVREditor()))
 		&& FPlatformApplicationMisc::IsThisApplicationForeground())
 	{
 		const EHMDWornState::Type LatestHMDWornState = HMD->GetHMDWornState();
@@ -58,11 +57,6 @@ void FVREditorModeManager::Tick( const float DeltaTime )
 			}
 			else if (HMDWornState == EHMDWornState::NotWorn && CurrentVREditorMode != nullptr)
 			{
-				if (GEditor->PlayWorld && !GEditor->bIsSimulatingInEditor)
-				{
-					CurrentVREditorMode->TogglePIEAndVREditor(); //-V595
-				}
-
 				EnableVREditor( false, false );
 			}
 		}
@@ -75,48 +69,11 @@ void FVREditorModeManager::Tick( const float DeltaTime )
 		CloseVREditor( bShouldDisableStereo );
 	}
 
-	// Only check for input if we started this play session from the VR Editor
-	if( GEditor->PlayWorld && !GEditor->bIsSimulatingInEditor && CurrentVREditorMode != nullptr)
-	{
-		// Shutdown PIE if we came from the VR Editor and we are not already requesting to start the VR Editor and when any of the players is holding down the required input
-		const float ShutDownInputKeyTime = 1.0f;
-		const float ShutDownTriggerValue = 0.7f;
-		TArray<APlayerController*> PlayerControllers;
-		GEngine->GetAllLocalPlayerControllers(PlayerControllers);
-		for(APlayerController* PlayerController : PlayerControllers)
-		{
-			const float LeftGripTimeDown = FMath::Max3(PlayerController->GetInputKeyTimeDown(EKeys::Vive_Left_Grip_Click),
-				PlayerController->GetInputKeyTimeDown(EKeys::ValveIndex_Left_Grip_Click),
-				PlayerController->GetInputKeyTimeDown(EKeys::OculusTouch_Left_Grip_Click));
-			const float RightGripTimeDown = FMath::Max3(PlayerController->GetInputKeyTimeDown(EKeys::Vive_Right_Grip_Click),
-				PlayerController->GetInputKeyTimeDown(EKeys::ValveIndex_Right_Grip_Click),
-				PlayerController->GetInputKeyTimeDown(EKeys::OculusTouch_Right_Grip_Click));
-			const float LeftTriggerValue = FMath::Max3(PlayerController->GetInputAxisValue(EKeys::Vive_Left_Trigger_Axis.GetFName()),
-				PlayerController->GetInputKeyTimeDown(EKeys::ValveIndex_Left_Trigger_Axis.GetFName()),
-				PlayerController->GetInputKeyTimeDown(EKeys::OculusTouch_Left_Trigger_Axis.GetFName()));
-			const float RightTriggerValue = FMath::Max3(PlayerController->GetInputAxisValue(EKeys::Vive_Right_Trigger_Axis.GetFName()),
-				PlayerController->GetInputKeyTimeDown(EKeys::ValveIndex_Right_Trigger_Axis.GetFName()),
-				PlayerController->GetInputKeyTimeDown(EKeys::OculusTouch_Right_Trigger_Axis.GetFName()));
-
-			if(LeftGripTimeDown > ShutDownInputKeyTime && RightGripTimeDown > ShutDownInputKeyTime &&
-				LeftTriggerValue > ShutDownTriggerValue && RightTriggerValue > ShutDownTriggerValue)
-			{
-				CurrentVREditorMode->TogglePIEAndVREditor();
-				// We need to clear the input of the playercontroller when exiting PIE. 
-				// Otherwise the input will still be pressed down causing toggle between PIE and VR Editor to be called instantly whenever entering PIE a second time.
-				PlayerController->PlayerInput->FlushPressedKeys();
-				break;
-			}
-		}
-	}
-	else
+	// Start the VR Editor mode
+	if (bEnableVRRequest)
 	{
-		// Start the VR Editor mode
-		if( bEnableVRRequest )
-		{
-			EnableVREditor( true, false );
-			bEnableVRRequest = false;
-		}
+		EnableVREditor(true, false);
+		bEnableVRRequest = false;
 	}
 }
 
diff --git a/Engine/Source/Editor/VREditor/VREditorPlacement.cpp b/Engine/Source/Editor/VREditor/VREditorPlacement.cpp
--- a/Engine/Source/Editor/VREditor/VREditorPlacement.cpp
+++ b/Engine/Source/Editor/VREditor/VREditorPlacement.cpp
@@ -35,22 +35,16 @@
 
 namespace VREd
 {
-	static FAutoConsoleVariable SizeOfActorsOverContentBrowserThumbnail( TEXT( "VREd.SizeOfActorsOverContentBrowserThumbnail" ), 6.0f, TEXT( "How large objects should be when rendered 'thumbnail size' over the Content Browser" ) );
 	static FAutoConsoleVariable HoverHapticFeedbackStrength( TEXT( "VREd.HoverHapticFeedbackStrength" ), 0.1f, TEXT( "Default strength for haptic feedback when hovering" ) );
 	static FAutoConsoleVariable HoverHapticFeedbackTime( TEXT( "VREd.HoverHapticFeedbackTime" ), 0.2f, TEXT( "The minimum time between haptic feedback for hovering" ) );
 	static FAutoConsoleVariable PivotPointTransformGizmo( TEXT( "VREd.PivotPointTransformGizmo" ), 1, TEXT( "If the pivot point transform gizmo is used instead of the bounding box gizmo" ) );
 	static FAutoConsoleVariable DragHapticFeedbackStrength( TEXT( "VREd.DragHapticFeedbackStrength" ), 1.0f, TEXT( "Default strength for haptic feedback when starting to drag objects" ) );
-	static FAutoConsoleVariable PlacementInterpolationEnabled( TEXT( "VREd.PlacementInterpolationEnabled" ), 1, TEXT( "If we interpolate to desired size and the end of the laser when dragging out of content browser." ) );
-	static FAutoConsoleVariable PlacementToEndOfLaser( TEXT( "VREd.PlacementToEndOfLaser" ), 0, TEXT( "If we interpolate to the end of the laser when dragging out of content browser." ) );
-	static FAutoConsoleVariable HideContentBrowserWhileDragging( TEXT( "VREd.HideContentBrowserWhileDragging" ), 0, TEXT( "" ) );
 }
 
 UVREditorPlacement::UVREditorPlacement() : 
 	Super(),
 	VRMode( nullptr ),
-	ViewportWorldInteraction(nullptr),
-	FloatingUIAssetDraggedFrom( nullptr ),
-	PlacingMaterialOrTextureAsset( nullptr )
+	ViewportWorldInteraction(nullptr)
 {
 }
 
@@ -59,9 +53,6 @@ void UVREditorPlacement::Init(UVREditorMode* InVRMode)
 	VRMode = InVRMode;
 	ViewportWorldInteraction = &InVRMode->GetWorldInteraction();
 
-	// Find out when the user drags stuff out of a content browser
-	FEditorDelegates::OnAssetDragStarted.AddUObject( this, &UVREditorPlacement::OnAssetDragStartedFromContentBrowser );
-
 	ViewportWorldInteraction->OnStopDragging().AddUObject( this, &UVREditorPlacement::StopDragging );
 	ViewportWorldInteraction->OnWorldScaleChanged().AddUObject( this, &UVREditorPlacement::UpdateNearClipPlaneOnScaleChange );
 }
@@ -92,12 +83,7 @@ void UVREditorPlacement::StopDragging( UViewportInteractor* Interactor )
 
 	const EViewportInteractionDraggingMode InteractorDraggingMode = Interactor->GetDraggingMode();
 
-	if ( InteractorDraggingMode == EViewportInteractionDraggingMode::Material )
-	{
-		// If we were placing a material, go ahead and do that now
-		PlaceDraggedMaterialOrTexture( Interactor );
-	}	
-	else if (Interactor->GetDraggingMode() == EViewportInteractionDraggingMode::TransformablesFreely)
+	if (Interactor->GetDraggingMode() == EViewportInteractionDraggingMode::TransformablesFreely)
 	{
 		VRMode->OnPlacePreviewActor().Broadcast(false);
 		GEditor->NoteSelectionChange();
@@ -116,389 +102,5 @@ void UVREditorPlacement::UpdateNearClipPlaneOnScaleChange(const float NewWorldTo
 	GNearClippingPlane = FMath::Min((VRMode->GetDefaultVRNearClipPlane()) * (NewWorldToMetersScale / DefaultWorldToMetersScale), VRMode->GetSavedEditorState().NearClipPlane);
 }
 
-void UVREditorPlacement::StartDraggingMaterialOrTexture( UViewportInteractor* Interactor, const FViewportActionKeyInput& Action, const FVector HitLocation, UObject* MaterialOrTextureAsset )
-{
-	check( MaterialOrTextureAsset != nullptr );
-
-	FVector LaserPointerStart, LaserPointerEnd;
-	const bool bHaveLaserPointer = Interactor->GetLaserPointer( /* Out */ LaserPointerStart, /* Out */ LaserPointerEnd );
-	if( bHaveLaserPointer )
-	{
-		PlacingMaterialOrTextureAsset = MaterialOrTextureAsset;
-
-		FViewportInteractorData& InteractorData = Interactor->GetInteractorData();
-
-		Interactor->SetDraggingMode( EViewportInteractionDraggingMode::Material );
-
-		// Starting a new drag, so make sure the other hand doesn't think it's assisting us
-		FViewportInteractorData& OtherInteractorData = Interactor->GetOtherInteractor()->GetInteractorData();
-		OtherInteractorData.bWasAssistingDrag = false;
-
-		InteractorData.bDraggingWithGrabberSphere = false;
-		InteractorData.bIsFirstDragUpdate = true;
-		InteractorData.bWasAssistingDrag = false;
-		InteractorData.DragRayLength = ( HitLocation - LaserPointerStart ).Size();
-		InteractorData.LastDragToLocation = HitLocation;
-		InteractorData.InteractorTransformAtDragStart = InteractorData.Transform;
-		InteractorData.GrabberSphereLocationAtDragStart = FVector::ZeroVector;
-		InteractorData.ImpactLocationAtDragStart = HitLocation;
-		InteractorData.DragTranslationVelocity = FVector::ZeroVector;
-		InteractorData.DragRayLengthVelocity = 0.0f;
-		InteractorData.bIsDrivingVelocityOfSimulatedTransformables = false;
-
-		InteractorData.DraggingTransformGizmoComponent = nullptr;
-	
-		InteractorData.DragOperationComponent.Reset();
-		InteractorData.GizmoStartTransform = FTransform::Identity;
-		InteractorData.GizmoLastTransform = InteractorData.GizmoTargetTransform = InteractorData.GizmoUnsnappedTargetTransform = InteractorData.GizmoInterpolationSnapshotTransform = InteractorData.GizmoStartTransform;
-		InteractorData.GizmoStartLocalBounds = FBox(EForceInit::ForceInit);
-
-		InteractorData.GizmoSpaceFirstDragUpdateOffsetAlongAxis = FVector::ZeroVector;	// Will be determined on first update
-		InteractorData.GizmoSpaceDragDeltaFromStartOffset = FVector::ZeroVector;	// Set every frame while dragging
-		InteractorData.LockedWorldDragMode = ELockedWorldDragMode::Unlocked;
-		InteractorData.GizmoScaleSinceDragStarted = 0.0f;
-		InteractorData.GizmoRotationRadiansSinceDragStarted = 0.0f;
-
-		ViewportWorldInteraction->SetDraggedSinceLastSelection( false );
-		ViewportWorldInteraction->SetLastDragGizmoStartTransform( FTransform::Identity );
-
-		// Play a haptic effect when objects aTrackingTransactionre picked up
-		Interactor->PlayHapticEffect( VREd::DragHapticFeedbackStrength->GetFloat() );
-	}
-}
-
-void UVREditorPlacement::OnAssetDragStartedFromContentBrowser( const TArray<FAssetData>& DraggedAssets, UActorFactory* FactoryToUse )
-{
-	FloatingUIAssetDraggedFrom = nullptr;
-	if( ViewportWorldInteraction != nullptr )
-	{
-		// Figure out which controller pressed the button and started dragging
-		// @todo vreditor placement: This logic could misfire.  Ideally we would be routed information from the pointer event, so we can determine the hand.
-		UVREditorInteractor* PlacingWithInteractor = nullptr;
-
-		const TArray<UViewportInteractor*>& Interactors = ViewportWorldInteraction->GetInteractors();
-		for( UViewportInteractor* Interactor : Interactors )
-		{
-			UVREditorInteractor* VREditorInteractor = Cast<UVREditorInteractor>( Interactor );
-			FViewportActionKeyInput* SelectAndMove_Action = Interactor->GetActionWithName( ViewportWorldActionTypes::SelectAndMove );
-			if( VREditorInteractor && SelectAndMove_Action != nullptr && SelectAndMove_Action->bIsInputCaptured )
-			{
-				if( VREditorInteractor->IsClickingOnUI() && !VREditorInteractor->IsRightClickingOnUI() )
-				{
-					PlacingWithInteractor = VREditorInteractor;
-					break;
-				}
-			}
-		}
-
-		if( PlacingWithInteractor != nullptr )
-		{
-			FEditorDelegates::LoadSelectedAssetsIfNeeded.Broadcast();
-
-			TArray< UObject* > DroppedObjects;
-			DroppedObjects.Reserve( DraggedAssets.Num() );
-
-			for( const FAssetData& AssetData : DraggedAssets )
-			{
-				UObject* AssetObj = AssetData.GetAsset();
-				if( AssetObj != nullptr )
-				{
-					// Don't add the same asset more than once
-					DroppedObjects.AddUnique( AssetObj );
-				}
-			}
-
-			if( DroppedObjects.Num() > 0 )
-			{
-				if( VREd::HideContentBrowserWhileDragging->GetInt() != 0 )
-				{
-					// Hide the UI panel that's being used to drag
-					FloatingUIAssetDraggedFrom = PlacingWithInteractor->GetLastHoveredWidgetComponent();
-					VRMode->GetUISystem().ShowEditorUIPanel( FloatingUIAssetDraggedFrom, PlacingWithInteractor, false, false, false );
-				}
-
-				const bool bShouldInterpolateFromDragLocation = VREd::PlacementInterpolationEnabled->GetInt() == 1;
-				StartPlacingObjects( DroppedObjects, FactoryToUse, PlacingWithInteractor, bShouldInterpolateFromDragLocation );
-
-				const UVREditorAssetContainer& AssetContainer = VRMode->GetAssetContainer();
-				VRMode->PlaySound(AssetContainer.DropFromContentBrowserSound, PlacingWithInteractor->GetTransform().GetLocation());
-			}
-		}
-	}
-}
-
-void UVREditorPlacement::StartPlacingObjects( const TArray<UObject*>& ObjectsToPlace, UActorFactory* FactoryToUse, UVREditorInteractor* PlacingWithInteractor, const bool bShouldInterpolateFromDragLocation )
-{
-	VRMode->OnPlacePreviewActor().Broadcast(true);
-	const bool bToEndOfLaser = VREd::PlacementToEndOfLaser->GetInt() == 1;
-
-	if( ViewportWorldInteraction == nullptr || PlacingWithInteractor == nullptr )
-	{
-		return;
-	}
-
-	const bool bIsPreview = false;	 // @todo vreditor placement: Consider supporting a "drop preview" actor (so you can cancel placement interactively)
-
-	bool bTransactionStarted = false;
-
-	{
-		// Cancel UI input
-		{
-			PlacingWithInteractor->SetIsClickingOnUI( false );
-			PlacingWithInteractor->SetIsRightClickingOnUI( false );
-
-			FViewportActionKeyInput* SelectAndMoveAction = PlacingWithInteractor->GetActionWithName( ViewportWorldActionTypes::SelectAndMove );
-			if( SelectAndMoveAction != nullptr )
-			{
-				SelectAndMoveAction->bIsInputCaptured = false;
-			}
-		}
-
-
-		TArray< UObject* > DroppedObjects;
-		TArray< AActor* > AllNewActors;
-
-		UObject* DraggingSingleMaterialOrTexture = nullptr;
-
-		FVector PlaceAt = PlacingWithInteractor->GetInteractorData().LastHoverLocationOverUI;
-
-		// Only place the object at the laser impact point if we're NOT going to interpolate to the impact
-		// location.  When interpolation is enabled, it looks much better to blend to the new location
-		if( !bShouldInterpolateFromDragLocation && bToEndOfLaser )
-		{
-			FVector HitLocation = FVector::ZeroVector;
-			if( ViewportWorldInteraction->FindPlacementPointUnderLaser( PlacingWithInteractor, /* Out */ HitLocation ) )
-			{
-				PlaceAt = HitLocation;
-			}
-		}
-
-		for( UObject* AssetObj : ObjectsToPlace )
-		{
-			bool bCanPlace = true;
-
-			const FString ObjectPath = AssetObj->GetPathName();
-			if( !ObjectTools::IsAssetValidForPlacing( ViewportWorldInteraction->GetWorld(), ObjectPath ) )
-			{
-				bCanPlace = false;
-			}
-			else
-			{
-				UClass* ClassObj = Cast<UClass>( AssetObj );
-				if( ClassObj )
-				{
-					if( !ObjectTools::IsClassValidForPlacing( ClassObj ) )
-					{
-						bCanPlace = false;
-					}
-
-					AssetObj = ClassObj->GetDefaultObject();
-				}
-			}
-
-			const bool bIsMaterialOrTexture = ( AssetObj->IsA( UMaterialInterface::StaticClass() ) || AssetObj->IsA( UTexture::StaticClass() ) );
-			if( bIsMaterialOrTexture && ObjectsToPlace.Num() == 1 )
-			{
-				DraggingSingleMaterialOrTexture = AssetObj;
-			}
-			else
-			{
-				// @todo mesheditor: We're dragging actors, so deactivate mesh editor mode for this.  They'll contend over transformables.
-				static FName MeshEditorModeName( TEXT( "MeshEditor" ) );
-				GLevelEditorModeTools().DeactivateMode( MeshEditorModeName );
-			}
-
-			// Check if the asset has an actor factory
-			bool bHasActorFactory = FActorFactoryAssetProxy::GetFactoryForAssetObject( AssetObj ) != nullptr;
-
-
-			if( !( AssetObj->IsA( AActor::StaticClass() ) || bHasActorFactory ) &&
-				!AssetObj->IsA( UBrushBuilder::StaticClass() ) )
-			{
-				bCanPlace = false;
-			}
-
-
-			FTrackingTransaction& TrackingTransaction = ViewportWorldInteraction->GetTrackingTransaction();
-			if( bCanPlace )
-			{
-				CA_SUPPRESS(6240);
-				if( !bTransactionStarted && !bIsPreview )
-				{
-					bTransactionStarted = true;
-
-					TrackingTransaction.TransCount++;
-					TrackingTransaction.Begin( LOCTEXT( "PlacingActors", "Placing Actors" ) );
-
-					// Suspend actor/component modification during each delta step to avoid recording unnecessary overhead into the transaction buffer
-					GEditor->DisableDeltaModification( true );
-				}
-
-				GEditor->ClickLocation = PlaceAt;
-				GEditor->ClickPlane = FPlane( PlaceAt, FVector::UpVector );
-
-				// Attempt to create actors from the dropped object
-				const bool bSelectNewActors = true;
-				const EObjectFlags NewObjectFlags = bIsPreview ? RF_Transient : RF_Transactional;
-
-				TArray<AActor*> NewActors = FLevelEditorViewportClient::TryPlacingActorFromObject( ViewportWorldInteraction->GetWorld()->GetCurrentLevel(), AssetObj, bSelectNewActors, NewObjectFlags, FactoryToUse );
-
-				AllNewActors.Append( NewActors );
-				if( NewActors.Num() > 0 )
-				{
-					DroppedObjects.Add( AssetObj );
-				}
-			}
-		}
-
-		// Cancel the transaction if nothing was placed
-		if( bTransactionStarted && AllNewActors.Num() == 0)
-		{
-			FTrackingTransaction& TrackingTransaction = ViewportWorldInteraction->GetTrackingTransaction();
-			--TrackingTransaction.TransCount;
-			TrackingTransaction.Cancel();
-			GEditor->DisableDeltaModification( false );
-		}
-
-		if( AllNewActors.Num() > 0 )
-		{
-			if( !bIsPreview )
-			{
-				if( IPlacementModeModule::IsAvailable() )
-				{
-					IPlacementModeModule::Get().AddToRecentlyPlaced( DroppedObjects, FactoryToUse );
-				}
-
-				FEditorDelegates::OnNewActorsDropped.Broadcast( DroppedObjects, AllNewActors );
-			}
-
-			static const auto CVarAllowCarryingCertainObjects = IConsoleManager::Get().FindTConsoleVariableDataInt(TEXT("VI.AllowCarryingCertainObjects"));
-			const bool bCanBeCarried = ViewportWorldInteraction->GetTransformables().Num() == 1 && ViewportWorldInteraction->GetTransformables()[0].Get()->ShouldBeCarried();
-			const bool bShouldInterpolateScaleFromDragLocation = bShouldInterpolateFromDragLocation && (CVarAllowCarryingCertainObjects->GetValueOnAnyThread() == 0 || !bCanBeCarried);
-
-			float DesiredScale = 1.f;
-			if( bShouldInterpolateScaleFromDragLocation )
-			{
-				FBox BoundsOfAllActors;
-				{
-					BoundsOfAllActors.Init();
-					for( AActor* NewActor : AllNewActors )
-					{
-						BoundsOfAllActors += NewActor->CalculateComponentsBoundingBoxInLocalSpace();
-					}
-				}
-				const float BoundsOfAllActorsSize = BoundsOfAllActors.GetSize().GetAbsMax();
-				const float UsedBoundsOfAllActorsSize = BoundsOfAllActorsSize == 0 ? 1 : BoundsOfAllActorsSize;
-				DesiredScale = (VREd::SizeOfActorsOverContentBrowserThumbnail->GetFloat() / UsedBoundsOfAllActorsSize) * ViewportWorldInteraction->GetWorldScaleFactor();
-
-				// Start the placed objects off scaled down to match the content browser thumbnail
-				for( AActor* NewActor : AllNewActors )
-				{
-					NewActor->SetActorScale3D(FVector(DesiredScale));
-				}
-			}
-			else if( bShouldInterpolateFromDragLocation )
-			{
-				PlaceAt = PlacingWithInteractor->GetInteractorData().Transform.GetLocation();
-			}
-
-			// We changed the initial scale of selected actors, so make sure our transformables and gizmo start transform are up to date.
-			GEditor->NoteSelectionChange();
-
-			// Start dragging the new actor(s)
-			UPrimitiveComponent* ClickedTransformGizmoComponent = nullptr;
-			const bool bIsPlacingNewObjects = true;
-			const bool bAllowInterpolationWhenPlacing = bShouldInterpolateFromDragLocation;
-			const bool bStartTransaction = false;
-			const bool bWithGrabberSphere = false;	// Always place using the laser, not the grabber sphere
-			ViewportWorldInteraction->StartDragging( PlacingWithInteractor, ClickedTransformGizmoComponent, PlaceAt, bIsPlacingNewObjects, bAllowInterpolationWhenPlacing, bToEndOfLaser, bStartTransaction, bWithGrabberSphere );
-
-			// If we're interpolating, update the target transform of the actors to use our overridden size.  When
-			// we placed them we set their size to be 'thumbnail sized', and we want them to interpolate to
-			// their actual size in the world
-			if( bShouldInterpolateFromDragLocation )
-			{
-				const FVector NewScale( 1.0f / DesiredScale );
-				FViewportInteractorData& InteractorData = PlacingWithInteractor->GetInteractorData();
-				InteractorData.GizmoUnsnappedTargetTransform.SetScale3D( NewScale );
-				InteractorData.GizmoLastTransform.SetScale3D( NewScale );
-				InteractorData.GizmoTargetTransform.SetScale3D( NewScale );
-			}
-		}
-
-		if( DraggingSingleMaterialOrTexture != nullptr )
-		{
-			const FViewportActionKeyInput Action( ViewportWorldActionTypes::SelectAndMove );
-
-			// Start dragging the material
-			StartDraggingMaterialOrTexture( PlacingWithInteractor,	Action, PlacingWithInteractor->GetHoverLocation(), DraggingSingleMaterialOrTexture );
-		}
-	}
-}
-
-
-void UVREditorPlacement::PlaceDraggedMaterialOrTexture( UViewportInteractor* Interactor )
-{
-	if( ensure( Interactor->GetDraggingMode() == EViewportInteractionDraggingMode::Material ) &&
-		PlacingMaterialOrTextureAsset != nullptr )
-	{
-		// Check to see if the laser pointer is over something we can drop on
-		UPrimitiveComponent* HitComponent = nullptr;
-		FVector HitLocation = FVector::ZeroVector;
-		{
-			const EHitResultGizmoFilterMode GizmoFilterMode = EHitResultGizmoFilterMode::NoGizmos;// Never place on top of gizmos, just ignore them
-			const bool bEvenIfUIIsInFront = true;	// Don't let the UI block placement
-			FHitResult HitResult = Interactor->GetHitResultFromLaserPointer( nullptr, GizmoFilterMode, nullptr, bEvenIfUIIsInFront );
-			if( HitResult.Actor.IsValid() )
-			{
-				if( ViewportWorldInteraction->IsInteractableComponent( HitResult.GetComponent() ) )	// @todo vreditor placement: We don't necessarily need to restrict to only VR-interactive components here
-				{
-					// Don't place materials on UI widget handles though!
-					if( Cast<AVREditorFloatingUI>( HitResult.GetComponent()->GetOwner() ) == nullptr )
-					{
-						HitComponent = HitResult.GetComponent();
-						HitLocation = HitResult.ImpactPoint;
-					}
-				}
-			}
-		}
-
-		if( HitComponent != nullptr )
-		{
-			UObject* ObjToUse = PlacingMaterialOrTextureAsset;
-
-			// Dropping a texture?
-			UTexture* DroppedObjAsTexture = Cast<UTexture>( ObjToUse );
-			if( DroppedObjAsTexture != NULL )
-			{
-				// Turn dropped textures into materials
-				ObjToUse = FLevelEditorViewportClient::GetOrCreateMaterialFromTexture( DroppedObjAsTexture );
-			}
-
-			// Dropping a material?
-			UMaterialInterface* DroppedObjAsMaterial = Cast<UMaterialInterface>( ObjToUse );
-			if( DroppedObjAsMaterial )
-			{
-				// @todo vreditor placement: How do we get the material ID that was dropped on?  Regular editor uses hit proxies.  We may need to augment FHitResult.
-				// @todo vreditor placement: Support optionally dropping on all materials, not only the impacted material
-				bool bPlaced = false;
-				check( VRMode );
-				VRMode->OnPlaceDraggedMaterial().Broadcast( HitComponent, DroppedObjAsMaterial, bPlaced );
-				if( !bPlaced )
-				{
-					const int32 TargetMaterialSlot = -1;	// All materials
-					bool bAppliedMaterial = FComponentEditorUtils::AttemptApplyMaterialToComponent( HitComponent, DroppedObjAsMaterial, TargetMaterialSlot );
-					if (bAppliedMaterial)
-					{
-						const UVREditorAssetContainer& AssetContainer = VRMode->GetAssetContainer();
-						VRMode->PlaySound(AssetContainer.DropFromContentBrowserSound, Interactor->GetTransform().GetLocation());
-					}
-				}
-			}
-		}
-	}
-
-	this->PlacingMaterialOrTextureAsset = nullptr;
-}
 
 #undef LOCTEXT_NAMESPACE
diff --git a/Engine/Source/Editor/VREditor/VREditorPlacement.h b/Engine/Source/Editor/VREditor/VREditorPlacement.h
--- a/Engine/Source/Editor/VREditor/VREditorPlacement.h
+++ b/Engine/Source/Editor/VREditor/VREditorPlacement.h
@@ -33,8 +33,6 @@ public:
 	/** Removes registered event */
 	void Shutdown();
 
-	/** Try spawn and start placing the specified objects */
-	void StartPlacingObjects( const TArray<UObject*>& ObjectsToPlace, class UActorFactory* FactoryToUse, class UVREditorInteractor* PlacingWithInteractor, const bool bShouldInterpolateFromDragLocation );
 
 protected:
 	/** When an interactor stops dragging */
@@ -43,15 +41,6 @@ protected:
 	/** When the world scale changes, update the near clip plane */
 	void UpdateNearClipPlaneOnScaleChange(const float NewWorldToMetersScale);
 
-	/** Starts dragging a material, allowing the user to drop it on an object in the scene to place it */
-	void StartDraggingMaterialOrTexture( UViewportInteractor* Interactor, const FViewportActionKeyInput& Action, const FVector HitLocation, UObject* MaterialOrTextureAsset );
-
-	/** Tries to place whatever material or texture that's being dragged on the object under the hand's laser pointer */
-	void PlaceDraggedMaterialOrTexture( UViewportInteractor* Interactor );
-
-	/** Called when FEditorDelegates::OnAssetDragStarted is broadcast */
-	void OnAssetDragStartedFromContentBrowser( const TArray<FAssetData>& DraggedAssets, class UActorFactory* FactoryToUse );
-
 protected:
 
 	/** Owning object */
diff --git a/Engine/Source/Programs/UnrealBuildTool/Configuration/UEBuildModuleCPP.cs b/Engine/Source/Programs/UnrealBuildTool/Configuration/UEBuildModuleCPP.cs
--- a/Engine/Source/Programs/UnrealBuildTool/Configuration/UEBuildModuleCPP.cs
+++ b/Engine/Source/Programs/UnrealBuildTool/Configuration/UEBuildModuleCPP.cs
@@ -119,14 +119,6 @@ namespace UnrealBuildTool
 			new KeyValuePair<string, string>("GameplayAbilitiesEditor", "BlueprintGraph"),
             new KeyValuePair<string, string>("UnrealEd", "ViewportInteraction"),
             new KeyValuePair<string, string>("UnrealEd", "VREditor"),
-            new KeyValuePair<string, string>("LandscapeEditor", "ViewportInteraction"),
-            new KeyValuePair<string, string>("LandscapeEditor", "VREditor"),
-            new KeyValuePair<string, string>("FoliageEdit", "ViewportInteraction"),
-            new KeyValuePair<string, string>("FoliageEdit", "VREditor"),
-            new KeyValuePair<string, string>("MeshPaint", "ViewportInteraction"),
-            new KeyValuePair<string, string>("MeshPaint", "VREditor"),
-            new KeyValuePair<string, string>("MeshPaintMode", "ViewportInteraction"),
-            new KeyValuePair<string, string>("MeshPaintMode", "VREditor"),
             new KeyValuePair<string, string>("Sequencer", "ViewportInteraction"),
             new KeyValuePair<string, string>("NavigationSystem", "UnrealEd"),
         };
