Automated wrapped object pointer upgrade for Engine + ShooterGame + ShooterGame referenced plugins

diff --git a/Engine/Source/Editor/UnrealEd/Classes/Factories/FbxImportUI.h b/Engine/Source/Editor/UnrealEd/Classes/Factories/FbxImportUI.h
--- a/Engine/Source/Editor/UnrealEd/Classes/Factories/FbxImportUI.h
+++ b/Engine/Source/Editor/UnrealEd/Classes/Factories/FbxImportUI.h
@@ -127,7 +127,7 @@ public:
 
 	/** Skeleton to use for imported asset. When importing a mesh, leaving this as "None" will create a new skeleton. When importing an animation this MUST be specified to import the asset. */
 	UPROPERTY(EditAnywhere, BlueprintReadWrite, Category=Mesh, meta=(ImportType="SkeletalMesh|Animation"))
-	class USkeleton* Skeleton;
+	TObjectPtr<class USkeleton> Skeleton;
 
 	/** If checked, create new PhysicsAsset if it doesn't have it */
 	UPROPERTY(EditAnywhere, BlueprintReadWrite, AdvancedDisplay, config, Category=Mesh, meta=(ImportType="SkeletalMesh"))
@@ -135,7 +135,7 @@ public:
 
 	/** If this is set, use this PhysicsAsset. It is possible bCreatePhysicsAsset == false, and PhysicsAsset == NULL. It is possible they do not like to create anything. */
 	UPROPERTY(EditAnywhere, BlueprintReadWrite, AdvancedDisplay, Category=Mesh, meta=(ImportType="SkeletalMesh", editcondition="!bCreatePhysicsAsset"))
-	class UPhysicsAsset* PhysicsAsset;
+	TObjectPtr<class UPhysicsAsset> PhysicsAsset;
 
 	/** If checked, the editor will automatically compute screen size values for the static mesh's LODs. If unchecked, the user can enter custom screen size values for each LOD. */
 	UPROPERTY(EditAnywhere, BlueprintReadWrite, config, Category = LODSettings, meta = (ImportType = "StaticMesh", DisplayName = "Auto Compute LOD Screen Size"))
@@ -199,19 +199,19 @@ public:
 
 	/** Import data used when importing static meshes */
 	UPROPERTY(EditAnywhere, BlueprintReadWrite, Transient, Instanced, Category = Mesh, meta=(ImportType = "StaticMesh"))
-	class UFbxStaticMeshImportData* StaticMeshImportData;
+	TObjectPtr<class UFbxStaticMeshImportData> StaticMeshImportData;
 
 	/** Import data used when importing skeletal meshes */
 	UPROPERTY(EditAnywhere, BlueprintReadWrite, Transient, Instanced, Category=Mesh, meta=(ImportType = "SkeletalMesh"))
-	class UFbxSkeletalMeshImportData* SkeletalMeshImportData;
+	TObjectPtr<class UFbxSkeletalMeshImportData> SkeletalMeshImportData;
 
 	/** Import data used when importing animations */
 	UPROPERTY(EditAnywhere, BlueprintReadWrite, Transient, Instanced, Category=Animation, meta=(editcondition="bImportAnimations", ImportType = "Animation"))
-	class UFbxAnimSequenceImportData* AnimSequenceImportData;
+	TObjectPtr<class UFbxAnimSequenceImportData> AnimSequenceImportData;
 
 	/** Import data used when importing textures */
 	UPROPERTY(EditAnywhere, BlueprintReadWrite, Transient, Instanced, Category=Material)
-	class UFbxTextureImportData* TextureImportData;
+	TObjectPtr<class UFbxTextureImportData> TextureImportData;
 
 	/** If true the automated import path should detect the import type.  If false the import type was specified by the user */
 	UPROPERTY(BlueprintReadWrite, Category = Miscellaneous)
diff --git a/Engine/Source/Editor/UnrealEd/Classes/Factories/FbxSceneImportFactory.h b/Engine/Source/Editor/UnrealEd/Classes/Factories/FbxSceneImportFactory.h
--- a/Engine/Source/Editor/UnrealEd/Classes/Factories/FbxSceneImportFactory.h
+++ b/Engine/Source/Editor/UnrealEd/Classes/Factories/FbxSceneImportFactory.h
@@ -357,31 +357,31 @@ class UNREALED_API UFbxSceneImportFactory : public USceneImportFactory
 
 	/** Import options UI detail when importing fbx scene */
 	UPROPERTY(BlueprintReadWrite, Category = "Editor Scripting | Fbx Scene Import", Transient)
-	class UFbxSceneImportOptions* SceneImportOptions;
+	TObjectPtr<class UFbxSceneImportOptions> SceneImportOptions;
 	
 	/** Import options UI detail when importing fbx scene static mesh*/
 	UPROPERTY(BlueprintReadWrite, Category = "Editor Scripting | Fbx Scene Import", Transient)
-	class UFbxSceneImportOptionsStaticMesh* SceneImportOptionsStaticMesh;
+	TObjectPtr<class UFbxSceneImportOptionsStaticMesh> SceneImportOptionsStaticMesh;
 	
 	/** Import options UI detail when importing fbx scene skeletal mesh*/
 	UPROPERTY(BlueprintReadWrite, Category = "Editor Scripting | Fbx Scene Import", Transient)
-	class UFbxSceneImportOptionsSkeletalMesh* SceneImportOptionsSkeletalMesh;
+	TObjectPtr<class UFbxSceneImportOptionsSkeletalMesh> SceneImportOptionsSkeletalMesh;
 	
 	/** Import data used when importing static meshes */
 	UPROPERTY(Transient)
-	class UFbxStaticMeshImportData* StaticMeshImportData;
+	TObjectPtr<class UFbxStaticMeshImportData> StaticMeshImportData;
 
 	/** Import data used when importing skeletal meshes */
 	UPROPERTY(Transient)
-	class UFbxSkeletalMeshImportData* SkeletalMeshImportData;
+	TObjectPtr<class UFbxSkeletalMeshImportData> SkeletalMeshImportData;
 
 	/** Import data used when importing animations */
 	UPROPERTY(Transient)
-	class UFbxAnimSequenceImportData* AnimSequenceImportData;
+	TObjectPtr<class UFbxAnimSequenceImportData> AnimSequenceImportData;
 
 	/** Import data used when importing textures */
 	UPROPERTY(Transient)
-	class UFbxTextureImportData* TextureImportData;
+	TObjectPtr<class UFbxTextureImportData> TextureImportData;
 	
 	/* Default Options always have the same name "Default" */
 	static FString DefaultOptionName;
@@ -444,7 +444,7 @@ protected:
 
 	/** Pointer on the fbx scene import data, we fill this object to be able to do re import of the scene */
 	UPROPERTY()
-	UFbxSceneImportData* ReimportData;
+	TObjectPtr<UFbxSceneImportData> ReimportData;
 	
 	/** Assets created by the factory*/
 	TMap<TSharedPtr<FFbxAttributeInfo>, UObject*> AllNewAssets;
diff --git a/Engine/Source/Editor/UnrealEd/Classes/Factories/MaterialFactoryNew.h b/Engine/Source/Editor/UnrealEd/Classes/Factories/MaterialFactoryNew.h
--- a/Engine/Source/Editor/UnrealEd/Classes/Factories/MaterialFactoryNew.h
+++ b/Engine/Source/Editor/UnrealEd/Classes/Factories/MaterialFactoryNew.h
@@ -18,7 +18,7 @@ class UMaterialFactoryNew : public UFactory
 
 	/** An initial texture to place in the newly created material */
 	UPROPERTY()
-	class UTexture* InitialTexture;
+	TObjectPtr<class UTexture> InitialTexture;
 
 	//~ Begin UFactory Interface
 	virtual UObject* FactoryCreateNew(UClass* Class,UObject* InParent,FName Name,EObjectFlags Flags,UObject* Context,FFeedbackContext* Warn) override;
diff --git a/Engine/Source/Editor/UnrealEd/Classes/Factories/MaterialFunctionInstanceFactory.h b/Engine/Source/Editor/UnrealEd/Classes/Factories/MaterialFunctionInstanceFactory.h
--- a/Engine/Source/Editor/UnrealEd/Classes/Factories/MaterialFunctionInstanceFactory.h
+++ b/Engine/Source/Editor/UnrealEd/Classes/Factories/MaterialFunctionInstanceFactory.h
@@ -13,7 +13,7 @@ class UMaterialFunctionInstanceFactory : public UFactory
 	GENERATED_UCLASS_BODY()
 
 	UPROPERTY()
-	class UMaterialFunctionInterface* InitialParent;
+	TObjectPtr<class UMaterialFunctionInterface> InitialParent;
 
 	//~ Begin UFactory Interface
 	virtual UObject* FactoryCreateNew(UClass* Class,UObject* InParent,FName Name,EObjectFlags Flags,UObject* Context,FFeedbackContext* Warn) override;
diff --git a/Engine/Source/Editor/UnrealEd/Classes/Factories/MaterialInstanceConstantFactoryNew.h b/Engine/Source/Editor/UnrealEd/Classes/Factories/MaterialInstanceConstantFactoryNew.h
--- a/Engine/Source/Editor/UnrealEd/Classes/Factories/MaterialInstanceConstantFactoryNew.h
+++ b/Engine/Source/Editor/UnrealEd/Classes/Factories/MaterialInstanceConstantFactoryNew.h
@@ -17,7 +17,7 @@ class UMaterialInstanceConstantFactoryNew : public UFactory
 	GENERATED_UCLASS_BODY()
 
 	UPROPERTY()
-	class UMaterialInterface* InitialParent;
+	TObjectPtr<class UMaterialInterface> InitialParent;
 
 	//~ Begin UFactory Interface
 	UNREALED_API virtual UObject* FactoryCreateNew(UClass* Class,UObject* InParent,FName Name,EObjectFlags Flags,UObject* Context,FFeedbackContext* Warn) override;
diff --git a/Engine/Source/Editor/UnrealEd/Classes/Factories/PhysicsAssetFactory.h b/Engine/Source/Editor/UnrealEd/Classes/Factories/PhysicsAssetFactory.h
--- a/Engine/Source/Editor/UnrealEd/Classes/Factories/PhysicsAssetFactory.h
+++ b/Engine/Source/Editor/UnrealEd/Classes/Factories/PhysicsAssetFactory.h
@@ -17,7 +17,7 @@ class UPhysicsAssetFactory : public UFactory
 
 	// The skeletal mesh with which to initialize this physics asset
 	UPROPERTY()
-	USkeletalMesh* TargetSkeletalMesh;
+	TObjectPtr<USkeletalMesh> TargetSkeletalMesh;
 
 	// UFactory Interface
 	virtual bool ConfigureProperties() override;
diff --git a/Engine/Source/Editor/UnrealEd/Classes/Factories/PoseAssetFactory.h b/Engine/Source/Editor/UnrealEd/Classes/Factories/PoseAssetFactory.h
--- a/Engine/Source/Editor/UnrealEd/Classes/Factories/PoseAssetFactory.h
+++ b/Engine/Source/Editor/UnrealEd/Classes/Factories/PoseAssetFactory.h
@@ -24,7 +24,7 @@ class UPoseAssetFactory : public UFactory
 
 	/* Used when creating a composite from an AnimSequence, becomes the only AnimSequence contained */
 	UPROPERTY(BlueprintReadWrite, EditAnywhere, Category = "PoseAssetFactory")
-	class UAnimSequence* SourceAnimation;
+	TObjectPtr<class UAnimSequence> SourceAnimation;
 
 	/** Optional. If specified the poses will be named according to this array.
 	If you don't specify, or you don't specify enough names for all poses, then default names will be generated.*/
@@ -34,10 +34,10 @@ class UPoseAssetFactory : public UFactory
 	// Only used for AnimationEditorUtils::ExecuteNewAnimAsset template. Do not use directly.
 
 	UPROPERTY()
-	class USkeleton* TargetSkeleton;
+	TObjectPtr<class USkeleton> TargetSkeleton;
 
 	UPROPERTY()
-	class USkeletalMesh* PreviewSkeletalMesh;
+	TObjectPtr<class USkeletalMesh> PreviewSkeletalMesh;
 
 	//~ Begin UFactory Interface
 	virtual bool ConfigureProperties() override;
diff --git a/Engine/Source/Editor/UnrealEd/Classes/Factories/ReimportTextureFactory.h b/Engine/Source/Editor/UnrealEd/Classes/Factories/ReimportTextureFactory.h
--- a/Engine/Source/Editor/UnrealEd/Classes/Factories/ReimportTextureFactory.h
+++ b/Engine/Source/Editor/UnrealEd/Classes/Factories/ReimportTextureFactory.h
@@ -21,7 +21,7 @@ class UReimportTextureFactory : public UTextureFactory, public FReimportHandler
 	GENERATED_UCLASS_BODY()
 
 	UPROPERTY()
-	class UTexture* pOriginalTex;
+	TObjectPtr<class UTexture> pOriginalTex;
 
 
 	//~ Begin FReimportHandler Interface
diff --git a/Engine/Source/Editor/UnrealEd/Classes/Factories/SkeletonFactory.h b/Engine/Source/Editor/UnrealEd/Classes/Factories/SkeletonFactory.h
--- a/Engine/Source/Editor/UnrealEd/Classes/Factories/SkeletonFactory.h
+++ b/Engine/Source/Editor/UnrealEd/Classes/Factories/SkeletonFactory.h
@@ -21,7 +21,7 @@ class USkeletonFactory : public UFactory
 
 	// The skeletal mesh with which to initialize this skeleton.
 	UPROPERTY()
-	class USkeletalMesh* TargetSkeletalMesh;
+	TObjectPtr<class USkeletalMesh> TargetSkeletalMesh;
 
 	//~ Begin UFactory Interface
 	virtual bool ConfigureProperties() override;
diff --git a/Engine/Source/Editor/UnrealEd/Classes/Factories/SlateBrushAssetFactory.h b/Engine/Source/Editor/UnrealEd/Classes/Factories/SlateBrushAssetFactory.h
--- a/Engine/Source/Editor/UnrealEd/Classes/Factories/SlateBrushAssetFactory.h
+++ b/Engine/Source/Editor/UnrealEd/Classes/Factories/SlateBrushAssetFactory.h
@@ -15,7 +15,7 @@ class USlateBrushAssetFactory : public UFactory
 
 	/** An initial texture to assign to the newly created slate brush */
 	UPROPERTY()
-	class UTexture2D* InitialTexture;
+	TObjectPtr<class UTexture2D> InitialTexture;
 
 	//~ Begin UFactory Interface
 	virtual FText GetDisplayName() const override;
diff --git a/Engine/Source/Editor/UnrealEd/Classes/Factories/SubUVAnimationFactory.h b/Engine/Source/Editor/UnrealEd/Classes/Factories/SubUVAnimationFactory.h
--- a/Engine/Source/Editor/UnrealEd/Classes/Factories/SubUVAnimationFactory.h
+++ b/Engine/Source/Editor/UnrealEd/Classes/Factories/SubUVAnimationFactory.h
@@ -18,7 +18,7 @@ class USubUVAnimationFactory : public UFactory
 
 	/** An initial texture to use */
 	UPROPERTY()
-	class UTexture2D* InitialTexture;
+	TObjectPtr<class UTexture2D> InitialTexture;
 
 	// UFactory interface
 	virtual UObject* FactoryCreateNew(UClass* Class, UObject* InParent, FName Name, EObjectFlags Flags, UObject* Context, FFeedbackContext* Warn) override;
diff --git a/Engine/Source/Editor/UnrealEd/Classes/Factories/Texture2DArrayFactory.h b/Engine/Source/Editor/UnrealEd/Classes/Factories/Texture2DArrayFactory.h
--- a/Engine/Source/Editor/UnrealEd/Classes/Factories/Texture2DArrayFactory.h
+++ b/Engine/Source/Editor/UnrealEd/Classes/Factories/Texture2DArrayFactory.h
@@ -16,7 +16,7 @@ class UTexture2DArrayFactory : public UFactory
 	GENERATED_UCLASS_BODY()
 
 	UPROPERTY()
-	TArray<UTexture2D*> InitialTextures;
+	TArray<TObjectPtr<UTexture2D>> InitialTextures;
 
 	//~ Begin UFactory Interface
 	virtual bool CanCreateNew() const override;
diff --git a/Engine/Source/Editor/UnrealEd/Classes/Factories/TrueTypeFontFactory.h b/Engine/Source/Editor/UnrealEd/Classes/Factories/TrueTypeFontFactory.h
--- a/Engine/Source/Editor/UnrealEd/Classes/Factories/TrueTypeFontFactory.h
+++ b/Engine/Source/Editor/UnrealEd/Classes/Factories/TrueTypeFontFactory.h
@@ -22,7 +22,7 @@ class UNREALED_API UTrueTypeFontFactory : public UTextureFactory, public FReimpo
 
 	/** Import options for this font */
 	UPROPERTY(EditAnywhere, Instanced, Category = TrueTypeFontFactory, meta = (ToolTip = "Import options for the font"))
-	class UFontImportOptions* ImportOptions;
+	TObjectPtr<class UFontImportOptions> ImportOptions;
 
 	/** True when the font dialog was shown for this factory during the non-legacy creation process */
 	UPROPERTY()
diff --git a/Engine/Source/Editor/UnrealEd/Classes/MaterialEditor/DEditorFontParameterValue.h b/Engine/Source/Editor/UnrealEd/Classes/MaterialEditor/DEditorFontParameterValue.h
--- a/Engine/Source/Editor/UnrealEd/Classes/MaterialEditor/DEditorFontParameterValue.h
+++ b/Engine/Source/Editor/UnrealEd/Classes/MaterialEditor/DEditorFontParameterValue.h
@@ -14,7 +14,7 @@ struct UNREALED_API FDFontParameters
 	GENERATED_USTRUCT_BODY()
 
 	UPROPERTY(EditAnywhere, Category = DFontParameter)
-	class UFont* FontValue = nullptr;
+	TObjectPtr<class UFont> FontValue = nullptr;
 
 	UPROPERTY(EditAnywhere, Category = DFontParameter)
 	int32 FontPage;
diff --git a/Engine/Source/Editor/UnrealEd/Classes/MaterialEditor/DEditorRuntimeVirtualTextureParameterValue.h b/Engine/Source/Editor/UnrealEd/Classes/MaterialEditor/DEditorRuntimeVirtualTextureParameterValue.h
--- a/Engine/Source/Editor/UnrealEd/Classes/MaterialEditor/DEditorRuntimeVirtualTextureParameterValue.h
+++ b/Engine/Source/Editor/UnrealEd/Classes/MaterialEditor/DEditorRuntimeVirtualTextureParameterValue.h
@@ -14,5 +14,5 @@ class UNREALED_API UDEditorRuntimeVirtualTextureParameterValue : public UDEditor
 	GENERATED_UCLASS_BODY()
 
 	UPROPERTY(EditAnywhere, Category=DEditorTextureParameterValue)
-	class URuntimeVirtualTexture* ParameterValue;
+	TObjectPtr<class URuntimeVirtualTexture> ParameterValue;
 };
diff --git a/Engine/Source/Editor/UnrealEd/Classes/MaterialEditor/DEditorTextureParameterValue.h b/Engine/Source/Editor/UnrealEd/Classes/MaterialEditor/DEditorTextureParameterValue.h
--- a/Engine/Source/Editor/UnrealEd/Classes/MaterialEditor/DEditorTextureParameterValue.h
+++ b/Engine/Source/Editor/UnrealEd/Classes/MaterialEditor/DEditorTextureParameterValue.h
@@ -14,7 +14,7 @@ class UNREALED_API UDEditorTextureParameterValue : public UDEditorParameterValue
 	GENERATED_UCLASS_BODY()
 
 	UPROPERTY(EditAnywhere, Category=DEditorTextureParameterValue)
-	class UTexture* ParameterValue;
+	TObjectPtr<class UTexture> ParameterValue;
 
 	UPROPERTY(Transient)
 	FParameterChannelNames ChannelNames;
diff --git a/Engine/Source/Editor/UnrealEd/Classes/MaterialEditor/MaterialEditorInstanceConstant.h b/Engine/Source/Editor/UnrealEd/Classes/MaterialEditor/MaterialEditorInstanceConstant.h
--- a/Engine/Source/Editor/UnrealEd/Classes/MaterialEditor/MaterialEditorInstanceConstant.h
+++ b/Engine/Source/Editor/UnrealEd/Classes/MaterialEditor/MaterialEditorInstanceConstant.h
@@ -33,7 +33,7 @@ struct FEditorParameterGroup
 	TEnumAsByte<EMaterialParameterAssociation> GroupAssociation;
 
 	UPROPERTY(EditAnywhere, editfixedsize, Instanced, Category=EditorParameterGroup)
-	TArray<class UDEditorParameterValue*> Parameters;
+	TArray<TObjectPtr<class UDEditorParameterValue>> Parameters;
 
 	UPROPERTY()
 	int32 GroupSortPriority;
@@ -93,7 +93,7 @@ struct FEditorTextureParameterValue : public FEditorParameterValue
 	GENERATED_USTRUCT_BODY()
 
 	UPROPERTY(EditAnywhere, Category=EditorTextureParameterValue)
-	class UTexture* ParameterValue;
+	TObjectPtr<class UTexture> ParameterValue;
 
 	FEditorTextureParameterValue()
 		: ParameterValue(NULL)
@@ -107,7 +107,7 @@ struct FEditorFontParameterValue : public FEditorParameterValue
 	GENERATED_USTRUCT_BODY()
 
 	UPROPERTY(EditAnywhere, Category=EditorFontParameterValue)
-	class UFont* FontValue;
+	TObjectPtr<class UFont> FontValue;
 
 	UPROPERTY(EditAnywhere, Category=EditorFontParameterValue)
 	int32 FontPage;
@@ -125,7 +125,7 @@ struct FEditorMaterialLayersParameterValue : public FEditorParameterValue
 	GENERATED_USTRUCT_BODY()
 
 	UPROPERTY(EditAnywhere, Category=EditorLayersParameterValue)
-	class UMaterialFunctionInterface* FunctionValue;
+	TObjectPtr<class UMaterialFunctionInterface> FunctionValue;
 
 	FEditorMaterialLayersParameterValue()
 		: FunctionValue(NULL)
@@ -219,11 +219,11 @@ class UNREALED_API UMaterialEditorInstanceConstant : public UObject
 
 	/** Physical material to use for this graphics material. Used for sounds, effects etc.*/
 	UPROPERTY(EditAnywhere, Category=MaterialEditorInstanceConstant)
-	class UPhysicalMaterial* PhysMaterial;
+	TObjectPtr<class UPhysicalMaterial> PhysMaterial;
 
 	// since the Parent may point across levels and the property editor needs to import this text, it must be marked lazy so it doesn't set itself to NULL in FindImportedObject
 	UPROPERTY(EditAnywhere, Category=MaterialEditorInstanceConstant, meta=(DisplayThumbnail="true"))
-	class UMaterialInterface* Parent;
+	TObjectPtr<class UMaterialInterface> Parent;
 
 	UPROPERTY(EditAnywhere, editfixedsize, Category=MaterialEditorInstanceConstant)
 	TArray<struct FEditorParameterGroup> ParameterGroups;
@@ -234,7 +234,7 @@ class UNREALED_API UMaterialEditorInstanceConstant : public UObject
 
 	/** SubsurfaceProfile, for Screen Space Subsurface Scattering */
 	UPROPERTY(EditAnywhere, BlueprintReadOnly, Category = Material, meta = (DisplayName = "Subsurface Profile"))
-	class USubsurfaceProfile* SubsurfaceProfile;
+	TObjectPtr<class USubsurfaceProfile> SubsurfaceProfile;
 
 	/** Defines if SubsurfaceProfile from tis instance is used or it uses the parent one. */
 	UPROPERTY(EditAnywhere, Category = MaterialEditorInstanceConstant)
@@ -253,10 +253,10 @@ class UNREALED_API UMaterialEditorInstanceConstant : public UObject
 	FMaterialInstanceBasePropertyOverrides BasePropertyOverrides;
 
 	UPROPERTY()
-	class UMaterialInstanceConstant* SourceInstance;
+	TObjectPtr<class UMaterialInstanceConstant> SourceInstance;
 
 	UPROPERTY()
-	class UMaterialFunctionInstance* SourceFunction;	
+	TObjectPtr<class UMaterialFunctionInstance> SourceFunction;	
 
 	UPROPERTY(transient, duplicatetransient)
 	TArray<FMaterialParameterInfo> VisibleExpressions;
@@ -328,10 +328,10 @@ class UNREALED_API UMaterialEditorInstanceConstant : public UObject
 
 #if WITH_EDITORONLY_DATA
 	UPROPERTY()
-	TArray<class UMaterialInstanceConstant*> StoredLayerPreviews;
+	TArray<TObjectPtr<class UMaterialInstanceConstant>> StoredLayerPreviews;
 
 	UPROPERTY()
-	TArray<class UMaterialInstanceConstant*> StoredBlendPreviews;
+	TArray<TObjectPtr<class UMaterialInstanceConstant>> StoredBlendPreviews;
 #endif
 
 	/** Whether or not we should show only overridden properties*/
diff --git a/Engine/Source/Editor/UnrealEd/Classes/MaterialEditor/MaterialEditorPreviewParameters.h b/Engine/Source/Editor/UnrealEd/Classes/MaterialEditor/MaterialEditorPreviewParameters.h
--- a/Engine/Source/Editor/UnrealEd/Classes/MaterialEditor/MaterialEditorPreviewParameters.h
+++ b/Engine/Source/Editor/UnrealEd/Classes/MaterialEditor/MaterialEditorPreviewParameters.h
@@ -31,20 +31,20 @@ public:
 	TArray<struct FEditorParameterGroup> ParameterGroups;
 
 	UPROPERTY()
-	class UMaterial* PreviewMaterial;
+	TObjectPtr<class UMaterial> PreviewMaterial;
 
 	UPROPERTY()
-	class UMaterialFunction* OriginalFunction;
+	TObjectPtr<class UMaterialFunction> OriginalFunction;
 
 	UPROPERTY()
-	class UMaterial* OriginalMaterial;
+	TObjectPtr<class UMaterial> OriginalMaterial;
 
 #if WITH_EDITORONLY_DATA
 	UPROPERTY()
-	TArray<class UMaterialInstanceConstant*> StoredLayerPreviews;
+	TArray<TObjectPtr<class UMaterialInstanceConstant>> StoredLayerPreviews;
 
 	UPROPERTY()
-	TArray<class UMaterialInstanceConstant*> StoredBlendPreviews;
+	TArray<TObjectPtr<class UMaterialInstanceConstant>> StoredBlendPreviews;
 #endif
 
 	//~ Begin UObject Interface.
diff --git a/Engine/Source/Editor/UnrealEd/Classes/MaterialGraph/MaterialGraph.h b/Engine/Source/Editor/UnrealEd/Classes/MaterialGraph/MaterialGraph.h
--- a/Engine/Source/Editor/UnrealEd/Classes/MaterialGraph/MaterialGraph.h
+++ b/Engine/Source/Editor/UnrealEd/Classes/MaterialGraph/MaterialGraph.h
@@ -105,19 +105,19 @@ class UNREALED_API UMaterialGraph : public UEdGraph
 
 	/** Material this Graph represents */
 	UPROPERTY()
-	class UMaterial*				Material;
+	TObjectPtr<class UMaterial>				Material;
 
 	/** Material Function this Graph represents (NULL for Materials) */
 	UPROPERTY()
-	class UMaterialFunction*		MaterialFunction;
+	TObjectPtr<class UMaterialFunction>		MaterialFunction;
 
 	/** Root node representing Material inputs (NULL for Material Functions) */
 	UPROPERTY()
-	class UMaterialGraphNode_Root*	RootNode;
+	TObjectPtr<class UMaterialGraphNode_Root>	RootNode;
 
 	/** Expression this subgraph represents (NULL if not subgraph, Material [Function] still populated) */
 	UPROPERTY()
-	UMaterialExpression*			SubgraphExpression;
+	TObjectPtr<UMaterialExpression>			SubgraphExpression;
 
 	/** List of Material Inputs (not set up for Material Functions) */
 	TArray<FMaterialInputInfo> MaterialInputs;
diff --git a/Engine/Source/Editor/UnrealEd/Classes/MaterialGraph/MaterialGraphNode.h b/Engine/Source/Editor/UnrealEd/Classes/MaterialGraph/MaterialGraphNode.h
--- a/Engine/Source/Editor/UnrealEd/Classes/MaterialGraph/MaterialGraphNode.h
+++ b/Engine/Source/Editor/UnrealEd/Classes/MaterialGraph/MaterialGraphNode.h
@@ -17,7 +17,7 @@ class UMaterialGraphNode : public UMaterialGraphNode_Base
 
 	/** Material Expression this node is representing */
 	UPROPERTY()
-	class UMaterialExpression* MaterialExpression;
+	TObjectPtr<class UMaterialExpression> MaterialExpression;
 
 	/** Set to true when Expression Preview compiles, so we can update SGraphNode */
 	bool bPreviewNeedsUpdate;
diff --git a/Engine/Source/Editor/UnrealEd/Classes/MaterialGraph/MaterialGraphNode_Comment.h b/Engine/Source/Editor/UnrealEd/Classes/MaterialGraph/MaterialGraphNode_Comment.h
--- a/Engine/Source/Editor/UnrealEd/Classes/MaterialGraph/MaterialGraphNode_Comment.h
+++ b/Engine/Source/Editor/UnrealEd/Classes/MaterialGraph/MaterialGraphNode_Comment.h
@@ -16,7 +16,7 @@ class UMaterialGraphNode_Comment : public UEdGraphNode_Comment
 
 	/** Material Comment that this node represents */
 	UPROPERTY()
-	class UMaterialExpressionComment* MaterialExpressionComment;
+	TObjectPtr<class UMaterialExpressionComment> MaterialExpressionComment;
 
 	/** Marks the Material Editor as dirty so that user prompted to apply change */
 	FSetMaterialDirty MaterialDirtyDelegate;
diff --git a/Engine/Source/Editor/UnrealEd/Classes/MaterialGraph/MaterialGraphNode_Composite.h b/Engine/Source/Editor/UnrealEd/Classes/MaterialGraph/MaterialGraphNode_Composite.h
--- a/Engine/Source/Editor/UnrealEd/Classes/MaterialGraph/MaterialGraphNode_Composite.h
+++ b/Engine/Source/Editor/UnrealEd/Classes/MaterialGraph/MaterialGraphNode_Composite.h
@@ -16,7 +16,7 @@ class UMaterialGraphNode_Composite : public UMaterialGraphNode
 
 	/** The graph that this composite node is representing */
 	UPROPERTY()
-	UMaterialGraph* BoundGraph;
+	TObjectPtr<UMaterialGraph> BoundGraph;
 
 	//~ Begin UObject Interface
 	virtual void PostEditUndo() override; 
diff --git a/Engine/Source/Editor/UnrealEd/Classes/MaterialGraph/MaterialGraphNode_Root.h b/Engine/Source/Editor/UnrealEd/Classes/MaterialGraph/MaterialGraphNode_Root.h
--- a/Engine/Source/Editor/UnrealEd/Classes/MaterialGraph/MaterialGraphNode_Root.h
+++ b/Engine/Source/Editor/UnrealEd/Classes/MaterialGraph/MaterialGraphNode_Root.h
@@ -16,7 +16,7 @@ class UMaterialGraphNode_Root : public UMaterialGraphNode_Base
 
 	/** Material whose inputs this root node represents */
 	UPROPERTY()
-	class UMaterial* Material;
+	TObjectPtr<class UMaterial> Material;
 
 	//~ Begin UEdGraphNode Interface.
 	virtual FText GetNodeTitle(ENodeTitleType::Type TitleType) const override;
diff --git a/Engine/Source/Editor/UnrealEd/Classes/MaterialGraph/MaterialGraphSchema.h b/Engine/Source/Editor/UnrealEd/Classes/MaterialGraph/MaterialGraphSchema.h
--- a/Engine/Source/Editor/UnrealEd/Classes/MaterialGraph/MaterialGraphSchema.h
+++ b/Engine/Source/Editor/UnrealEd/Classes/MaterialGraph/MaterialGraphSchema.h
@@ -18,7 +18,7 @@ struct UNREALED_API FMaterialGraphSchemaAction_NewNode : public FEdGraphSchemaAc
 
 	/** Class of expression we want to create */
 	UPROPERTY()
-	class UClass* MaterialExpressionClass;
+	TObjectPtr<class UClass> MaterialExpressionClass;
 
 	// Simple type info
 	static FName StaticGetTypeId() {static FName Type("FMaterialGraphSchemaAction_NewNode"); return Type;}
@@ -130,7 +130,7 @@ struct UNREALED_API FMaterialGraphSchemaAction_NewNamedRerouteUsage : public FEd
 
 	// Declaration that we want to add an usage of
 	UPROPERTY()
-	class UMaterialExpressionNamedRerouteDeclaration* Declaration = nullptr;
+	TObjectPtr<class UMaterialExpressionNamedRerouteDeclaration> Declaration = nullptr;
 
 	// Simple type info
 	static FName StaticGetTypeId() {static FName Type("FMaterialGraphSchemaAction_NewNamedRerouteUsage"); return Type;}
diff --git a/Engine/Source/Editor/UnrealEd/Classes/Preferences/UnrealEdOptions.h b/Engine/Source/Editor/UnrealEd/Classes/Preferences/UnrealEdOptions.h
--- a/Engine/Source/Editor/UnrealEd/Classes/Preferences/UnrealEdOptions.h
+++ b/Engine/Source/Editor/UnrealEd/Classes/Preferences/UnrealEdOptions.h
@@ -83,7 +83,7 @@ class UNREALED_API UUnrealEdOptions : public UObject
 
 	/** Pointer to the key bindings object that actually stores key bindings for the editor. */
 	UPROPERTY()
-	class UUnrealEdKeyBindings* EditorKeyBindings;
+	TObjectPtr<class UUnrealEdKeyBindings> EditorKeyBindings;
 
 	/** If true, the list of classes in the class picker dialog will be expanded */
 	UPROPERTY(config)
diff --git a/Engine/Source/Editor/UnrealEd/Classes/Settings/EditorExperimentalSettings.h b/Engine/Source/Editor/UnrealEd/Classes/Settings/EditorExperimentalSettings.h
--- a/Engine/Source/Editor/UnrealEd/Classes/Settings/EditorExperimentalSettings.h
+++ b/Engine/Source/Editor/UnrealEd/Classes/Settings/EditorExperimentalSettings.h
@@ -70,7 +70,7 @@ protected:
 	TArray<TSoftClassPtr<UObject>> BaseClassesToAllowRecompilingDuringPlayInEditor;
 
 	UPROPERTY(Transient)
-	mutable TArray<UClass*> ResolvedBaseClassesToAllowRecompilingDuringPlayInEditor;
+	mutable TArray<TObjectPtr<UClass>> ResolvedBaseClassesToAllowRecompilingDuringPlayInEditor;
 
 public:
 	bool IsClassAllowedToRecompileDuringPIE(UClass* TestClass) const;
diff --git a/Engine/Source/Editor/UnrealEd/Classes/ThumbnailRendering/ParticleSystemThumbnailRenderer.h b/Engine/Source/Editor/UnrealEd/Classes/ThumbnailRendering/ParticleSystemThumbnailRenderer.h
--- a/Engine/Source/Editor/UnrealEd/Classes/ThumbnailRendering/ParticleSystemThumbnailRenderer.h
+++ b/Engine/Source/Editor/UnrealEd/Classes/ThumbnailRendering/ParticleSystemThumbnailRenderer.h
@@ -21,10 +21,10 @@ class UParticleSystemThumbnailRenderer : public UTextureThumbnailRenderer
 	GENERATED_UCLASS_BODY()
 
 	UPROPERTY()
-	class UTexture2D* NoImage;
+	TObjectPtr<class UTexture2D> NoImage;
 
 	UPROPERTY()
-	class UTexture2D* OutOfDate;
+	TObjectPtr<class UTexture2D> OutOfDate;
 
 
 	// Begin UThumbnailRenderer Object
diff --git a/Engine/Source/Editor/UnrealEd/Classes/ThumbnailRendering/ThumbnailManager.h b/Engine/Source/Editor/UnrealEd/Classes/ThumbnailRendering/ThumbnailManager.h
--- a/Engine/Source/Editor/UnrealEd/Classes/ThumbnailRendering/ThumbnailManager.h
+++ b/Engine/Source/Editor/UnrealEd/Classes/ThumbnailRendering/ThumbnailManager.h
@@ -68,7 +68,7 @@ struct FThumbnailRenderingInfo
 	 * The instance of the renderer class
 	 */
 	UPROPERTY()
-	class UThumbnailRenderer* Renderer;
+	TObjectPtr<class UThumbnailRenderer> Renderer;
 
 public:
 	FThumbnailRenderingInfo()
@@ -93,28 +93,28 @@ public:
 
 	// All these meshes/materials/textures are preloaded via default properties
 	UPROPERTY(Transient)
-	class UStaticMesh* EditorCube;
+	TObjectPtr<class UStaticMesh> EditorCube;
 
 	UPROPERTY(Transient)
-	class UStaticMesh* EditorSphere;
+	TObjectPtr<class UStaticMesh> EditorSphere;
 
 	UPROPERTY(Transient)
-	class UStaticMesh* EditorCylinder;
+	TObjectPtr<class UStaticMesh> EditorCylinder;
 
 	UPROPERTY(Transient)
-	class UStaticMesh* EditorPlane;
+	TObjectPtr<class UStaticMesh> EditorPlane;
 
 	UPROPERTY(Transient)
-	class UStaticMesh* EditorSkySphere;
+	TObjectPtr<class UStaticMesh> EditorSkySphere;
 
 	UPROPERTY(Transient)
-	class UMaterial* FloorPlaneMaterial;
+	TObjectPtr<class UMaterial> FloorPlaneMaterial;
 
 	UPROPERTY(Transient)
-	class UTextureCube* AmbientCubemap;
+	TObjectPtr<class UTextureCube> AmbientCubemap;
 
 	UPROPERTY(Transient)
-	class UTexture2D* CheckerboardTexture;
+	TObjectPtr<class UTexture2D> CheckerboardTexture;
 
 public:
 	//~ Begin UObject Interface
diff --git a/Engine/Source/Editor/UnrealEd/Classes/ThumbnailRendering/VolumeTextureThumbnailRenderer.h b/Engine/Source/Editor/UnrealEd/Classes/ThumbnailRendering/VolumeTextureThumbnailRenderer.h
--- a/Engine/Source/Editor/UnrealEd/Classes/ThumbnailRendering/VolumeTextureThumbnailRenderer.h
+++ b/Engine/Source/Editor/UnrealEd/Classes/ThumbnailRendering/VolumeTextureThumbnailRenderer.h
@@ -34,6 +34,6 @@ private:
 	class FVolumeTextureThumbnailScene* ThumbnailScene;
 
 	UPROPERTY()
-	UMaterialInstanceConstant* MaterialInstance;
+	TObjectPtr<UMaterialInstanceConstant> MaterialInstance;
 };
 
diff --git a/Engine/Source/Editor/UnrealEd/Public/AssetImportTask.h b/Engine/Source/Editor/UnrealEd/Public/AssetImportTask.h
--- a/Engine/Source/Editor/UnrealEd/Public/AssetImportTask.h
+++ b/Engine/Source/Editor/UnrealEd/Public/AssetImportTask.h
@@ -53,11 +53,11 @@ public:
 
 	/** Optional factory to use */
 	UPROPERTY(BlueprintReadWrite, Category = "Asset Import Task")
-	UFactory* Factory;
+	TObjectPtr<UFactory> Factory;
 
 	/** Import options specific to the type of asset */
 	UPROPERTY(BlueprintReadWrite, Category = "Asset Import Task")
-	UObject* Options;
+	TObjectPtr<UObject> Options;
 
 	/** Paths to objects created or updated after import */
 	UPROPERTY(BlueprintReadWrite, Category = "Asset Import Task")
@@ -65,6 +65,6 @@ public:
 
 	/** Imported objects */
 	UPROPERTY(BlueprintReadWrite, Category = "Asset Import Task")
-	TArray<UObject*> Result;
+	TArray<TObjectPtr<UObject>> Result;
 };
 
diff --git a/Engine/Source/Editor/UnrealEd/Public/AutomatedAssetImportData.h b/Engine/Source/Editor/UnrealEd/Public/AutomatedAssetImportData.h
--- a/Engine/Source/Editor/UnrealEd/Public/AutomatedAssetImportData.h
+++ b/Engine/Source/Editor/UnrealEd/Public/AutomatedAssetImportData.h
@@ -58,7 +58,7 @@ public:
 
 	/** Pointer to the factory currently being sued */
 	UPROPERTY(BlueprintReadWrite, Category = "Asset Import Data")
-	UFactory* Factory;
+	TObjectPtr<UFactory> Factory;
 
 	/** Full path to level to load before importing this group (only matters if importing assets into a level) */
 	UPROPERTY(BlueprintReadWrite, Category = "Asset Import Data")
diff --git a/Engine/Source/Editor/UnrealEd/Public/BaseWidgetBlueprint.h b/Engine/Source/Editor/UnrealEd/Public/BaseWidgetBlueprint.h
--- a/Engine/Source/Editor/UnrealEd/Public/BaseWidgetBlueprint.h
+++ b/Engine/Source/Editor/UnrealEd/Public/BaseWidgetBlueprint.h
@@ -14,7 +14,7 @@ public:
 #if WITH_EDITORONLY_DATA
 	/** A tree of the widget templates to be created */
 	UPROPERTY()
-	class UWidgetTree* WidgetTree;
+	TObjectPtr<class UWidgetTree> WidgetTree;
 #endif
 
 	virtual void PostLoad() override;
diff --git a/Engine/Source/Editor/UnrealEd/Public/Commandlets/ImportAssetsCommandlet.h b/Engine/Source/Editor/UnrealEd/Public/Commandlets/ImportAssetsCommandlet.h
--- a/Engine/Source/Editor/UnrealEd/Public/Commandlets/ImportAssetsCommandlet.h
+++ b/Engine/Source/Editor/UnrealEd/Public/Commandlets/ImportAssetsCommandlet.h
@@ -61,11 +61,11 @@ private:
 	 * or as a base to fallback on if settings are not overridden by the file 
 	 */
 	UPROPERTY()
-	UAutomatedAssetImportData* GlobalImportData;
+	TObjectPtr<UAutomatedAssetImportData> GlobalImportData;
 
 	/** List of import data to import.  Each element in the list represents a list of assets using the same import settings */
 	UPROPERTY()
-	TArray<UAutomatedAssetImportData*> ImportDataList;
+	TArray<TObjectPtr<UAutomatedAssetImportData>> ImportDataList;
 
 	/** If true we allow source control operations during import */
 	bool bAllowSourceControl;
diff --git a/Engine/Source/Editor/UnrealEd/Public/EditorWorldExtension.h b/Engine/Source/Editor/UnrealEd/Public/EditorWorldExtension.h
--- a/Engine/Source/Editor/UnrealEd/Public/EditorWorldExtension.h
+++ b/Engine/Source/Editor/UnrealEd/Public/EditorWorldExtension.h
@@ -32,7 +32,7 @@ struct FEditorWorldExtensionActorData
 	GENERATED_BODY()
 
 	UPROPERTY()
-	AActor* Actor = nullptr;
+	TObjectPtr<AActor> Actor = nullptr;
 
 	UPROPERTY()
 	bool bValidForPIE;
@@ -251,5 +251,5 @@ private:
 
 	/** Map of all the editor world maps */
 	UPROPERTY()
-	TArray<UEditorWorldExtensionCollection*> EditorWorldExtensionCollection;
+	TArray<TObjectPtr<UEditorWorldExtensionCollection>> EditorWorldExtensionCollection;
 };
diff --git a/Engine/Source/Editor/UnrealEd/Public/Selection.h b/Engine/Source/Editor/UnrealEd/Public/Selection.h
--- a/Engine/Source/Editor/UnrealEd/Public/Selection.h
+++ b/Engine/Source/Editor/UnrealEd/Public/Selection.h
@@ -377,7 +377,7 @@ private:
 
 	/** Underlying element selection set (if any). */
 	UPROPERTY()
-	UTypedElementSelectionSet* ElementSelectionSet = nullptr;
+	TObjectPtr<UTypedElementSelectionSet> ElementSelectionSet = nullptr;
 
 	/** Tracks the most recently selected actor classes.  Used for UnrealEd menus. */
 	ClassArray SelectedClasses;
diff --git a/Engine/Source/Editor/UnrealEd/Public/Subsystems/AssetEditorSubsystem.h b/Engine/Source/Editor/UnrealEd/Public/Subsystems/AssetEditorSubsystem.h
--- a/Engine/Source/Editor/UnrealEd/Public/Subsystems/AssetEditorSubsystem.h
+++ b/Engine/Source/Editor/UnrealEd/Public/Subsystems/AssetEditorSubsystem.h
@@ -300,7 +300,7 @@ private:
 	TWeakPtr<SNotificationItem> RestorePreviouslyOpenAssetsNotificationPtr;
 	
 	UPROPERTY(Transient)
-	TArray<UAssetEditor*> OwnedAssetEditors;
+	TArray<TObjectPtr<UAssetEditor>> OwnedAssetEditors;
 	
 	/** Map of FEditorModeId to EditorModeInfo for all known UEdModes when the subsystem initialized */
 	RegisteredModeInfoMap EditorModes;
diff --git a/Engine/Source/Editor/UnrealEd/Public/Tests/FbxAutomationCommon.h b/Engine/Source/Editor/UnrealEd/Public/Tests/FbxAutomationCommon.h
--- a/Engine/Source/Editor/UnrealEd/Public/Tests/FbxAutomationCommon.h
+++ b/Engine/Source/Editor/UnrealEd/Public/Tests/FbxAutomationCommon.h
@@ -150,7 +150,7 @@ class UFbxTestPlan : public UObject
 
 	/* Options use for this test plan, Transient because we manually serialize the options. */
 	UPROPERTY(EditAnywhere, Transient, Instanced, Category = Options)
-	UFbxImportUI* ImportUI;
+	TObjectPtr<UFbxImportUI> ImportUI;
 };
 
 namespace FbxAutomationTestsAPI
diff --git a/Engine/Source/Editor/UnrealEd/Public/Tools/UEdMode.h b/Engine/Source/Editor/UnrealEd/Public/Tools/UEdMode.h
--- a/Engine/Source/Editor/UnrealEd/Public/Tools/UEdMode.h
+++ b/Engine/Source/Editor/UnrealEd/Public/Tools/UEdMode.h
@@ -300,6 +300,6 @@ protected:
 	TSoftClassPtr<UObject> SettingsClass;
 
 	UPROPERTY(Transient)
-	UObject* SettingsObject;
+	TObjectPtr<UObject> SettingsObject;
 
 };
diff --git a/Engine/Source/Editor/VREditor/Public/VREditorAssetContainer.h b/Engine/Source/Editor/VREditor/Public/VREditorAssetContainer.h
--- a/Engine/Source/Editor/VREditor/Public/VREditorAssetContainer.h
+++ b/Engine/Source/Editor/VREditor/Public/VREditorAssetContainer.h
@@ -30,143 +30,143 @@ public:
 	//
 
 	UPROPERTY(EditAnywhere, Category = Sound)
-	USoundBase* DockableWindowCloseSound;
+	TObjectPtr<USoundBase> DockableWindowCloseSound;
 
 	UPROPERTY(EditAnywhere, Category = Sound)
-	USoundBase* DockableWindowOpenSound;
+	TObjectPtr<USoundBase> DockableWindowOpenSound;
 
 	UPROPERTY(EditAnywhere, Category = Sound)
-	USoundBase* DockableWindowDropSound;
+	TObjectPtr<USoundBase> DockableWindowDropSound;
 
 	UPROPERTY(EditAnywhere, Category = Sound)
-	USoundBase* DockableWindowDragSound;
+	TObjectPtr<USoundBase> DockableWindowDragSound;
 
 	UPROPERTY(EditAnywhere, Category = Sound)
-	USoundBase* DropFromContentBrowserSound;
+	TObjectPtr<USoundBase> DropFromContentBrowserSound;
 
 	UPROPERTY(EditAnywhere, Category = Sound)
-	USoundBase* RadialMenuOpenSound;
+	TObjectPtr<USoundBase> RadialMenuOpenSound;
 
 	UPROPERTY(EditAnywhere, Category = Sound)
-	USoundBase* RadialMenuCloseSound;
+	TObjectPtr<USoundBase> RadialMenuCloseSound;
 
 	UPROPERTY(EditAnywhere, Category = Sound)
-	USoundBase* TeleportSound;
+	TObjectPtr<USoundBase> TeleportSound;
 
 	UPROPERTY(EditAnywhere, Category = Sound)
-	USoundCue* ButtonPressSound;
+	TObjectPtr<USoundCue> ButtonPressSound;
 
 	UPROPERTY(EditAnywhere, Category = Sound)
-	USoundBase* AutoScaleSound;
+	TObjectPtr<USoundBase> AutoScaleSound;
 
 	//
 	// Meshes
 	//
 
 	UPROPERTY(EditAnywhere, Category = Mesh)
-	UStaticMesh* GenericHMDMesh;
+	TObjectPtr<UStaticMesh> GenericHMDMesh;
 
 	UPROPERTY(EditAnywhere, Category = Mesh)
-	UStaticMesh* PlaneMesh;
+	TObjectPtr<UStaticMesh> PlaneMesh;
 
 	UPROPERTY(EditAnywhere, Category = Mesh)
-	UStaticMesh* CylinderMesh;
+	TObjectPtr<UStaticMesh> CylinderMesh;
 
 	UPROPERTY(EditAnywhere, Category = Mesh)
-	UStaticMesh* LaserPointerStartMesh;
+	TObjectPtr<UStaticMesh> LaserPointerStartMesh;
 
 	UPROPERTY(EditAnywhere, Category = Mesh)
-	UStaticMesh* LaserPointerMesh;
+	TObjectPtr<UStaticMesh> LaserPointerMesh;
 
 	UPROPERTY(EditAnywhere, Category = Mesh)
-	UStaticMesh* LaserPointerEndMesh;
+	TObjectPtr<UStaticMesh> LaserPointerEndMesh;
 
 	UPROPERTY(EditAnywhere, Category = Mesh)
-	UStaticMesh* LaserPointerHoverMesh;
+	TObjectPtr<UStaticMesh> LaserPointerHoverMesh;
 
 	UPROPERTY(EditAnywhere, Category = Mesh)
-	UStaticMesh* VivePreControllerMesh;
+	TObjectPtr<UStaticMesh> VivePreControllerMesh;
 
 	UPROPERTY(EditAnywhere, Category = Mesh)
-	UStaticMesh* OculusControllerMesh;
+	TObjectPtr<UStaticMesh> OculusControllerMesh;
 
 	UPROPERTY(EditAnywhere, Category = Mesh)
-	UStaticMesh* GenericControllerMesh;
+	TObjectPtr<UStaticMesh> GenericControllerMesh;
 
 	UPROPERTY(EditAnywhere, Category = Mesh)
-	UStaticMesh* TeleportRootMesh;
+	TObjectPtr<UStaticMesh> TeleportRootMesh;
 
 	UPROPERTY(EditAnywhere, Category = Mesh)
-	UStaticMesh* WindowMesh;
+	TObjectPtr<UStaticMesh> WindowMesh;
 
 	UPROPERTY(EditAnywhere, Category = Mesh)
-	UStaticMesh* WindowSelectionBarMesh;
+	TObjectPtr<UStaticMesh> WindowSelectionBarMesh;
 
 	UPROPERTY(EditAnywhere, Category = Mesh)
-	UStaticMesh* WindowCloseButtonMesh;
+	TObjectPtr<UStaticMesh> WindowCloseButtonMesh;
 
 	UPROPERTY(EditAnywhere, Category = Mesh)
-	UStaticMesh* RadialMenuMainMesh;
+	TObjectPtr<UStaticMesh> RadialMenuMainMesh;
 
 	UPROPERTY(EditAnywhere, Category = Mesh)
-	UStaticMesh* RadialMenuPointerMesh;
+	TObjectPtr<UStaticMesh> RadialMenuPointerMesh;
 	
 	UPROPERTY(EditAnywhere, Category = Mesh)
-	UStaticMesh* PointerCursorMesh;
+	TObjectPtr<UStaticMesh> PointerCursorMesh;
 
 	UPROPERTY(EditAnywhere, Category = Mesh)
-	UStaticMesh* LineSegmentCylinderMesh;
+	TObjectPtr<UStaticMesh> LineSegmentCylinderMesh;
 
 	UPROPERTY(EditAnywhere, Category = Mesh)
-	UStaticMesh* JointSphereMesh;
+	TObjectPtr<UStaticMesh> JointSphereMesh;
 
 	UPROPERTY(EditAnywhere, Category = Mesh)
-	UStaticMesh* DockingButtonMesh;
+	TObjectPtr<UStaticMesh> DockingButtonMesh;
 
 	//
 	// Materials
 	//
 
 	UPROPERTY(EditAnywhere, Category = Material)
-	UMaterialInterface* GridMaterial;
+	TObjectPtr<UMaterialInterface> GridMaterial;
 
 	UPROPERTY(EditAnywhere, Category = Material)
-	UMaterialInterface* LaserPointerMaterial;
+	TObjectPtr<UMaterialInterface> LaserPointerMaterial;
 
 	UPROPERTY(EditAnywhere, Category = Material)
-	UMaterialInterface* LaserPointerTranslucentMaterial;
+	TObjectPtr<UMaterialInterface> LaserPointerTranslucentMaterial;
 
 	UPROPERTY(EditAnywhere, Category = Material)
-	UMaterial* WorldMovementPostProcessMaterial;
+	TObjectPtr<UMaterial> WorldMovementPostProcessMaterial;
 
 	UPROPERTY(EditAnywhere, Category = Material)
-	UMaterialInterface* TextMaterial;
+	TObjectPtr<UMaterialInterface> TextMaterial;
 	
 	UPROPERTY(EditAnywhere, Category = Material)
-	UMaterialInterface* VivePreControllerMaterial;
+	TObjectPtr<UMaterialInterface> VivePreControllerMaterial;
 
 	UPROPERTY(EditAnywhere, Category = Material)
-	UMaterialInterface* OculusControllerMaterial;
+	TObjectPtr<UMaterialInterface> OculusControllerMaterial;
 
 	UPROPERTY(EditAnywhere, Category = Material)
-	UMaterialInterface* TeleportMaterial;
+	TObjectPtr<UMaterialInterface> TeleportMaterial;
 
 	UPROPERTY(EditAnywhere, Category = Material)
-	UMaterialInterface* WindowMaterial;
+	TObjectPtr<UMaterialInterface> WindowMaterial;
 
 	UPROPERTY(EditAnywhere, Category = Material)
-	UMaterialInterface* WindowTranslucentMaterial;
+	TObjectPtr<UMaterialInterface> WindowTranslucentMaterial;
 
 	UPROPERTY(EditAnywhere, Category = Material)
-	UMaterial* LineMaterial;
+	TObjectPtr<UMaterial> LineMaterial;
 
 	UPROPERTY(EditAnywhere, Category = Material)
-	UMaterialInterface* TranslucentTextMaterial;
+	TObjectPtr<UMaterialInterface> TranslucentTextMaterial;
 
 	//
 	// Fonts
 	//
 
 	UPROPERTY(EditAnywhere, Category = Font)
-	UFont* TextFont;
+	TObjectPtr<UFont> TextFont;
 };
diff --git a/Engine/Source/Editor/VREditor/Public/VREditorInteractor.h b/Engine/Source/Editor/VREditor/Public/VREditorInteractor.h
--- a/Engine/Source/Editor/VREditor/Public/VREditorInteractor.h
+++ b/Engine/Source/Editor/VREditor/Public/VREditorInteractor.h
@@ -266,7 +266,7 @@ protected:
 
 	/** Motion controller component which handles late-frame transform updates of all parented sub-components */
 	UPROPERTY()
-	class UMotionControllerComponent* MotionControllerComponent;
+	TObjectPtr<class UMotionControllerComponent> MotionControllerComponent;
 
 	//
 	// Graphics
@@ -274,42 +274,42 @@ protected:
 
 	/** Access to the current handmesh. Use ReplaceHandMeshComponent() to update the entire StaticMeshComponent. */
 	UPROPERTY(BlueprintReadWrite, Category = "VREditorInteractor")
-	class UStaticMeshComponent* HandMeshComponent;
+	TObjectPtr<class UStaticMeshComponent> HandMeshComponent;
 
 
 private:
 
 	/** Spline for this hand's laser pointer */
 	UPROPERTY()
-	class USplineComponent* LaserSplineComponent;
+	TObjectPtr<class USplineComponent> LaserSplineComponent;
 
 	/** Spline meshes for curved laser */
 	UPROPERTY()
-	TArray<class USplineMeshComponent*> LaserSplineMeshComponents;
+	TArray<TObjectPtr<class USplineMeshComponent>> LaserSplineMeshComponents;
 
 	/** MID for laser pointer material (opaque parts) */
 	UPROPERTY()
-	class UMaterialInstanceDynamic* LaserPointerMID;
+	TObjectPtr<class UMaterialInstanceDynamic> LaserPointerMID;
 
 	/** MID for laser pointer material (translucent parts) */
 	UPROPERTY()
-	class UMaterialInstanceDynamic* TranslucentLaserPointerMID;
+	TObjectPtr<class UMaterialInstanceDynamic> TranslucentLaserPointerMID;
 
 	/** Hover impact indicator mesh */
 	UPROPERTY()
-	class UStaticMeshComponent* HoverMeshComponent;
+	TObjectPtr<class UStaticMeshComponent> HoverMeshComponent;
 
 	/** Hover point light */
 	UPROPERTY()
-	class UPointLightComponent* HoverPointLightComponent;
+	TObjectPtr<class UPointLightComponent> HoverPointLightComponent;
 
 	/** MID for hand mesh */
 	UPROPERTY()
-	class UMaterialInstanceDynamic* HandMeshMID;
+	TObjectPtr<class UMaterialInstanceDynamic> HandMeshMID;
 
 	/** Avatar actor that we're attached to. Cached so that we can change the HandMeshComponent via BP at runtime. */
 	UPROPERTY()
-	AActor* OwningAvatar;
+	TObjectPtr<AActor> OwningAvatar;
 
 	/** True if this hand has a motion controller (or both!) */
 	bool bHaveMotionController;
@@ -487,5 +487,5 @@ protected:
 
 	/** The mode that owns this interactor */
 	UPROPERTY()
-	class UVREditorMode* VRMode;
+	TObjectPtr<class UVREditorMode> VRMode;
 };
diff --git a/Engine/Source/Editor/VREditor/Public/VREditorMode.h b/Engine/Source/Editor/VREditor/Public/VREditorMode.h
--- a/Engine/Source/Editor/VREditor/Public/VREditorMode.h
+++ b/Engine/Source/Editor/VREditor/Public/VREditorMode.h
@@ -367,7 +367,7 @@ protected:
 
 	/** Actor with components to represent the VR avatar in the world, including motion controller meshes */
 	UPROPERTY()
-	class AVREditorAvatarActor* AvatarActor;
+	TObjectPtr<class AVREditorAvatarActor> AvatarActor;
 
 
 	//
@@ -406,15 +406,15 @@ protected:
 
 	/** VR UI system */
 	UPROPERTY()
-	class UVREditorUISystem* UISystem;
+	TObjectPtr<class UVREditorUISystem> UISystem;
 
 	/** Teleporter system */
 	UPROPERTY()
-	class AVREditorTeleporter* TeleportActor;
+	TObjectPtr<class AVREditorTeleporter> TeleportActor;
 
 	/** Automatic scale system */
 	UPROPERTY()
-	class UVREditorAutoScaler* AutoScalerSystem;
+	TObjectPtr<class UVREditorAutoScaler> AutoScalerSystem;
 
 	//
 	// World interaction
@@ -422,20 +422,20 @@ protected:
 
 	/** World interaction manager */
 	UPROPERTY()
-	class UViewportWorldInteraction* WorldInteraction;
+	TObjectPtr<class UViewportWorldInteraction> WorldInteraction;
 
 	/** The current Gizmo type that is used for the TransformGizmo Actor */
 	EGizmoHandleTypes CurrentGizmoType;
 
 	UPROPERTY()
-	class UVREditorPlacement* PlacementSystem;
+	TObjectPtr<class UVREditorPlacement> PlacementSystem;
 
 	//
 	// Interactors
 	//
 
 	UPROPERTY()
-	TArray<UVREditorInteractor*> Interactors;
+	TArray<TObjectPtr<UVREditorInteractor>> Interactors;
 
 	//
 	// Colors
@@ -493,7 +493,7 @@ private:
 
 	/** Container of assets */
 	UPROPERTY()
-	class UVREditorAssetContainer* AssetContainer;
+	TObjectPtr<class UVREditorAssetContainer> AssetContainer;
 
 	/** Whether currently in debug mode or not. */
 	static bool bDebugModeEnabled;
diff --git a/Engine/Source/Editor/VREditor/Teleporter/VREditorAutoScaler.h b/Engine/Source/Editor/VREditor/Teleporter/VREditorAutoScaler.h
--- a/Engine/Source/Editor/VREditor/Teleporter/VREditorAutoScaler.h
+++ b/Engine/Source/Editor/VREditor/Teleporter/VREditorAutoScaler.h
@@ -39,5 +39,5 @@ private:
 
 	/** Owning mode */
 	UPROPERTY()
-	UVREditorMode* VRMode;
+	TObjectPtr<UVREditorMode> VRMode;
 };
diff --git a/Engine/Source/Editor/VREditor/Teleporter/VREditorTeleporter.h b/Engine/Source/Editor/VREditor/Teleporter/VREditorTeleporter.h
--- a/Engine/Source/Editor/VREditor/Teleporter/VREditorTeleporter.h
+++ b/Engine/Source/Editor/VREditor/Teleporter/VREditorTeleporter.h
@@ -120,7 +120,7 @@ private:
 
 	/** The owning VR mode */
 	UPROPERTY()
-	UVREditorMode* VRMode;
+	TObjectPtr<UVREditorMode> VRMode;
 
 	enum EState
 	{
@@ -143,27 +143,27 @@ private:
 
 	/** Visuals for the feet location of the teleporter with the same direction of the HMD yaw */
 	UPROPERTY()
-	UStaticMeshComponent* TeleportDirectionMeshComponent;
+	TObjectPtr<UStaticMeshComponent> TeleportDirectionMeshComponent;
 
 	/** Visuals for teleport HMD */
 	UPROPERTY()
-	UStaticMeshComponent* HMDMeshComponent;
+	TObjectPtr<UStaticMeshComponent> HMDMeshComponent;
 
 	/** Visuals for teleport left motion controller */
 	UPROPERTY()
-	UStaticMeshComponent* LeftMotionControllerMeshComponent;
+	TObjectPtr<UStaticMeshComponent> LeftMotionControllerMeshComponent;
 
 	/** Visuals for teleport right motion controller */
 	UPROPERTY()
-	UStaticMeshComponent* RightMotionControllerMeshComponent;
+	TObjectPtr<UStaticMeshComponent> RightMotionControllerMeshComponent;
 
 	/** Dynamic material for teleport visuals */
 	UPROPERTY()
-	UMaterialInstanceDynamic* TeleportMID;
+	TObjectPtr<UMaterialInstanceDynamic> TeleportMID;
 
 	/** The interactor that started aiming to teleport */
 	UPROPERTY()
-	UViewportInteractor* InteractorTryingTeleport;
+	TObjectPtr<UViewportInteractor> InteractorTryingTeleport;
 
 	/** When offset between the hoverlocation of the laser and the calculated  teleport*/
 	FVector OffsetDistance;
diff --git a/Engine/Source/Editor/VREditor/UI/VREditorDockableWindow.h b/Engine/Source/Editor/VREditor/UI/VREditorDockableWindow.h
--- a/Engine/Source/Editor/VREditor/UI/VREditorDockableWindow.h
+++ b/Engine/Source/Editor/VREditor/UI/VREditorDockableWindow.h
@@ -79,38 +79,38 @@ private:
 
 	/** Mesh underneath the window for docking it to the world */
 	UPROPERTY()
-	class UStaticMeshComponent* DockButtonMeshComponent;
+	TObjectPtr<class UStaticMeshComponent> DockButtonMeshComponent;
 
 	/** Mesh underneath the window for easy selecting and dragging */
 	UPROPERTY()
-	class UStaticMeshComponent* SelectionBarMeshComponent;
+	TObjectPtr<class UStaticMeshComponent> SelectionBarMeshComponent;
 
 	/** Mesh that represents the close button for this UI */
 	UPROPERTY()
-	class UStaticMeshComponent* CloseButtonMeshComponent;
+	TObjectPtr<class UStaticMeshComponent> CloseButtonMeshComponent;
 
 	/** Dock Button dynamic material  (opaque) */
 	UPROPERTY()
-	class UMaterialInstanceDynamic* DockButtonMID;
+	TObjectPtr<class UMaterialInstanceDynamic> DockButtonMID;
 
 	/** Selection bar dynamic material  (opaque) */
 	UPROPERTY()
-	class UMaterialInstanceDynamic* SelectionBarMID;
+	TObjectPtr<class UMaterialInstanceDynamic> SelectionBarMID;
 
 	/** Select bar dynamic material (translucent) */
 	UPROPERTY()
-	class UMaterialInstanceDynamic* SelectionBarTranslucentMID;
+	TObjectPtr<class UMaterialInstanceDynamic> SelectionBarTranslucentMID;
 
 	/** Close button dynamic material  (opaque) */
 	UPROPERTY()
-	class UMaterialInstanceDynamic* CloseButtonMID;
+	TObjectPtr<class UMaterialInstanceDynamic> CloseButtonMID;
 
 	/** Close button dynamic material (translucent) */
 	UPROPERTY()
-	class UMaterialInstanceDynamic* CloseButtonTranslucentMID;
+	TObjectPtr<class UMaterialInstanceDynamic> CloseButtonTranslucentMID;
 
 	UPROPERTY()
-	class UViewportDragOperationComponent* DragOperationComponent;
+	TObjectPtr<class UViewportDragOperationComponent> DragOperationComponent;
 
 	/** True if at least one hand's laser is aiming toward the UI */
 	bool bIsLaserAimingTowardUI;
diff --git a/Engine/Source/Editor/VREditor/UI/VREditorFloatingUI.h b/Engine/Source/Editor/VREditor/UI/VREditorFloatingUI.h
--- a/Engine/Source/Editor/VREditor/UI/VREditorFloatingUI.h
+++ b/Engine/Source/Editor/VREditor/UI/VREditorFloatingUI.h
@@ -34,7 +34,7 @@ public:
 	FName PanelID;
 
 	UPROPERTY(EditAnywhere, BlueprintReadWrite, Category = "Virtual Production UI")
-	AActor* ParentActor = nullptr;
+	TObjectPtr<AActor> ParentActor = nullptr;
 
 	/** Optional offset from HMD where the window opens. Pass FTransform::Identity for default logic - window will open at controller location. */
 	UPROPERTY(EditAnywhere, BlueprintReadWrite, Category = "VR Mode UI")
@@ -46,7 +46,7 @@ public:
 
 	/** Optional custom mesh to use for the VR window. */
 	UPROPERTY(EditAnywhere, BlueprintReadWrite, Category = "VR Mode UI")
-	UStaticMesh* PanelMesh = nullptr;
+	TObjectPtr<UStaticMesh> PanelMesh = nullptr;
 
 	/** Optional override for "VREd.EditorUISize". Leave at 0 for default. */
 	UPROPERTY(EditAnywhere, BlueprintReadWrite, Category = "VR Mode UI")
@@ -201,15 +201,15 @@ protected:
 	
 	/** UMG user widget we're drawing, or nullptr if we're drawing a Slate widget */
 	UPROPERTY()
-	UUserWidget* UserWidget;
+	TObjectPtr<UUserWidget> UserWidget;
 	
 	/** When in a spawned state, this is the widget component to represent the widget */
 	UPROPERTY()
-	class UVREditorWidgetComponent* WidgetComponent;
+	TObjectPtr<class UVREditorWidgetComponent> WidgetComponent;
 
 	/** The floating window mesh */
 	UPROPERTY()
-	class UStaticMeshComponent* WindowMeshComponent;
+	TObjectPtr<class UStaticMeshComponent> WindowMeshComponent;
 
 	/** Resolution we should draw this UI at, regardless of scale */
 	FIntPoint Resolution;
@@ -223,7 +223,7 @@ private:
 
 	/** Class of the UMG widget we're spawning */
 	UPROPERTY()
-	UClass* UserWidgetClass;
+	TObjectPtr<UClass> UserWidgetClass;
 
 	/** True if the UI wants to be visible, or false if it wants to be hidden.  Remember, we might still be visually 
 	    transitioning between states */
diff --git a/Engine/Source/Editor/VREditor/UI/VREditorRadialFloatingUI.h b/Engine/Source/Editor/VREditor/UI/VREditorRadialFloatingUI.h
--- a/Engine/Source/Editor/VREditor/UI/VREditorRadialFloatingUI.h
+++ b/Engine/Source/Editor/VREditor/UI/VREditorRadialFloatingUI.h
@@ -155,19 +155,19 @@ protected:
 	
 	/** When in a spawned state, this is the widget component to represent the widget */
 	UPROPERTY()
-	TArray<class UVREditorWidgetComponent*> WidgetComponents;
+	TArray<TObjectPtr<class UVREditorWidgetComponent>> WidgetComponents;
 
 	/** The floating window mesh */
 	UPROPERTY()
-	class UStaticMeshComponent* WindowMeshComponent;
+	TObjectPtr<class UStaticMeshComponent> WindowMeshComponent;
 
 	/** The arrow indicator mesh */
 	UPROPERTY()
-	class UStaticMeshComponent* ArrowMeshComponent;
+	TObjectPtr<class UStaticMeshComponent> ArrowMeshComponent;
 
 	/** The central helper widget */
 	UPROPERTY()
-	class UVREditorWidgetComponent* CentralWidgetComponent;
+	TObjectPtr<class UVREditorWidgetComponent> CentralWidgetComponent;
 
 	TSharedPtr<SWidget> CentralSlateWidget;
 
diff --git a/Engine/Source/Editor/VREditor/UI/VREditorUISystem.h b/Engine/Source/Editor/VREditor/UI/VREditorUISystem.h
--- a/Engine/Source/Editor/VREditor/UI/VREditorUISystem.h
+++ b/Engine/Source/Editor/VREditor/UI/VREditorUISystem.h
@@ -50,7 +50,7 @@ struct FVRButton
 
 	/** Pointer to button */
 	UPROPERTY()
-	UVREditorWidgetComponent* ButtonWidget;
+	TObjectPtr<UVREditorWidgetComponent> ButtonWidget;
 
 	/** Animation playback state of the button */
 	EVREditorAnimationState AnimationDirection;
@@ -335,19 +335,19 @@ protected:
 
 	/** Owning object */
 	UPROPERTY()
-	class UVREditorMode* VRMode;
+	TObjectPtr<class UVREditorMode> VRMode;
 
 	/** All of the floating UIs.  These may or may not be visible (spawned) */
 	UPROPERTY()
-	TMap<FName, class AVREditorFloatingUI*> FloatingUIs;
+	TMap<FName, TObjectPtr<class AVREditorFloatingUI>> FloatingUIs;
 
 	/** All of the preview window info. */
 	UPROPERTY()
-	TMap<FName, class AActor*> PreviewWindowInfo;
+	TMap<FName, TObjectPtr<class AActor>> PreviewWindowInfo;
 
 	/** Our Quick Menu UI */
 	UPROPERTY()
-	class AVREditorFloatingUI* InfoDisplayPanel;
+	TObjectPtr<class AVREditorFloatingUI> InfoDisplayPanel;
 
 	/** 
 	 * The current widget used on the info display. Often we wrap a widget in a widget to configure the settings (e.g. DPI). 
@@ -357,7 +357,7 @@ protected:
 
 	/** The Radial Menu UI */
 	UPROPERTY()
-	class AVREditorRadialFloatingUI* QuickRadialMenu;
+	TObjectPtr<class AVREditorRadialFloatingUI> QuickRadialMenu;
 
 	/** The time since the radial menu was updated */
 	float RadialMenuHideDelayTime;
@@ -380,21 +380,21 @@ protected:
 
 	/** The current UI that is being dragged */
 	UPROPERTY()
-	class AVREditorDockableWindow* DraggingUI;
+	TObjectPtr<class AVREditorDockableWindow> DraggingUI;
 
 	/** The color picker dockable window */
 	UPROPERTY()
-	class AVREditorDockableWindow* ColorPickerUI;
+	TObjectPtr<class AVREditorDockableWindow> ColorPickerUI;
 
 	//
 	// Asymmetry
 	//
 	/** Interactor that has a laser and is generally interacting with the scene */
 	UPROPERTY()
-	class UVREditorInteractor* LaserInteractor;
+	TObjectPtr<class UVREditorInteractor> LaserInteractor;
 	/** Interactor that usually accesses UI and other helper functionality */
 	UPROPERTY()
-	class UVREditorInteractor* UIInteractor;
+	TObjectPtr<class UVREditorInteractor> UIInteractor;
 
 	//
 	// Tab Manager UI
@@ -415,7 +415,7 @@ protected:
 
 	/** The add-on that handles radial menu switching */
 	UPROPERTY()
-	UVRRadialMenuHandler* RadialMenuHandler;
+	TObjectPtr<UVRRadialMenuHandler> RadialMenuHandler;
 
 	/** When replacing the actions menu, store off any existing actions */
 	FOnRadialMenuGenerated ExistingActionsMenu;
diff --git a/Engine/Source/Editor/VREditor/VREditorAvatarActor.h b/Engine/Source/Editor/VREditor/VREditorAvatarActor.h
--- a/Engine/Source/Editor/VREditor/VREditorAvatarActor.h
+++ b/Engine/Source/Editor/VREditor/VREditorAvatarActor.h
@@ -45,11 +45,11 @@ private:
 
 	/** Our avatar's head mesh */
 	UPROPERTY()
-	class UStaticMeshComponent* HeadMeshComponent;
+	TObjectPtr<class UStaticMeshComponent> HeadMeshComponent;
 
 	/** The grid that appears while the user is dragging the world around */
 	UPROPERTY()
-	class UStaticMeshComponent* WorldMovementGridMeshComponent;
+	TObjectPtr<class UStaticMeshComponent> WorldMovementGridMeshComponent;
 
 	//
 	// World movement grid & FX
@@ -57,7 +57,7 @@ private:
 
 	/** Grid mesh component dynamic material instance to set the opacity */
 	UPROPERTY()
-	class UMaterialInstanceDynamic* WorldMovementGridMID;
+	TObjectPtr<class UMaterialInstanceDynamic> WorldMovementGridMID;
 
 	/** Opacity of the movement grid and post process */
 	UPROPERTY()
@@ -69,7 +69,7 @@ private:
 
 	/** Post process material for "greying out" the world while in world movement mode */
 	UPROPERTY()
-	class UMaterialInstanceDynamic* WorldMovementPostProcessMaterial;
+	TObjectPtr<class UMaterialInstanceDynamic> WorldMovementPostProcessMaterial;
 
 	//
 	// World scaling progress bar
@@ -77,31 +77,31 @@ private:
 
 	/** Background progressbar scaling mesh */
 	UPROPERTY()
-	class UStaticMeshComponent* ScaleProgressMeshComponent;
+	TObjectPtr<class UStaticMeshComponent> ScaleProgressMeshComponent;
 
 	/** Current scale progressbar mesh */
 	UPROPERTY()
-	class UStaticMeshComponent* CurrentScaleProgressMeshComponent;
+	TObjectPtr<class UStaticMeshComponent> CurrentScaleProgressMeshComponent;
 
 	/** Current scale text */
 	UPROPERTY()
-	class UTextRenderComponent* UserScaleIndicatorText;
+	TObjectPtr<class UTextRenderComponent> UserScaleIndicatorText;
 
 	/** Base dynamic material for the user scale fixed progressbar */
 	UPROPERTY()
-	UMaterialInstanceDynamic* FixedUserScaleMID;
+	TObjectPtr<UMaterialInstanceDynamic> FixedUserScaleMID;
 
 	/** Translucent dynamic material for the user scale fixed progressbar */
 	UPROPERTY()
-	UMaterialInstanceDynamic* TranslucentFixedUserScaleMID;
+	TObjectPtr<UMaterialInstanceDynamic> TranslucentFixedUserScaleMID;
 	
 	/** Base dynamic material for the current user scale progressbar */
 	UPROPERTY()
-	UMaterialInstanceDynamic* CurrentUserScaleMID;
+	TObjectPtr<UMaterialInstanceDynamic> CurrentUserScaleMID;
 
 	/** Translucent dynamic material for the current user scale progressbar */
 	UPROPERTY()
-	UMaterialInstanceDynamic* TranslucentCurrentUserScaleMID;
+	TObjectPtr<UMaterialInstanceDynamic> TranslucentCurrentUserScaleMID;
 
 	//
 	// Post process
@@ -109,9 +109,9 @@ private:
 
 	/** Post process for drawing VR-specific post effects */
 	UPROPERTY()
-	class UPostProcessComponent* PostProcessComponent;
+	TObjectPtr<class UPostProcessComponent> PostProcessComponent;
 
 	/** Owning object */
 	UPROPERTY()
-	class UVREditorMode* VRMode;
+	TObjectPtr<class UVREditorMode> VRMode;
 };
diff --git a/Engine/Source/Editor/VREditor/VREditorBaseActor.h b/Engine/Source/Editor/VREditor/VREditorBaseActor.h
--- a/Engine/Source/Editor/VREditor/VREditorBaseActor.h
+++ b/Engine/Source/Editor/VREditor/VREditorBaseActor.h
@@ -101,7 +101,7 @@ protected:
 
 	/** The VR mode that owns this actor */
 	UPROPERTY()
-	class UVREditorMode* VRMode;
+	TObjectPtr<class UVREditorMode> VRMode;
 
 private:
 	/** Local rotation of the UI */
diff --git a/Engine/Source/Editor/VREditor/VREditorFloatingText.h b/Engine/Source/Editor/VREditor/VREditorFloatingText.h
--- a/Engine/Source/Editor/VREditor/VREditorFloatingText.h
+++ b/Engine/Source/Editor/VREditor/VREditorFloatingText.h
@@ -41,38 +41,38 @@ private:
 
 	/** Scene component root of this actor */
 	UPROPERTY()
-	class USceneComponent* SceneComponent;
+	TObjectPtr<class USceneComponent> SceneComponent;
 
 	/** First line segment component.  Starts at the designation location, goes toward the line connection point. */
 	UPROPERTY()
-	class UStaticMeshComponent* FirstLineComponent;
+	TObjectPtr<class UStaticMeshComponent> FirstLineComponent;
 
 	/** Sphere that connects the two line segments and makes the joint look smooth and round */
 	UPROPERTY()
-	class UStaticMeshComponent* JointSphereComponent;
+	TObjectPtr<class UStaticMeshComponent> JointSphereComponent;
 
 	/** Second line segment component.  Starts at the connection point and goes toward the 3D text. */
 	UPROPERTY()
-	class UStaticMeshComponent* SecondLineComponent;
+	TObjectPtr<class UStaticMeshComponent> SecondLineComponent;
 
 	/** The 3D text we're drawing.  Positioned at the end of the second line. */
 	UPROPERTY()
-	class UTextRenderComponent* TextComponent;
+	TObjectPtr<class UTextRenderComponent> TextComponent;
 
 	/** Masked text material.  Used after faded in */
 	UPROPERTY()
-	class UMaterialInterface* MaskedTextMaterial;
+	TObjectPtr<class UMaterialInterface> MaskedTextMaterial;
 
 	/** Translucent text material.  Used during fading */
 	UPROPERTY()
-	class UMaterialInterface* TranslucentTextMaterial;
+	TObjectPtr<class UMaterialInterface> TranslucentTextMaterial;
 
 	/** Material to use for the line meshes */
 	UPROPERTY()
-	class UMaterialInterface* LineMaterial;
+	TObjectPtr<class UMaterialInterface> LineMaterial;
 
 	/** Dynamic material instance for fading lines in and out */
 	UPROPERTY( transient )
-	class UMaterialInstanceDynamic* LineMaterialMID;
+	TObjectPtr<class UMaterialInstanceDynamic> LineMaterialMID;
 
 };
diff --git a/Engine/Source/Editor/VREditor/VREditorPlacement.h b/Engine/Source/Editor/VREditor/VREditorPlacement.h
--- a/Engine/Source/Editor/VREditor/VREditorPlacement.h
+++ b/Engine/Source/Editor/VREditor/VREditorPlacement.h
@@ -45,11 +45,11 @@ protected:
 
 	/** Owning object */
 	UPROPERTY()
-	UVREditorMode* VRMode;
+	TObjectPtr<UVREditorMode> VRMode;
 
 	/** The actual ViewportWorldInteraction */
 	UPROPERTY()
-	UViewportWorldInteraction* ViewportWorldInteraction;
+	TObjectPtr<UViewportWorldInteraction> ViewportWorldInteraction;
 
 	//
 	// Dragging object from UI
@@ -57,10 +57,10 @@ protected:
 
 	/** The UI used to drag an asset into the level */
 	UPROPERTY()
-	class UWidgetComponent* FloatingUIAssetDraggedFrom;
+	TObjectPtr<class UWidgetComponent> FloatingUIAssetDraggedFrom;
 
 	/** The material or texture asset we're dragging to place on an object */
 	UPROPERTY()
-	UObject* PlacingMaterialOrTextureAsset;
+	TObjectPtr<UObject> PlacingMaterialOrTextureAsset;
 };
 
diff --git a/Engine/Source/Editor/VREditor/VRScoutingInteractor.h b/Engine/Source/Editor/VREditor/VRScoutingInteractor.h
--- a/Engine/Source/Editor/VREditor/VRScoutingInteractor.h
+++ b/Engine/Source/Editor/VREditor/VRScoutingInteractor.h
@@ -49,5 +49,5 @@ public:
 
 	/** Shown in Navigation mode */
 	UPROPERTY(Category = Interactor, EditAnywhere, BlueprintReadOnly)
-	class UStaticMeshComponent* FlyingIndicatorComponent;
+	TObjectPtr<class UStaticMeshComponent> FlyingIndicatorComponent;
 };
diff --git a/Engine/Source/Editor/ViewportInteraction/Gizmo/VIPivotTransformGizmo.h b/Engine/Source/Editor/ViewportInteraction/Gizmo/VIPivotTransformGizmo.h
--- a/Engine/Source/Editor/ViewportInteraction/Gizmo/VIPivotTransformGizmo.h
+++ b/Engine/Source/Editor/ViewportInteraction/Gizmo/VIPivotTransformGizmo.h
@@ -31,34 +31,34 @@ public:
 private:
 	/** Uniform scale handle group component */
 	UPROPERTY()
-	class UUniformScaleGizmoHandleGroup* UniformScaleGizmoHandleGroup;
+	TObjectPtr<class UUniformScaleGizmoHandleGroup> UniformScaleGizmoHandleGroup;
 
 	/** Translation handle group component */
 	UPROPERTY()
-	class UPivotTranslationGizmoHandleGroup* TranslationGizmoHandleGroup;
+	TObjectPtr<class UPivotTranslationGizmoHandleGroup> TranslationGizmoHandleGroup;
 	
 	/** Scale handle group component */
 	UPROPERTY()
-	class UPivotScaleGizmoHandleGroup* ScaleGizmoHandleGroup;
+	TObjectPtr<class UPivotScaleGizmoHandleGroup> ScaleGizmoHandleGroup;
 	
 	/** Plane translation handle group component */
 	UPROPERTY()
-	class UPivotPlaneTranslationGizmoHandleGroup* PlaneTranslationGizmoHandleGroup;
+	TObjectPtr<class UPivotPlaneTranslationGizmoHandleGroup> PlaneTranslationGizmoHandleGroup;
 
 	/** Rotation handle group component */
 	UPROPERTY()
-	class UPivotRotationGizmoHandleGroup* RotationGizmoHandleGroup;
+	TObjectPtr<class UPivotRotationGizmoHandleGroup> RotationGizmoHandleGroup;
 
 	/** Stretch handle group component */
 	UPROPERTY()
-	class UStretchGizmoHandleGroup* StretchGizmoHandleGroup;
+	TObjectPtr<class UStretchGizmoHandleGroup> StretchGizmoHandleGroup;
 
 	/** The alpha for gizmo animation when aiming at it with a laser */
 	float AimingAtGizmoScaleAlpha;
 
 	/** Handle from previous tick that was dragged */
 	UPROPERTY()
-	UActorComponent* LastDraggingHandle;
+	TObjectPtr<UActorComponent> LastDraggingHandle;
 };
 
 /**
@@ -162,27 +162,27 @@ private:
 
 	/** Root component of all the mesh components that are used to visualize the rotation when dragging */
 	UPROPERTY()
-	USceneComponent* RootFullRotationHandleComponent;
+	TObjectPtr<USceneComponent> RootFullRotationHandleComponent;
 
 	/** When dragging a rotation handle the full rotation circle appears */
 	UPROPERTY()
-	class UGizmoHandleMeshComponent* FullRotationHandleMeshComponent;
+	TObjectPtr<class UGizmoHandleMeshComponent> FullRotationHandleMeshComponent;
 
 	/** The mesh that indicated the start rotation */
 	UPROPERTY()
-	class UGizmoHandleMeshComponent* StartRotationIndicatorMeshComponent;
+	TObjectPtr<class UGizmoHandleMeshComponent> StartRotationIndicatorMeshComponent;
 
 	/** The root component of the start rotation indicator */
 	UPROPERTY()
-	USceneComponent* RootStartRotationIdicatorComponent;
+	TObjectPtr<USceneComponent> RootStartRotationIdicatorComponent;
 
 	/** The mesh that indicated the delta rotation */
 	UPROPERTY()
-	class UGizmoHandleMeshComponent* DeltaRotationIndicatorMeshComponent;
+	TObjectPtr<class UGizmoHandleMeshComponent> DeltaRotationIndicatorMeshComponent;
 
 	/** The root component of the delta rotation indicator */
 	UPROPERTY()
-	USceneComponent* RootDeltaRotationIndicatorComponent;
+	TObjectPtr<USceneComponent> RootDeltaRotationIndicatorComponent;
 
 	/** The rotation when starting to drag the gizmo */
 	TOptional<FQuat> StartDragRotation;
diff --git a/Engine/Source/Editor/ViewportInteraction/Public/VIBaseTransformGizmo.h b/Engine/Source/Editor/ViewportInteraction/Public/VIBaseTransformGizmo.h
--- a/Engine/Source/Editor/ViewportInteraction/Public/VIBaseTransformGizmo.h
+++ b/Engine/Source/Editor/ViewportInteraction/Public/VIBaseTransformGizmo.h
@@ -30,7 +30,7 @@ struct FTransformGizmoMeasurement
 
 	/** The text that displays the actual measurement and units */
 	UPROPERTY()
-	class UTextRenderComponent* MeasurementText = nullptr;
+	TObjectPtr<class UTextRenderComponent> MeasurementText = nullptr;
 };
 
 
@@ -97,15 +97,15 @@ protected:
 	
 	/** Scene component root of this actor */
 	UPROPERTY()
-	USceneComponent* SceneComponent;
+	TObjectPtr<USceneComponent> SceneComponent;
 
 	/** All gizmo components */
 	UPROPERTY()
-	TArray< class UGizmoHandleGroup* > AllHandleGroups;
+	TArray< TObjectPtr<class UGizmoHandleGroup> > AllHandleGroups;
 
 	/** Owning object */
 	UPROPERTY()
-	class UViewportWorldInteraction* WorldInteraction;
+	TObjectPtr<class UViewportWorldInteraction> WorldInteraction;
 	
 	/** Current gizmo type */
 	EGizmoHandleTypes GizmoType;
diff --git a/Engine/Source/Editor/ViewportInteraction/Public/VIGizmoHandle.h b/Engine/Source/Editor/ViewportInteraction/Public/VIGizmoHandle.h
--- a/Engine/Source/Editor/ViewportInteraction/Public/VIGizmoHandle.h
+++ b/Engine/Source/Editor/ViewportInteraction/Public/VIGizmoHandle.h
@@ -109,11 +109,11 @@ protected:
 
 	/** Gizmo material (opaque) */
 	UPROPERTY()
-	UMaterialInterface* GizmoMaterial;
+	TObjectPtr<UMaterialInterface> GizmoMaterial;
 
 	/** Gizmo material (translucent) */
 	UPROPERTY()
-	UMaterialInterface* TranslucentGizmoMaterial;
+	TObjectPtr<UMaterialInterface> TranslucentGizmoMaterial;
 	
 	/** All the StaticMeshes for this handle type */
 	UPROPERTY()
@@ -121,10 +121,10 @@ protected:
 
 	/** The actor transform gizmo owning this handlegroup */
 	UPROPERTY()
-	class ABaseTransformGizmo* OwningTransformGizmoActor;
+	TObjectPtr<class ABaseTransformGizmo> OwningTransformGizmoActor;
 
 	UPROPERTY()
-	class UViewportDragOperationComponent* DragOperationComponent;
+	TObjectPtr<class UViewportDragOperationComponent> DragOperationComponent;
 
 private:
 
diff --git a/Engine/Source/Editor/ViewportInteraction/Public/ViewportDragOperation.h b/Engine/Source/Editor/ViewportInteraction/Public/ViewportDragOperation.h
--- a/Engine/Source/Editor/ViewportInteraction/Public/ViewportDragOperation.h
+++ b/Engine/Source/Editor/ViewportInteraction/Public/ViewportDragOperation.h
@@ -74,7 +74,7 @@ private:
 
 	/** The actual dragging implementation */
 	UPROPERTY()
-	UViewportDragOperation* DragOperation;
+	TObjectPtr<UViewportDragOperation> DragOperation;
 
 	/** The next class that will be used as drag operation */
 	UPROPERTY()
diff --git a/Engine/Source/Editor/ViewportInteraction/Public/ViewportInteractionAssetContainer.h b/Engine/Source/Editor/ViewportInteraction/Public/ViewportInteractionAssetContainer.h
--- a/Engine/Source/Editor/ViewportInteraction/Public/ViewportInteractionAssetContainer.h
+++ b/Engine/Source/Editor/ViewportInteraction/Public/ViewportInteractionAssetContainer.h
@@ -25,76 +25,76 @@ public:
 	//
 
 	UPROPERTY(EditAnywhere, Category = Sound)
-	USoundBase* GizmoHandleSelectedSound;
+	TObjectPtr<USoundBase> GizmoHandleSelectedSound;
 
 	UPROPERTY(EditAnywhere, Category = Sound)
-	USoundBase* GizmoHandleDropSound;
+	TObjectPtr<USoundBase> GizmoHandleDropSound;
 
 	UPROPERTY(EditAnywhere, Category = Sound)
-	USoundBase* SelectionChangeSound;
+	TObjectPtr<USoundBase> SelectionChangeSound;
 
 	UPROPERTY(EditAnywhere, Category = Sound)
-	USoundBase* SelectionDropSound;
+	TObjectPtr<USoundBase> SelectionDropSound;
 
 	UPROPERTY(EditAnywhere, Category = Sound)
-	USoundBase* SelectionStartDragSound;
+	TObjectPtr<USoundBase> SelectionStartDragSound;
 
 	UPROPERTY(EditAnywhere, Category = Sound)
-	USoundBase* GridSnapSound;
+	TObjectPtr<USoundBase> GridSnapSound;
 
 	UPROPERTY(EditAnywhere, Category = Sound)
-	USoundBase* ActorSnapSound;
+	TObjectPtr<USoundBase> ActorSnapSound;
 
 	UPROPERTY(EditAnywhere, Category = Sound)
-	USoundBase* UndoSound;
+	TObjectPtr<USoundBase> UndoSound;
 
 	UPROPERTY(EditAnywhere, Category = Sound)
-	USoundBase* RedoSound;
+	TObjectPtr<USoundBase> RedoSound;
 
 	//
 	// Meshes
 	//
 
 	UPROPERTY(EditAnywhere, Category = Mesh)
-	UStaticMesh* GridMesh;
+	TObjectPtr<UStaticMesh> GridMesh;
 
 	UPROPERTY(EditAnywhere, Category = Mesh)
-	UStaticMesh* TranslationHandleMesh;
+	TObjectPtr<UStaticMesh> TranslationHandleMesh;
 
 	UPROPERTY(EditAnywhere, Category = Mesh)
-	UStaticMesh* UniformScaleHandleMesh;
+	TObjectPtr<UStaticMesh> UniformScaleHandleMesh;
 
 	UPROPERTY(EditAnywhere, Category = Mesh)
-	UStaticMesh* ScaleHandleMesh;
+	TObjectPtr<UStaticMesh> ScaleHandleMesh;
 
 	UPROPERTY(EditAnywhere, Category = Mesh)
-	UStaticMesh* PlaneTranslationHandleMesh;
+	TObjectPtr<UStaticMesh> PlaneTranslationHandleMesh;
 
 	UPROPERTY(EditAnywhere, Category = Mesh)
-	UStaticMesh* RotationHandleMesh;
+	TObjectPtr<UStaticMesh> RotationHandleMesh;
 
 	UPROPERTY(EditAnywhere, Category = Mesh)
-	UStaticMesh* RotationHandleSelectedMesh;
+	TObjectPtr<UStaticMesh> RotationHandleSelectedMesh;
 
 	UPROPERTY(EditAnywhere, Category = Mesh)
-	UStaticMesh* StartRotationIndicatorMesh;
+	TObjectPtr<UStaticMesh> StartRotationIndicatorMesh;
 
 	UPROPERTY(EditAnywhere, Category = Mesh)
-	UStaticMesh* CurrentRotationIndicatorMesh;
+	TObjectPtr<UStaticMesh> CurrentRotationIndicatorMesh;
 
 	UPROPERTY(EditAnywhere, Category = Mesh)
-	UStaticMesh* FreeRotationHandleMesh;
+	TObjectPtr<UStaticMesh> FreeRotationHandleMesh;
 
 	//
 	// Materials
 	//
 
 	UPROPERTY(EditAnywhere, Category = Material)
-	UMaterialInterface* GridMaterial;
+	TObjectPtr<UMaterialInterface> GridMaterial;
 
 	UPROPERTY(EditAnywhere, Category = Material)
-	UMaterialInterface* TransformGizmoMaterial;
+	TObjectPtr<UMaterialInterface> TransformGizmoMaterial;
 
 	UPROPERTY(EditAnywhere, Category = Material)
-	UMaterialInterface* TranslucentTransformGizmoMaterial;
+	TObjectPtr<UMaterialInterface> TranslucentTransformGizmoMaterial;
 };
diff --git a/Engine/Source/Editor/ViewportInteraction/Public/ViewportInteractor.h b/Engine/Source/Editor/ViewportInteraction/Public/ViewportInteractor.h
--- a/Engine/Source/Editor/ViewportInteraction/Public/ViewportInteractor.h
+++ b/Engine/Source/Editor/ViewportInteraction/Public/ViewportInteractor.h
@@ -296,11 +296,11 @@ protected:
 
 	/** The owning world interaction */
 	UPROPERTY()
-	class UViewportWorldInteraction* WorldInteraction;
+	TObjectPtr<class UViewportWorldInteraction> WorldInteraction;
 
 	/** The paired interactor by the world interaction */
 	UPROPERTY()
-	UViewportInteractor* OtherInteractor;
+	TObjectPtr<UViewportInteractor> OtherInteractor;
 
 	/** True if this interactor supports 'grabber sphere' interaction.  Usually disabled for mouse cursors */
 	bool bAllowGrabberSphere;
diff --git a/Engine/Source/Editor/ViewportInteraction/Public/ViewportTransformer.h b/Engine/Source/Editor/ViewportInteraction/Public/ViewportTransformer.h
--- a/Engine/Source/Editor/ViewportInteraction/Public/ViewportTransformer.h
+++ b/Engine/Source/Editor/ViewportInteraction/Public/ViewportTransformer.h
@@ -48,7 +48,7 @@ protected:
 
 	/** The viewport world interaction object we're registered with */
 	UPROPERTY()
-	UViewportWorldInteraction* ViewportWorldInteraction;
+	TObjectPtr<UViewportWorldInteraction> ViewportWorldInteraction;
 
 };
 
diff --git a/Engine/Source/Editor/ViewportInteraction/Public/ViewportWorldInteraction.h b/Engine/Source/Editor/ViewportInteraction/Public/ViewportWorldInteraction.h
--- a/Engine/Source/Editor/ViewportInteraction/Public/ViewportWorldInteraction.h
+++ b/Engine/Source/Editor/ViewportInteraction/Public/ViewportWorldInteraction.h
@@ -545,11 +545,11 @@ private:
 
 	/** All the interactors registered to modify the world */
 	UPROPERTY()
-	TArray< class UViewportInteractor* > Interactors;
+	TArray< TObjectPtr<class UViewportInteractor> > Interactors;
 
 	/** The active system being used to transform objects */
 	UPROPERTY()
-	class UViewportTransformer* ViewportTransformer;
+	TObjectPtr<class UViewportTransformer> ViewportTransformer;
 
 	//
 	// Viewport
@@ -621,7 +621,7 @@ private:
 
 	/** Transform gizmo actor, for manipulating selected actor(s) */
 	UPROPERTY()
-	class ABaseTransformGizmo* TransformGizmoActor;
+	TObjectPtr<class ABaseTransformGizmo> TransformGizmoActor;
 
 	/** Class used for when refreshing the transform gizmo */
 	TSubclassOf<ABaseTransformGizmo> TransformGizmoClass;
@@ -650,15 +650,15 @@ private:
 
 	/** Actor for the snap grid */
 	UPROPERTY()
-	class AActor* SnapGridActor;
+	TObjectPtr<class AActor> SnapGridActor;
 
 	/** The plane mesh we use to draw a snap grid under selected objects */
 	UPROPERTY()
-	class UStaticMeshComponent* SnapGridMeshComponent;
+	TObjectPtr<class UStaticMeshComponent> SnapGridMeshComponent;
 
 	/** MID for the snap grid, so we can update snap values on the fly */
 	UPROPERTY()
-	class UMaterialInstanceDynamic* SnapGridMID;
+	TObjectPtr<class UMaterialInstanceDynamic> SnapGridMID;
 
 	//
 	// Interactable
@@ -710,7 +710,7 @@ private:
 
 	/** The default mouse cursor interactor, activated on demand */
 	UPROPERTY()
-	class UMouseCursorInteractor* DefaultMouseCursorInteractor;
+	TObjectPtr<class UMouseCursorInteractor> DefaultMouseCursorInteractor;
 
 	/** Reference count for the default mouse cursor interactor.  When this reaches zero, the mouse cursor interactor goes away. */
 	int32 DefaultMouseCursorInteractorRefCount;
@@ -753,7 +753,7 @@ private:
 
 	/** Container of assets */
 	UPROPERTY()
-	const UViewportInteractionAssetContainer* AssetContainer;
+	TObjectPtr<const UViewportInteractionAssetContainer> AssetContainer;
 
 	/** If we want to skip playing the sound when refreshing the transform gizmo next time */
 	bool bPlayNextRefreshTransformGizmoSound;
diff --git a/Engine/Source/Runtime/AIModule/Classes/AIController.h b/Engine/Source/Runtime/AIModule/Classes/AIController.h
--- a/Engine/Source/Runtime/AIModule/Classes/AIController.h
+++ b/Engine/Source/Runtime/AIModule/Classes/AIController.h
@@ -125,28 +125,28 @@ private:
 
 	/** Component used for moving along a path. */
 	UPROPERTY(VisibleDefaultsOnly, Category = AI)
-	UPathFollowingComponent* PathFollowingComponent;
+	TObjectPtr<UPathFollowingComponent> PathFollowingComponent;
 
 public:
 
 	/** Component responsible for behaviors. */
 	UPROPERTY(BlueprintReadWrite, Category = AI)
-	UBrainComponent* BrainComponent;
+	TObjectPtr<UBrainComponent> BrainComponent;
 
 	UPROPERTY(VisibleDefaultsOnly, Category = AI)
-	UAIPerceptionComponent* PerceptionComponent;
+	TObjectPtr<UAIPerceptionComponent> PerceptionComponent;
 
 private:
 	UPROPERTY(BlueprintReadOnly, Category = AI, meta = (AllowPrivateAccess = "true"))
-	UPawnActionsComponent* ActionsComp;
+	TObjectPtr<UPawnActionsComponent> ActionsComp;
 
 protected:
 	/** blackboard */
 	UPROPERTY(BlueprintReadOnly, Category = AI, meta = (AllowPrivateAccess = "true"))
-	UBlackboardComponent* Blackboard;
+	TObjectPtr<UBlackboardComponent> Blackboard;
 
 	UPROPERTY()
-	UGameplayTasksComponent* CachedGameplayTasksComponent;
+	TObjectPtr<UGameplayTasksComponent> CachedGameplayTasksComponent;
 	
 	UPROPERTY(EditAnywhere, BlueprintReadWrite, Category = AI)
 	TSubclassOf<UNavigationQueryFilter> DefaultNavigationFilterClass;
diff --git a/Engine/Source/Runtime/AIModule/Classes/AISubsystem.h b/Engine/Source/Runtime/AIModule/Classes/AISubsystem.h
--- a/Engine/Source/Runtime/AIModule/Classes/AISubsystem.h
+++ b/Engine/Source/Runtime/AIModule/Classes/AISubsystem.h
@@ -18,7 +18,7 @@ class AIMODULE_API UAISubsystem : public UObject, public FTickableGameObject
 
 private:
 	UPROPERTY()
-	UAISystem* AISystem;
+	TObjectPtr<UAISystem> AISystem;
 			
 public:
 	UAISubsystem(const FObjectInitializer& ObjectInitializer = FObjectInitializer::Get());
diff --git a/Engine/Source/Runtime/AIModule/Classes/AISystem.h b/Engine/Source/Runtime/AIModule/Classes/AISystem.h
--- a/Engine/Source/Runtime/AIModule/Classes/AISystem.h
+++ b/Engine/Source/Runtime/AIModule/Classes/AISystem.h
@@ -99,23 +99,23 @@ public:
 protected:
 	/** Behavior tree manager used by game */
 	UPROPERTY(Transient)
-	UBehaviorTreeManager* BehaviorTreeManager;
+	TObjectPtr<UBehaviorTreeManager> BehaviorTreeManager;
 
 	/** Environment query manager used by game */
 	UPROPERTY(Transient)
-	UEnvQueryManager* EnvironmentQueryManager;
+	TObjectPtr<UEnvQueryManager> EnvironmentQueryManager;
 
 	UPROPERTY(Transient)
-	UAIPerceptionSystem* PerceptionSystem;
+	TObjectPtr<UAIPerceptionSystem> PerceptionSystem;
 
 	UPROPERTY(Transient)
-	TArray<UAIAsyncTaskBlueprintProxy*> AllProxyObjects;
+	TArray<TObjectPtr<UAIAsyncTaskBlueprintProxy>> AllProxyObjects;
 
 	UPROPERTY(Transient)
-	UAIHotSpotManager* HotSpotManager;
+	TObjectPtr<UAIHotSpotManager> HotSpotManager;
 
 	UPROPERTY(Transient)
-	UNavLocalGridManager* NavLocalGrids;
+	TObjectPtr<UNavLocalGridManager> NavLocalGrids;
 
 	typedef TMultiMap<TWeakObjectPtr<UBlackboardData>, TWeakObjectPtr<UBlackboardComponent> > FBlackboardDataToComponentsMap;
 
diff --git a/Engine/Source/Runtime/AIModule/Classes/AITypes.h b/Engine/Source/Runtime/AIModule/Classes/AITypes.h
--- a/Engine/Source/Runtime/AIModule/Classes/AITypes.h
+++ b/Engine/Source/Runtime/AIModule/Classes/AITypes.h
@@ -543,7 +543,7 @@ protected:
 
 	/** move goal: actor */
 	UPROPERTY()
-	AActor* GoalActor;
+	TObjectPtr<AActor> GoalActor;
 
 	/** move goal: location */
 	mutable FVector GoalLocation;
diff --git a/Engine/Source/Runtime/AIModule/Classes/Actions/PawnAction.h b/Engine/Source/Runtime/AIModule/Classes/Actions/PawnAction.h
--- a/Engine/Source/Runtime/AIModule/Classes/Actions/PawnAction.h
+++ b/Engine/Source/Runtime/AIModule/Classes/Actions/PawnAction.h
@@ -54,24 +54,24 @@ class AIMODULE_API UPawnAction : public UObject
 private:
 	/** Current child node executing on top of this Action */
 	UPROPERTY(Transient)
-	UPawnAction* ChildAction;
+	TObjectPtr<UPawnAction> ChildAction;
 
 	UPROPERTY(Transient)
-	UPawnAction* ParentAction;
+	TObjectPtr<UPawnAction> ParentAction;
 
 	/** Extra reference to the component this action is being governed by */
 	UPROPERTY(Transient)
-	UPawnActionsComponent* OwnerComponent;
+	TObjectPtr<UPawnActionsComponent> OwnerComponent;
 	
 	/** indicates an object that caused this action. Used for mass removal of actions 
 	 *	by specific object */
 	UPROPERTY(Transient)
-	UObject* Instigator;
+	TObjectPtr<UObject> Instigator;
 
 protected:
 	/** @Note: THIS IS HERE _ONLY_ BECAUSE OF THE WAY AI MESSAGING IS CURRENTLY IMPLEMENTED. WILL GO AWAY! */
 	UPROPERTY(Transient)
-	UBrainComponent* BrainComp;
+	TObjectPtr<UBrainComponent> BrainComp;
 
 private:
 	/** stores registered message observers */
diff --git a/Engine/Source/Runtime/AIModule/Classes/Actions/PawnAction_Move.h b/Engine/Source/Runtime/AIModule/Classes/Actions/PawnAction_Move.h
--- a/Engine/Source/Runtime/AIModule/Classes/Actions/PawnAction_Move.h
+++ b/Engine/Source/Runtime/AIModule/Classes/Actions/PawnAction_Move.h
@@ -29,7 +29,7 @@ class AIMODULE_API UPawnAction_Move : public UPawnAction
 	GENERATED_UCLASS_BODY()
 protected:
 	UPROPERTY(Category = PawnAction, EditAnywhere, BlueprintReadWrite)
-	AActor* GoalActor;
+	TObjectPtr<AActor> GoalActor;
 
 	UPROPERTY(Category = PawnAction, EditAnywhere, BlueprintReadWrite)
 	FVector GoalLocation;
diff --git a/Engine/Source/Runtime/AIModule/Classes/Actions/PawnAction_Repeat.h b/Engine/Source/Runtime/AIModule/Classes/Actions/PawnAction_Repeat.h
--- a/Engine/Source/Runtime/AIModule/Classes/Actions/PawnAction_Repeat.h
+++ b/Engine/Source/Runtime/AIModule/Classes/Actions/PawnAction_Repeat.h
@@ -19,10 +19,10 @@ class AIMODULE_API UPawnAction_Repeat : public UPawnAction
 
 	/** Action to repeat. This instance won't really be run, it's a source for copying actions to be actually performed */
 	UPROPERTY()
-	UPawnAction* ActionToRepeat;
+	TObjectPtr<UPawnAction> ActionToRepeat;
 
 	UPROPERTY(Transient)
-	UPawnAction* RecentActionCopy;
+	TObjectPtr<UPawnAction> RecentActionCopy;
 
 	UPROPERTY(Category = PawnAction, EditAnywhere, BlueprintReadOnly)
 	TEnumAsByte<EPawnActionFailHandling::Type> ChildFailureHandlingMode;
diff --git a/Engine/Source/Runtime/AIModule/Classes/Actions/PawnAction_Sequence.h b/Engine/Source/Runtime/AIModule/Classes/Actions/PawnAction_Sequence.h
--- a/Engine/Source/Runtime/AIModule/Classes/Actions/PawnAction_Sequence.h
+++ b/Engine/Source/Runtime/AIModule/Classes/Actions/PawnAction_Sequence.h
@@ -13,13 +13,13 @@ class AIMODULE_API UPawnAction_Sequence : public UPawnAction
 	GENERATED_UCLASS_BODY()
 
 	UPROPERTY()
-	TArray<UPawnAction*> ActionSequence;
+	TArray<TObjectPtr<UPawnAction>> ActionSequence;
 
 	UPROPERTY(Category = PawnAction, EditAnywhere, BlueprintReadOnly)
 	TEnumAsByte<EPawnActionFailHandling::Type> ChildFailureHandlingMode;
 
 	UPROPERTY(Transient)
-	UPawnAction* RecentActionCopy;
+	TObjectPtr<UPawnAction> RecentActionCopy;
 
 	uint32 CurrentActionIndex;
 
diff --git a/Engine/Source/Runtime/AIModule/Classes/Actions/PawnActionsComponent.h b/Engine/Source/Runtime/AIModule/Classes/Actions/PawnActionsComponent.h
--- a/Engine/Source/Runtime/AIModule/Classes/Actions/PawnActionsComponent.h
+++ b/Engine/Source/Runtime/AIModule/Classes/Actions/PawnActionsComponent.h
@@ -22,7 +22,7 @@ struct AIMODULE_API FPawnActionEvent
 	static const int32 FakeActionIndex = INDEX_NONE;
 
 	UPROPERTY()
-	UPawnAction* Action;
+	TObjectPtr<UPawnAction> Action;
 
 	EPawnActionEventType::Type EventType;
 
@@ -50,7 +50,7 @@ struct AIMODULE_API FPawnActionStack
 
 private:
 	UPROPERTY()
-	UPawnAction* TopAction;
+	TObjectPtr<UPawnAction> TopAction;
 
 public:
 	void Pause();
@@ -81,7 +81,7 @@ class AIMODULE_API UPawnActionsComponent : public UActorComponent
 
 protected:
 	UPROPERTY(BlueprintReadOnly, Category="PawnActions")
-	APawn* ControlledPawn;
+	TObjectPtr<APawn> ControlledPawn;
 
 	UPROPERTY()
 	TArray<FPawnActionStack> ActionStacks;
@@ -90,7 +90,7 @@ protected:
 	TArray<FPawnActionEvent> ActionEvents;
 
 	UPROPERTY(Transient)
-	UPawnAction* CurrentAction;
+	TObjectPtr<UPawnAction> CurrentAction;
 
 	/** set when logic was locked by hi priority stack */
 	uint32 bLockedAILogic : 1;
diff --git a/Engine/Source/Runtime/AIModule/Classes/BehaviorTree/BTCompositeNode.h b/Engine/Source/Runtime/AIModule/Classes/BehaviorTree/BTCompositeNode.h
--- a/Engine/Source/Runtime/AIModule/Classes/BehaviorTree/BTCompositeNode.h
+++ b/Engine/Source/Runtime/AIModule/Classes/BehaviorTree/BTCompositeNode.h
@@ -70,14 +70,14 @@ struct FBTCompositeChild
 
 	/** child node */
 	UPROPERTY()
-	UBTCompositeNode* ChildComposite = nullptr;
+	TObjectPtr<UBTCompositeNode> ChildComposite = nullptr;
 
 	UPROPERTY()
-	UBTTaskNode* ChildTask = nullptr;
+	TObjectPtr<UBTTaskNode> ChildTask = nullptr;
 
 	/** execution decorators */
 	UPROPERTY()
-	TArray<UBTDecorator*> Decorators;
+	TArray<TObjectPtr<UBTDecorator>> Decorators;
 
 	/** logic operations for decorators */
 	UPROPERTY()
@@ -95,7 +95,7 @@ class AIMODULE_API UBTCompositeNode : public UBTNode
 
 	/** service nodes */
 	UPROPERTY()
-	TArray<UBTService*> Services;
+	TArray<TObjectPtr<UBTService>> Services;
 
 	/** delegate for finding next child to execute */
 	UE_DEPRECATED(4.21, "OnNextChild is no longer being used. Please override UBTCompositeNode::GetNextChildHandler instead")
diff --git a/Engine/Source/Runtime/AIModule/Classes/BehaviorTree/BTNode.h b/Engine/Source/Runtime/AIModule/Classes/BehaviorTree/BTNode.h
--- a/Engine/Source/Runtime/AIModule/Classes/BehaviorTree/BTNode.h
+++ b/Engine/Source/Runtime/AIModule/Classes/BehaviorTree/BTNode.h
@@ -186,11 +186,11 @@ private:
 
 	/** source asset */
 	UPROPERTY()
-	UBehaviorTree* TreeAsset;
+	TObjectPtr<UBehaviorTree> TreeAsset;
 
 	/** parent node */
 	UPROPERTY()
-	UBTCompositeNode* ParentNode;
+	TObjectPtr<UBTCompositeNode> ParentNode;
 
 #if USE_BEHAVIORTREE_DEBUGGER
 	/** next node in execution order */
diff --git a/Engine/Source/Runtime/AIModule/Classes/BehaviorTree/BTTaskNode.h b/Engine/Source/Runtime/AIModule/Classes/BehaviorTree/BTTaskNode.h
--- a/Engine/Source/Runtime/AIModule/Classes/BehaviorTree/BTTaskNode.h
+++ b/Engine/Source/Runtime/AIModule/Classes/BehaviorTree/BTTaskNode.h
@@ -77,7 +77,7 @@ public:
 
 	/** service nodes */
 	UPROPERTY()
-	TArray<UBTService*> Services;
+	TArray<TObjectPtr<UBTService>> Services;
 
 protected:
 
diff --git a/Engine/Source/Runtime/AIModule/Classes/BehaviorTree/BehaviorTree.h b/Engine/Source/Runtime/AIModule/Classes/BehaviorTree/BehaviorTree.h
--- a/Engine/Source/Runtime/AIModule/Classes/BehaviorTree/BehaviorTree.h
+++ b/Engine/Source/Runtime/AIModule/Classes/BehaviorTree/BehaviorTree.h
@@ -18,13 +18,13 @@ class AIMODULE_API UBehaviorTree : public UObject, public IBlackboardAssetProvid
 
 	/** root node of loaded tree */
 	UPROPERTY()
-	UBTCompositeNode* RootNode;
+	TObjectPtr<UBTCompositeNode> RootNode;
 
 #if WITH_EDITORONLY_DATA
 
 	/** Graph for Behavior Tree */
 	UPROPERTY()
-	class UEdGraph*	BTGraph;
+	TObjectPtr<class UEdGraph>	BTGraph;
 
 	/** Info about the graphs we last edited */
 	UPROPERTY()
@@ -39,11 +39,11 @@ class AIMODULE_API UBehaviorTree : public UObject, public IBlackboardAssetProvid
 
 	/** blackboard asset for this tree */
 	UPROPERTY()
-	UBlackboardData* BlackboardAsset;
+	TObjectPtr<UBlackboardData> BlackboardAsset;
 
 	/** root level decorators, used by subtrees */
 	UPROPERTY()
-	TArray<UBTDecorator*> RootDecorators;
+	TArray<TObjectPtr<UBTDecorator>> RootDecorators;
 
 	/** logic operators for root level decorators, used by subtrees  */
 	UPROPERTY()
diff --git a/Engine/Source/Runtime/AIModule/Classes/BehaviorTree/BehaviorTreeComponent.h b/Engine/Source/Runtime/AIModule/Classes/BehaviorTree/BehaviorTreeComponent.h
--- a/Engine/Source/Runtime/AIModule/Classes/BehaviorTree/BehaviorTreeComponent.h
+++ b/Engine/Source/Runtime/AIModule/Classes/BehaviorTree/BehaviorTreeComponent.h
@@ -260,7 +260,7 @@ protected:
 
 	/** instanced nodes */
 	UPROPERTY(transient)
-	TArray<UBTNode*> NodeInstances;
+	TArray<TObjectPtr<UBTNode>> NodeInstances;
 
 	/** search data being currently used */
 	FBehaviorTreeSearchData SearchData;
@@ -430,7 +430,7 @@ protected:
 protected:
 	/** data asset defining the tree */
 	UPROPERTY(EditDefaultsOnly, BlueprintReadWrite, Category = AI)
-	UBehaviorTree* DefaultBehaviorTreeAsset;
+	TObjectPtr<UBehaviorTree> DefaultBehaviorTreeAsset;
 
 	/** Used to tell tickmanager that we want interval ticking */
 	bool bTickedOnce = false;
diff --git a/Engine/Source/Runtime/AIModule/Classes/BehaviorTree/BehaviorTreeManager.h b/Engine/Source/Runtime/AIModule/Classes/BehaviorTree/BehaviorTreeManager.h
--- a/Engine/Source/Runtime/AIModule/Classes/BehaviorTree/BehaviorTreeManager.h
+++ b/Engine/Source/Runtime/AIModule/Classes/BehaviorTree/BehaviorTreeManager.h
@@ -19,11 +19,11 @@ struct FBehaviorTreeTemplateInfo
 
 	/** behavior tree asset */
 	UPROPERTY()
-	UBehaviorTree* Asset = nullptr;
+	TObjectPtr<UBehaviorTree> Asset = nullptr;
 
 	/** initialized template */
 	UPROPERTY(transient)
-	UBTCompositeNode* Template = nullptr;
+	TObjectPtr<UBTCompositeNode> Template = nullptr;
 
 	/** size required for instance memory */
 	uint16 InstanceMemorySize;
@@ -68,5 +68,5 @@ protected:
 	TArray<FBehaviorTreeTemplateInfo> LoadedTemplates;
 
 	UPROPERTY()
-	TArray<UBehaviorTreeComponent*> ActiveComponents;
+	TArray<TObjectPtr<UBehaviorTreeComponent>> ActiveComponents;
 };
diff --git a/Engine/Source/Runtime/AIModule/Classes/BehaviorTree/BehaviorTreeTypes.h b/Engine/Source/Runtime/AIModule/Classes/BehaviorTree/BehaviorTreeTypes.h
--- a/Engine/Source/Runtime/AIModule/Classes/BehaviorTree/BehaviorTreeTypes.h
+++ b/Engine/Source/Runtime/AIModule/Classes/BehaviorTree/BehaviorTreeTypes.h
@@ -586,7 +586,7 @@ struct AIMODULE_API FBlackboardKeySelector
 	/** array of allowed types with additional properties (e.g. uobject's base class) 
 	  * EditAnywhere is required for FBlackboardSelectorDetails::CacheBlackboardData() */
 	UPROPERTY(transient, EditAnywhere, BlueprintReadWrite, Category = Blackboard)
-	TArray<UBlackboardKeyType*> AllowedTypes;
+	TArray<TObjectPtr<UBlackboardKeyType>> AllowedTypes;
 
 	/** name of selected key */
 	UPROPERTY(EditInstanceOnly, BlueprintReadWrite, Category = Blackboard)
diff --git a/Engine/Source/Runtime/AIModule/Classes/BehaviorTree/Blackboard/BlackboardKeyType_Class.h b/Engine/Source/Runtime/AIModule/Classes/BehaviorTree/Blackboard/BlackboardKeyType_Class.h
--- a/Engine/Source/Runtime/AIModule/Classes/BehaviorTree/Blackboard/BlackboardKeyType_Class.h
+++ b/Engine/Source/Runtime/AIModule/Classes/BehaviorTree/Blackboard/BlackboardKeyType_Class.h
@@ -17,7 +17,7 @@ class AIMODULE_API UBlackboardKeyType_Class : public UBlackboardKeyType
 	static const FDataType InvalidValue;
 
 	UPROPERTY(Category=Blackboard, EditDefaultsOnly, meta=(AllowAbstract="1"))
-	UClass* BaseClass;
+	TObjectPtr<UClass> BaseClass;
 	
 	static UClass* GetValue(const UBlackboardKeyType_Class* KeyOb, const uint8* RawData);
 	static bool SetValue(UBlackboardKeyType_Class* KeyOb, uint8* RawData, UClass* Value);
diff --git a/Engine/Source/Runtime/AIModule/Classes/BehaviorTree/Blackboard/BlackboardKeyType_Enum.h b/Engine/Source/Runtime/AIModule/Classes/BehaviorTree/Blackboard/BlackboardKeyType_Enum.h
--- a/Engine/Source/Runtime/AIModule/Classes/BehaviorTree/Blackboard/BlackboardKeyType_Enum.h
+++ b/Engine/Source/Runtime/AIModule/Classes/BehaviorTree/Blackboard/BlackboardKeyType_Enum.h
@@ -17,7 +17,7 @@ class AIMODULE_API UBlackboardKeyType_Enum : public UBlackboardKeyType
 	static const FDataType InvalidValue;
 
 	UPROPERTY(Category=Blackboard, EditDefaultsOnly)
-	UEnum* EnumType;
+	TObjectPtr<UEnum> EnumType;
 
 	/** name of enum defined in c++ code, will take priority over asset from EnumType property */
 	UPROPERTY(Category=Blackboard, EditDefaultsOnly)
diff --git a/Engine/Source/Runtime/AIModule/Classes/BehaviorTree/Blackboard/BlackboardKeyType_NativeEnum.h b/Engine/Source/Runtime/AIModule/Classes/BehaviorTree/Blackboard/BlackboardKeyType_NativeEnum.h
--- a/Engine/Source/Runtime/AIModule/Classes/BehaviorTree/Blackboard/BlackboardKeyType_NativeEnum.h
+++ b/Engine/Source/Runtime/AIModule/Classes/BehaviorTree/Blackboard/BlackboardKeyType_NativeEnum.h
@@ -21,7 +21,7 @@ class AIMODULE_API UBlackboardKeyType_NativeEnum : public UBlackboardKeyType
 	FString EnumName;
 
 	UPROPERTY()
-	UEnum* EnumType;
+	TObjectPtr<UEnum> EnumType;
 
 	virtual UBlackboardKeyType* UpdateDeprecatedKey() override;
 
diff --git a/Engine/Source/Runtime/AIModule/Classes/BehaviorTree/Blackboard/BlackboardKeyType_Object.h b/Engine/Source/Runtime/AIModule/Classes/BehaviorTree/Blackboard/BlackboardKeyType_Object.h
--- a/Engine/Source/Runtime/AIModule/Classes/BehaviorTree/Blackboard/BlackboardKeyType_Object.h
+++ b/Engine/Source/Runtime/AIModule/Classes/BehaviorTree/Blackboard/BlackboardKeyType_Object.h
@@ -17,7 +17,7 @@ class AIMODULE_API UBlackboardKeyType_Object : public UBlackboardKeyType
 	static const FDataType InvalidValue;
 
 	UPROPERTY(Category=Blackboard, EditDefaultsOnly, meta=(AllowAbstract="1"))
-	UClass* BaseClass;
+	TObjectPtr<UClass> BaseClass;
 
 	static UObject* GetValue(const UBlackboardKeyType_Object* KeyOb, const uint8* RawData);
 	static bool SetValue(UBlackboardKeyType_Object* KeyOb, uint8* RawData, UObject* Value);
diff --git a/Engine/Source/Runtime/AIModule/Classes/BehaviorTree/BlackboardComponent.h b/Engine/Source/Runtime/AIModule/Classes/BehaviorTree/BlackboardComponent.h
--- a/Engine/Source/Runtime/AIModule/Classes/BehaviorTree/BlackboardComponent.h
+++ b/Engine/Source/Runtime/AIModule/Classes/BehaviorTree/BlackboardComponent.h
@@ -233,18 +233,18 @@ protected:
 
 	/** cached behavior tree component */
 	UPROPERTY(transient)
-	UBrainComponent* BrainComp;
+	TObjectPtr<UBrainComponent> BrainComp;
 
 	/** data asset defining entries. Will be used as part of InitializeComponent 
 	 *	call provided BlackboardAsset hasn't been already set (via a InitializeBlackboard 
 	 *	call). */
 	UPROPERTY(EditDefaultsOnly, Category = AI)
-	UBlackboardData* DefaultBlackboardAsset;
+	TObjectPtr<UBlackboardData> DefaultBlackboardAsset;
 
 	/** internal use, current BB asset being used. Will be made private in the future */
 	UE_DEPRECATED_FORGAME(4.26, "Directly accessing BlackboardAsset is not longer supported. Use DefaultBlackboardAsset or InitializeBlackboard to set it and GetBlackboardAsset to retrieve it")
 	UPROPERTY(transient)
-	UBlackboardData* BlackboardAsset;
+	TObjectPtr<UBlackboardData> BlackboardAsset;
 
 	/** memory block holding all values */
 	TArray<uint8> ValueMemory;
@@ -254,7 +254,7 @@ protected:
 
 	/** instanced keys with custom data allocations */
 	UPROPERTY(transient)
-	TArray<UBlackboardKeyType*> KeyInstances;
+	TArray<TObjectPtr<UBlackboardKeyType>> KeyInstances;
 
 protected:
 	struct FOnBlackboardChangeNotificationInfo
diff --git a/Engine/Source/Runtime/AIModule/Classes/BehaviorTree/BlackboardData.h b/Engine/Source/Runtime/AIModule/Classes/BehaviorTree/BlackboardData.h
--- a/Engine/Source/Runtime/AIModule/Classes/BehaviorTree/BlackboardData.h
+++ b/Engine/Source/Runtime/AIModule/Classes/BehaviorTree/BlackboardData.h
@@ -28,7 +28,7 @@ struct FBlackboardEntry
 
 	/** key type and additional properties */
 	UPROPERTY(EditAnywhere, Instanced, Category=Blackboard)
-	UBlackboardKeyType* KeyType;
+	TObjectPtr<UBlackboardKeyType> KeyType;
 
 	/** if set to true then this field will be synchronized across all instances of this blackboard */
 	UPROPERTY(EditAnywhere, Category=Blackboard)
@@ -49,7 +49,7 @@ class AIMODULE_API UBlackboardData : public UDataAsset
 
 	/** parent blackboard (keys can be overridden) */
 	UPROPERTY(EditAnywhere, Category=Parent)
-	UBlackboardData* Parent;
+	TObjectPtr<UBlackboardData> Parent;
 
 #if WITH_EDITORONLY_DATA
 	/** all keys inherited from parent chain */
diff --git a/Engine/Source/Runtime/AIModule/Classes/BehaviorTree/Decorators/BTDecorator_BlueprintBase.h b/Engine/Source/Runtime/AIModule/Classes/BehaviorTree/Decorators/BTDecorator_BlueprintBase.h
--- a/Engine/Source/Runtime/AIModule/Classes/BehaviorTree/Decorators/BTDecorator_BlueprintBase.h
+++ b/Engine/Source/Runtime/AIModule/Classes/BehaviorTree/Decorators/BTDecorator_BlueprintBase.h
@@ -66,11 +66,11 @@ public:
 protected:
 	/** Cached AIController owner of BehaviorTreeComponent. */
 	UPROPERTY(Transient)
-	AAIController* AIOwner;
+	TObjectPtr<AAIController> AIOwner;
 
 	/** Cached AIController owner of BehaviorTreeComponent. */
 	UPROPERTY(Transient)
-	AActor* ActorOwner;
+	TObjectPtr<AActor> ActorOwner;
 
 	/** blackboard key names that should be observed */
 	UPROPERTY()
diff --git a/Engine/Source/Runtime/AIModule/Classes/BehaviorTree/Services/BTService_BlueprintBase.h b/Engine/Source/Runtime/AIModule/Classes/BehaviorTree/Services/BTService_BlueprintBase.h
--- a/Engine/Source/Runtime/AIModule/Classes/BehaviorTree/Services/BTService_BlueprintBase.h
+++ b/Engine/Source/Runtime/AIModule/Classes/BehaviorTree/Services/BTService_BlueprintBase.h
@@ -38,11 +38,11 @@ class AIMODULE_API UBTService_BlueprintBase : public UBTService
 protected:
 	/** Cached AIController owner of BehaviorTreeComponent. */
 	UPROPERTY(Transient)
-	AAIController* AIOwner;
+	TObjectPtr<AAIController> AIOwner;
 
 	/** Cached actor owner of BehaviorTreeComponent. */
 	UPROPERTY(Transient)
-	AActor* ActorOwner;
+	TObjectPtr<AActor> ActorOwner;
 
 	// Gets the description for our service
 	virtual FString GetStaticServiceDescription() const override;
diff --git a/Engine/Source/Runtime/AIModule/Classes/BehaviorTree/Tasks/BTTask_PlayAnimation.h b/Engine/Source/Runtime/AIModule/Classes/BehaviorTree/Tasks/BTTask_PlayAnimation.h
--- a/Engine/Source/Runtime/AIModule/Classes/BehaviorTree/Tasks/BTTask_PlayAnimation.h
+++ b/Engine/Source/Runtime/AIModule/Classes/BehaviorTree/Tasks/BTTask_PlayAnimation.h
@@ -22,7 +22,7 @@ class AIMODULE_API UBTTask_PlayAnimation : public UBTTaskNode
 
 	/** Animation asset to play. Note that it needs to match the skeleton of pawn this BT is controlling */
 	UPROPERTY(Category = Node, EditAnywhere)
-	UAnimationAsset* AnimationToPlay;
+	TObjectPtr<UAnimationAsset> AnimationToPlay;
 	
 	UPROPERTY(Category = Node, EditAnywhere)
 	uint32 bLooping : 1; 
@@ -32,10 +32,10 @@ class AIMODULE_API UBTTask_PlayAnimation : public UBTTaskNode
 	uint32 bNonBlocking : 1;
 
 	UPROPERTY()
-	UBehaviorTreeComponent* MyOwnerComp;
+	TObjectPtr<UBehaviorTreeComponent> MyOwnerComp;
 
 	UPROPERTY()
-	USkeletalMeshComponent* CachedSkelMesh;
+	TObjectPtr<USkeletalMeshComponent> CachedSkelMesh;
 
 	EAnimationMode::Type PreviousAnimationMode;
 
