Adding more conditionals for closing sequencer when entering VR

diff --git a/Engine/Source/Editor/VREditor/VREditorMode.cpp b/Engine/Source/Editor/VREditor/VREditorMode.cpp
--- a/Engine/Source/Editor/VREditor/VREditorMode.cpp
+++ b/Engine/Source/Editor/VREditor/VREditorMode.cpp
@@ -278,8 +278,13 @@ void UVREditorMode::Enter()
 	// Switch us back to default mode and close any open sequencer windows
 	FVREditorActionCallbacks::ChangeEditorModes(FBuiltinEditorModes::EM_Default);
 	FLevelEditorModule& LevelEditorModule = FModuleManager::GetModuleChecked<FLevelEditorModule>(TEXT("LevelEditor"));
-	LevelEditorModule.GetLevelEditorTabManager()->TryInvokeTab(FTabId("Sequencer"))->RequestCloseTab();
-
+	if (LevelEditorModule.GetLevelEditorTabManager())
+	{
+		if (TSharedPtr<SDockTab> SequencerTab = LevelEditorModule.GetLevelEditorTabManager()->TryInvokeTab(FTabId("Sequencer")))
+		{
+			SequencerTab->RequestCloseTab();
+		}
+	}
 	// Setup our avatar
 	if (AvatarActor == nullptr)
 	{
