Virtual Scouting: Changes required for compatibility with OpenXR.

This migrates the mostly declarative action mappings in UVREditorInteractor::Init_Implementation into a new GetKnownActionMappings method, which we can poll in order to translate them into an OpenXR action set at the point during early initialization which OpenXR requires them to be registered. This also adds a `bIsAxis` flag to FViewportActionKeyInput, because axis mappings need to be handled differently.

The glue code leveraging this is introduced in a new VirtualScouting plugin, where we'll ultimately migrate the existing code as well (away from VirtualProductionUtilities).

diff --git a/Engine/Plugins/Experimental/VirtualScouting/Source/VirtualScoutingOpenXR/Private/VirtualScoutingOpenXR.cpp b/Engine/Plugins/Experimental/VirtualScouting/Source/VirtualScoutingOpenXR/Private/VirtualScoutingOpenXR.cpp
new file mode 100644
--- /dev/null
+++ b/Engine/Plugins/Experimental/VirtualScouting/Source/VirtualScoutingOpenXR/Private/VirtualScoutingOpenXR.cpp
@@ -0,0 +1,260 @@
+// Copyright Epic Games, Inc. All Rights Reserved.
+
+#include "VirtualScoutingOpenXR.h"
+
+#if WITH_EDITOR
+
+#include "Logging/LogMacros.h"
+#include "IVREditorModule.h"
+#include "ViewportWorldInteraction.h"
+#include "VREditorMode.h"
+#include "VREditorInteractor.h"
+#include "VRModeSettings.h"
+
+
+#define LOCTEXT_NAMESPACE "VirtualScouting"
+
+
+DECLARE_LOG_CATEGORY_EXTERN(LogVPOpenXRDebug, VeryVerbose, All);
+DEFINE_LOG_CATEGORY(LogVPOpenXRDebug);
+
+
+static TAutoConsoleVariable<int32> CVarOpenXRDebugLogging(
+	TEXT("VirtualScouting.OpenXRDebugLogging"),
+	0,
+	TEXT("If true, register an Unreal log sink via XR_EXT_debug_utils.\n"),
+	ECVF_Default);
+
+
+class FVirtualScoutingOpenXRModule : public IModuleInterface
+{
+private:
+	virtual void StartupModule() override
+	{
+		OpenXRExt = MakeUnique<FVirtualScoutingOpenXRExtension>();
+	}
+
+	virtual void ShutdownModule() override
+	{
+		OpenXRExt.Reset();
+	}
+
+private:
+	TUniquePtr<FVirtualScoutingOpenXRExtension> OpenXRExt;
+};
+
+
+IMPLEMENT_MODULE(FVirtualScoutingOpenXRModule, VirtualScoutingOpenXR);
+
+
+FVirtualScoutingOpenXRExtension::FVirtualScoutingOpenXRExtension()
+{
+	RegisterOpenXRExtensionModularFeature();
+
+	InitCompleteDelegate = FCoreDelegates::OnFEngineLoopInitComplete.AddLambda(
+		[this]()
+		{
+			IVREditorModule::Get().OnVREditingModeEnter().AddRaw(this, &FVirtualScoutingOpenXRExtension::OnVREditingModeEnter);
+			IVREditorModule::Get().OnVREditingModeExit().AddRaw(this, &FVirtualScoutingOpenXRExtension::OnVREditingModeExit);
+
+			// Must happen last; this implicitly deallocates this lambda's captures.
+			FCoreDelegates::OnFEngineLoopInitComplete.Remove(InitCompleteDelegate);
+		}
+	);
+}
+
+
+FVirtualScoutingOpenXRExtension::~FVirtualScoutingOpenXRExtension()
+{
+#if 0 // FIXME?: Too late to use Instance here, and don't see any suitable extension interface methods.
+	// Might also be OK not to explicitly clean this up.
+	if (Messenger != XR_NULL_HANDLE)
+	{
+		PFN_xrDestroyDebugUtilsMessengerEXT PfnXrDestroyDebugUtilsMessengerEXT;
+		if (XR_ENSURE(xrGetInstanceProcAddr(Instance, "xrDestroyDebugUtilsMessengerEXT",
+			reinterpret_cast<PFN_xrVoidFunction*>(&PfnXrDestroyDebugUtilsMessengerEXT))))
+		{
+			XR_ENSURE(PfnXrDestroyDebugUtilsMessengerEXT(Messenger));
+			Messenger = XR_NULL_HANDLE;
+		}
+	}
+#endif
+
+	UnregisterOpenXRExtensionModularFeature();
+
+	if (IVREditorModule::IsAvailable())
+	{
+		IVREditorModule::Get().OnVREditingModeEnter().RemoveAll(this);
+		IVREditorModule::Get().OnVREditingModeExit().RemoveAll(this);
+	}
+}
+
+
+bool FVirtualScoutingOpenXRExtension::GetOptionalExtensions(TArray<const ANSICHAR*>& OutExtensions)
+{
+	if (CVarOpenXRDebugLogging.GetValueOnAnyThread() != 0)
+	{
+		OutExtensions.Add(XR_EXT_DEBUG_UTILS_EXTENSION_NAME);
+	}
+
+	return true;
+}
+
+
+void FVirtualScoutingOpenXRExtension::PostCreateInstance(XrInstance InInstance)
+{
+	Instance = InInstance;
+
+	if (CVarOpenXRDebugLogging.GetValueOnAnyThread() != 0)
+	{
+		PFN_xrCreateDebugUtilsMessengerEXT PfnXrCreateDebugUtilsMessengerEXT;
+		if (XR_ENSURE(xrGetInstanceProcAddr(Instance, "xrCreateDebugUtilsMessengerEXT",
+			reinterpret_cast<PFN_xrVoidFunction*>(&PfnXrCreateDebugUtilsMessengerEXT))))
+		{
+			XrDebugUtilsMessengerCreateInfoEXT DebugMessengerCreateInfo;
+			DebugMessengerCreateInfo.type = XR_TYPE_DEBUG_UTILS_MESSENGER_CREATE_INFO_EXT;
+			DebugMessengerCreateInfo.next = nullptr;
+			DebugMessengerCreateInfo.messageSeverities =
+				XR_DEBUG_UTILS_MESSAGE_SEVERITY_ERROR_BIT_EXT | XR_DEBUG_UTILS_MESSAGE_SEVERITY_WARNING_BIT_EXT |
+				XR_DEBUG_UTILS_MESSAGE_SEVERITY_INFO_BIT_EXT | XR_DEBUG_UTILS_MESSAGE_SEVERITY_VERBOSE_BIT_EXT;
+			DebugMessengerCreateInfo.messageTypes =
+				XR_DEBUG_UTILS_MESSAGE_TYPE_GENERAL_BIT_EXT | XR_DEBUG_UTILS_MESSAGE_TYPE_VALIDATION_BIT_EXT |
+				XR_DEBUG_UTILS_MESSAGE_TYPE_PERFORMANCE_BIT_EXT | XR_DEBUG_UTILS_MESSAGE_TYPE_CONFORMANCE_BIT_EXT;
+			DebugMessengerCreateInfo.userCallback = &FVirtualScoutingOpenXRExtension::XrDebugUtilsMessengerCallback_Trampoline;
+			DebugMessengerCreateInfo.userData = this;
+
+			if (XR_ENSURE(PfnXrCreateDebugUtilsMessengerEXT(Instance, &DebugMessengerCreateInfo, &Messenger)))
+			{
+				UE_LOG(LogVPOpenXRDebug, Log, TEXT("XR_EXT_debug_utils messenger ACTIVE"));
+				return;
+			}
+		}
+	}
+
+	UE_LOG(LogVPOpenXRDebug, Log, TEXT("XR_EXT_debug_utils messenger DISABLED"));
+}
+
+
+void FVirtualScoutingOpenXRExtension::AddActions(XrInstance InInstance, FCreateActionSetFunc CreateActionSetFunc, FCreateActionFunc CreateActionFunc)
+{
+	// Clean up any previously created action set, which also implicitly cleans up the actions belonging to it.
+	if (ActionSet != XR_NULL_HANDLE)
+	{
+		XR_ENSURE(xrDestroyActionSet(ActionSet));
+		ActionSet = XR_NULL_HANDLE;
+	}
+
+	UClass* InteractorClass = GetDefault<UVRModeSettings>()->InteractorClass.LoadSynchronous();
+	if (!InteractorClass)
+	{
+		return;
+	}
+
+	const UVREditorInteractor* InteractorCDO = Cast<const UVREditorInteractor>(InteractorClass->GetDefaultObject());
+	if (!InteractorCDO)
+	{
+		return;
+	}
+
+	// Create action set
+	{
+		FActionSetParams ActionSetParams;
+		ActionSetParams.Name = "ue_vp_scouting";
+		ActionSetParams.LocalizedName = LOCTEXT("ActionSetName", "Virtual Scouting");
+		ActionSetParams.Priority = 100;
+		ActionSet = CreateActionSetFunc(ActionSetParams);
+	}
+
+	XrPath LeftHand, RightHand;
+	check(XR_SUCCEEDED(xrStringToPath(Instance, "/user/hand/left", &LeftHand)));
+	check(XR_SUCCEEDED(xrStringToPath(Instance, "/user/hand/right", &RightHand)));
+
+	FActionParams ActionParams;
+	ActionParams.Set = ActionSet;
+	ActionParams.SubactionPaths = { LeftHand, RightHand };
+
+	for (TPair<FViewportActionKeyInput, TArray<FKey>>& ActionMappings : InteractorCDO->GetKnownActionMappings())
+	{
+		FViewportActionKeyInput& Action = ActionMappings.Key;
+		TArray<FKey>& Keys = ActionMappings.Value;
+
+		ActionParams.Type = Action.bIsAxis ? XR_ACTION_TYPE_FLOAT_INPUT : XR_ACTION_TYPE_BOOLEAN_INPUT;
+		ActionParams.Name = Action.ActionType;
+		// TODO: ActionParams.LocalizedName
+		ActionParams.SuggestedBindings = MoveTemp(Keys);
+		CreateActionFunc(ActionParams);
+	}
+}
+
+
+void FVirtualScoutingOpenXRExtension::GetActiveActionSetsForSync(TArray<XrActiveActionSet>& OutActiveSets)
+{
+	if (bIsVrEditingModeActive && ensure(ActionSet != XR_NULL_HANDLE))
+	{
+		XrActiveActionSet ActiveSet;
+		ActiveSet.actionSet = ActionSet;
+		ActiveSet.subactionPath = XR_NULL_PATH;
+		OutActiveSets.Add(ActiveSet);
+	}
+}
+
+
+void FVirtualScoutingOpenXRExtension::OnVREditingModeEnter()
+{
+	bIsVrEditingModeActive = true;
+}
+
+
+void FVirtualScoutingOpenXRExtension::OnVREditingModeExit()
+{
+	bIsVrEditingModeActive = false;
+}
+
+
+//static
+XrBool32 XRAPI_CALL FVirtualScoutingOpenXRExtension::XrDebugUtilsMessengerCallback_Trampoline(
+	XrDebugUtilsMessageSeverityFlagsEXT InMessageSeverity,
+	XrDebugUtilsMessageTypeFlagsEXT InMessageTypes,
+	const XrDebugUtilsMessengerCallbackDataEXT* InCallbackData,
+	void* InUserData)
+{
+	// TODO?: InUserData contains the FVirtualScoutingOpenXRExtension*, could forward to an instance method.
+	// However, doing the work directly in here seems fine for now.
+
+	ELogVerbosity::Type Verbosity;
+	switch (InMessageSeverity)
+	{
+		case XR_DEBUG_UTILS_MESSAGE_SEVERITY_VERBOSE_BIT_EXT: Verbosity = ELogVerbosity::Verbose; break;
+		case XR_DEBUG_UTILS_MESSAGE_SEVERITY_INFO_BIT_EXT:    Verbosity = ELogVerbosity::Display; break;
+		case XR_DEBUG_UTILS_MESSAGE_SEVERITY_WARNING_BIT_EXT: Verbosity = ELogVerbosity::Warning; break;
+		case XR_DEBUG_UTILS_MESSAGE_SEVERITY_ERROR_BIT_EXT:   Verbosity = ELogVerbosity::Error;   break;
+		default:
+			ensureMsgf(false, TEXT("Unhandled XrDebugUtilsMessageSeverityFlagsEXT: %X"), InMessageSeverity);
+			Verbosity = ELogVerbosity::Error;
+			break;
+	}
+
+	// Results in "____" if no bits are set, or "GVPC" if all bits are set, or some combination.
+	FString Types = TEXT("####");
+	Types[0] = (InMessageTypes & XR_DEBUG_UTILS_MESSAGE_TYPE_GENERAL_BIT_EXT)     ? TEXT('G') : TEXT('_');
+	Types[1] = (InMessageTypes & XR_DEBUG_UTILS_MESSAGE_TYPE_VALIDATION_BIT_EXT)  ? TEXT('V') : TEXT('_');
+	Types[2] = (InMessageTypes & XR_DEBUG_UTILS_MESSAGE_TYPE_PERFORMANCE_BIT_EXT) ? TEXT('P') : TEXT('_');
+	Types[3] = (InMessageTypes & XR_DEBUG_UTILS_MESSAGE_TYPE_CONFORMANCE_BIT_EXT) ? TEXT('C') : TEXT('_');
+
+	FMsg::Logf(__FILE__, __LINE__, LogVPOpenXRDebug.GetCategoryName(), Verbosity,
+		TEXT("[%s]: %S(): %S"), *Types, InCallbackData->functionName, InCallbackData->message);
+
+	// "A value of XR_TRUE indicates that the application wants to abort this call. [...]
+	// Applications should always return XR_FALSE so that they see the same behavior with
+	// and without validation layers enabled."
+	return XR_FALSE;
+}
+
+
+#undef LOCTEXT_NAMESPACE
+
+#else // not WITH_EDITOR
+
+IMPLEMENT_MODULE(FDefaultModuleImpl, VirtualScoutingOpenXR);
+
+#endif // #if WITH_EDITOR .. #else
diff --git a/Engine/Plugins/Experimental/VirtualScouting/Source/VirtualScoutingOpenXR/Private/VirtualScoutingOpenXR.h b/Engine/Plugins/Experimental/VirtualScouting/Source/VirtualScoutingOpenXR/Private/VirtualScoutingOpenXR.h
new file mode 100644
--- /dev/null
+++ b/Engine/Plugins/Experimental/VirtualScouting/Source/VirtualScoutingOpenXR/Private/VirtualScoutingOpenXR.h
@@ -0,0 +1,50 @@
+// Copyright Epic Games, Inc. All Rights Reserved.
+
+#if WITH_EDITOR
+
+#include "IOpenXRExtensionPlugin.h"
+
+#include "OpenXRCore.h"
+
+
+class VIRTUALSCOUTINGOPENXR_API FVirtualScoutingOpenXRExtension : public IOpenXRExtensionPlugin
+{
+public:
+	FVirtualScoutingOpenXRExtension();
+	virtual ~FVirtualScoutingOpenXRExtension();
+
+	//~ Begin IOpenXRExtensionPlugin interface
+	virtual FString GetDisplayName() override
+	{
+		return FString(TEXT("VirtualScouting"));
+	}
+
+	virtual bool GetOptionalExtensions(TArray<const ANSICHAR*>& OutExtensions) override;
+	virtual void PostCreateInstance(XrInstance InInstance) override;
+	virtual void AddActions(XrInstance InInstance, FCreateActionSetFunc CreateActionSet, FCreateActionFunc CreateAction) override;
+	virtual void GetActiveActionSetsForSync(TArray<XrActiveActionSet>& OutActiveSets) override;
+	//~ End IOpenXRExtensionPlugin interface
+
+private:
+	void OnVREditingModeEnter();
+	void OnVREditingModeExit();
+
+private:
+	FDelegateHandle InitCompleteDelegate;
+
+	XrDebugUtilsMessengerEXT Messenger = XR_NULL_HANDLE;
+
+	XrInstance Instance = XR_NULL_HANDLE;
+	XrActionSet ActionSet = XR_NULL_HANDLE;
+
+	bool bIsVrEditingModeActive = false;
+
+private:
+	static XrBool32 XRAPI_CALL XrDebugUtilsMessengerCallback_Trampoline(
+		XrDebugUtilsMessageSeverityFlagsEXT InMessageSeverity,
+		XrDebugUtilsMessageTypeFlagsEXT InMessageTypes,
+		const XrDebugUtilsMessengerCallbackDataEXT* InCallbackData,
+		void* InUserData);
+};
+
+#endif // #if WITH_EDITOR
diff --git a/Engine/Plugins/Experimental/VirtualScouting/Source/VirtualScoutingOpenXR/VirtualScoutingOpenXR.Build.cs b/Engine/Plugins/Experimental/VirtualScouting/Source/VirtualScoutingOpenXR/VirtualScoutingOpenXR.Build.cs
new file mode 100644
--- /dev/null
+++ b/Engine/Plugins/Experimental/VirtualScouting/Source/VirtualScoutingOpenXR/VirtualScoutingOpenXR.Build.cs
@@ -0,0 +1,33 @@
+// Copyright Epic Games, Inc. All Rights Reserved.
+
+using UnrealBuildTool;
+
+public class VirtualScoutingOpenXR : ModuleRules
+{
+	public VirtualScoutingOpenXR(ReadOnlyTargetRules Target) : base(Target)
+	{
+		PCHUsage = ModuleRules.PCHUsageMode.UseExplicitOrSharedPCHs;
+
+		PublicDependencyModuleNames.AddRange(new string[]
+			{
+				"Core",
+				"CoreUObject",
+				"Engine",
+				"InputCore",
+				"OpenXRHMD",
+			}
+		);
+
+		// This entire module is an Editor module in spirit, but in practice must be Runtime,
+		// because IOpenXRExtensionPlugin requires the PostConfigInit LoadingPhase.
+		if (Target.bBuildEditor)
+		{
+			PublicDependencyModuleNames.Add("VREditor");
+		}
+
+		PrivateDependencyModuleNames.AddRange(new string[]
+			{
+			}
+		);
+	}
+}
diff --git a/Engine/Plugins/Experimental/VirtualScouting/VirtualScouting.uplugin b/Engine/Plugins/Experimental/VirtualScouting/VirtualScouting.uplugin
new file mode 100644
--- /dev/null
+++ b/Engine/Plugins/Experimental/VirtualScouting/VirtualScouting.uplugin
@@ -0,0 +1,41 @@
+{
+	"FileVersion": 3,
+	"Version": 1,
+	"VersionName": "1.0",
+	"FriendlyName": "Virtual Scouting",
+	"Description": "Virtual Scouting lets filmmakers scout a digital environment in virtual reality.",
+	"Category": "Virtual Production",
+	"CreatedBy": "Epic Games, Inc.",
+	"CreatedByURL": "http://epicgames.com",
+	"DocsURL": "",
+	"MarketplaceURL": "",
+	"SupportURL": "",
+	"CanContainContent": true,
+	"IsExperimentalVersion": true,
+	"Installed": false,
+	"SupportedTargetPlatforms": [
+		"Win64",
+		"Linux"
+	],
+	"Modules": [
+		{
+			"Name": "VirtualScoutingOpenXR",
+			"Type": "Runtime",
+			"LoadingPhase": "PostConfigInit",
+			"PlatformAllowList": [
+				"Win64",
+				"Linux"
+			]
+		}
+	],
+	"Plugins": [
+		{
+			"Name": "VirtualProductionUtilities",
+			"Enabled": true
+		},
+		{
+			"Name": "OpenXR",
+			"Enabled": true
+		}
+	]
+}
diff --git a/Engine/Source/Editor/UnrealEd/Private/EditorEngine.cpp b/Engine/Source/Editor/UnrealEd/Private/EditorEngine.cpp
--- a/Engine/Source/Editor/UnrealEd/Private/EditorEngine.cpp
+++ b/Engine/Source/Editor/UnrealEd/Private/EditorEngine.cpp
@@ -2110,8 +2110,8 @@ float UEditorEngine::GetMaxTickRate( float DeltaTime, bool bAllowFrameRateSmooth
 	float MaxTickRate = 0.0f;
 	if( !ShouldThrottleCPUUsage() )
 	{
-		// do not limit fps in VR Preview mode
-		if (IsVRPreviewActive())
+		// do not limit fps in VR Preview mode or VR editor mode
+		if (IsVRPreviewActive() || GEnableVREditorHacks)
 		{
 			return 0.0f;
 		}
diff --git a/Engine/Source/Editor/VREditor/Public/VREditorInteractor.h b/Engine/Source/Editor/VREditor/Public/VREditorInteractor.h
--- a/Engine/Source/Editor/VREditor/Public/VREditorInteractor.h
+++ b/Engine/Source/Editor/VREditor/Public/VREditorInteractor.h
@@ -66,6 +66,8 @@ public:
 		return KeyToActionMap;
 	}
 
+	virtual TMap<FViewportActionKeyInput, TArray<FKey>> GetKnownActionMappings(EControllerHand InHand = EControllerHand::AnyHand, FName InHMDDeviceType = NAME_None) const;
+
 	/** Initialize default values */
 	UFUNCTION( BlueprintNativeEvent, CallInEditor, Category = "UVREditorInteractor" )
 	void Init( class UVREditorMode* InVRMode );
@@ -74,6 +76,7 @@ public:
 	UFUNCTION( BlueprintNativeEvent, CallInEditor, Category = "UVREditorInteractor" )
 	void SetupComponent( AActor* OwningActor );
 
+
 	// ViewportInteractorInterface overrides
 	virtual void Shutdown_Implementation() override;
 	virtual void Tick_Implementation( const float DeltaTime ) override;
@@ -95,7 +98,7 @@ public:
 	void HandleInputAxis( FEditorViewportClient& ViewportClient, FViewportActionKeyInput& Action, const FKey Key, const float Delta, const float DeltaTime, bool& bOutWasHandled );
 
 	/** Toggles whether or not this controller is being used to scrub sequencer */
-	void ToggleSequencerScrubbingMode();;
+	void ToggleSequencerScrubbingMode();
 
 	/** Returns whether or not this controller is being used to scrub sequencer */
 	bool IsScrubbingSequencer() const;
diff --git a/Engine/Source/Editor/VREditor/Public/VREditorMode.h b/Engine/Source/Editor/VREditor/Public/VREditorMode.h
--- a/Engine/Source/Editor/VREditor/Public/VREditorMode.h
+++ b/Engine/Source/Editor/VREditor/Public/VREditorMode.h
@@ -35,7 +35,7 @@ namespace VRActionTypes
 {
 	static const FName Touch( "Touch" );
 	static const FName Modifier( "Modifier" );
-	static const FName Modifier2( "Modifier2" ); //Only used by Oculus
+	static const FName Modifier2( "Modifier2" );
 	static const FName ConfirmRadialSelection( "ConfirmRadialSelection" );
 	static const FName TrackpadPositionX( "TrackpadPosition_X" );
 	static const FName TrackpadPositionY( "TrackpadPosition_Y" );
diff --git a/Engine/Source/Editor/VREditor/VREditorInteractor.cpp b/Engine/Source/Editor/VREditor/VREditorInteractor.cpp
--- a/Engine/Source/Editor/VREditor/VREditorInteractor.cpp
+++ b/Engine/Source/Editor/VREditor/VREditorInteractor.cpp
@@ -82,6 +82,7 @@ namespace VREditorKeyNames
 
 static const FName OculusDeviceType( TEXT( "OculusHMD" ) );
 static const FName SteamVRDeviceType( TEXT( "SteamVR" ) );
+static const FName OpenXRDeviceType( TEXT( "OpenXR" ) );
 
 const FName UVREditorInteractor::TrackpadPositionX = FName( "TrackpadPosition_X" );
 const FName UVREditorInteractor::TrackpadPositionY = FName( "TrackpadPosition_Y" );
@@ -141,67 +142,145 @@ UVREditorInteractor::UVREditorInteractor() :
 {
 }
 
+TMap<FViewportActionKeyInput, TArray<FKey>> UVREditorInteractor::GetKnownActionMappings(
+	EControllerHand InHand /* = EControllerHand::AnyHand */,
+	FName InHMDDeviceType /* = NAME_None */
+) const
+{
+	TMap<FViewportActionKeyInput, TArray<FKey>> ReturnMap;
+
+	if (InHand == EControllerHand::Left || InHand == EControllerHand::AnyHand)
+	{
+		if (InHMDDeviceType == SteamVRDeviceType || InHMDDeviceType == OpenXRDeviceType || InHMDDeviceType == NAME_None)
+		{
+			// HTC Vive
+			ReturnMap.FindOrAdd(FViewportActionKeyInput(ViewportWorldActionTypes::WorldMovement)).Add(EKeys::Vive_Left_Grip_Click);
+			ReturnMap.FindOrAdd(FViewportActionKeyInput(VRActionTypes::Modifier)).Add(EKeys::Vive_Left_Menu_Click);
+
+			ReturnMap.FindOrAdd(FViewportActionKeyInput(VRActionTypes::TrackpadDown)).Add(EKeys::Vive_Left_Trackpad_Down); // down
+			ReturnMap.FindOrAdd(FViewportActionKeyInput(VRActionTypes::TrackpadLeft)).Add(EKeys::Vive_Left_Trackpad_Left);
+			ReturnMap.FindOrAdd(FViewportActionKeyInput(VRActionTypes::TrackpadRight)).Add(EKeys::Vive_Left_Trackpad_Right);
+			ReturnMap.FindOrAdd(FViewportActionKeyInput(VRActionTypes::TrackpadUp)).Add(EKeys::Vive_Left_Trackpad_Up);
+
+			ReturnMap.FindOrAdd(FViewportActionKeyInput::Axis(UVREditorInteractor::TriggerAxis)).Add(EKeys::Vive_Left_Trigger_Axis);
+			ReturnMap.FindOrAdd(FViewportActionKeyInput::Axis(UVREditorInteractor::TrackpadPositionX)).Add(EKeys::Vive_Left_Trackpad_X);
+			ReturnMap.FindOrAdd(FViewportActionKeyInput::Axis(UVREditorInteractor::TrackpadPositionY)).Add(EKeys::Vive_Left_Trackpad_Y);
+			ReturnMap.FindOrAdd(FViewportActionKeyInput(VRActionTypes::ConfirmRadialSelection)).Add(EKeys::Vive_Left_Trackpad_Click);
+			ReturnMap.FindOrAdd(FViewportActionKeyInput(VRActionTypes::Touch)).Add(EKeys::Vive_Left_Trackpad_Touch);
+
+			// Valve Index
+			ReturnMap.FindOrAdd(FViewportActionKeyInput(ViewportWorldActionTypes::WorldMovement)).Add(EKeys::ValveIndex_Left_Trackpad_Touch);
+			ReturnMap.FindOrAdd(FViewportActionKeyInput(VRActionTypes::Modifier)).Add(EKeys::ValveIndex_Left_A_Click);
+			ReturnMap.FindOrAdd(FViewportActionKeyInput(VRActionTypes::Modifier2)).Add(EKeys::ValveIndex_Left_B_Click);
+			ReturnMap.FindOrAdd(FViewportActionKeyInput(VRActionTypes::Touch)).Add(EKeys::ValveIndex_Left_Thumbstick_Touch);
+
+			ReturnMap.FindOrAdd(FViewportActionKeyInput(VRActionTypes::TrackpadDown)).Add(EKeys::ValveIndex_Left_Thumbstick_Down); // down
+			ReturnMap.FindOrAdd(FViewportActionKeyInput(VRActionTypes::TrackpadUp)).Add(EKeys::ValveIndex_Left_Thumbstick_Up);
+			ReturnMap.FindOrAdd(FViewportActionKeyInput(VRActionTypes::TrackpadLeft)).Add(EKeys::ValveIndex_Left_Thumbstick_Left);
+			ReturnMap.FindOrAdd(FViewportActionKeyInput(VRActionTypes::TrackpadRight)).Add(EKeys::ValveIndex_Left_Thumbstick_Right);
+
+			ReturnMap.FindOrAdd(FViewportActionKeyInput::Axis(UVREditorInteractor::TriggerAxis)).Add(EKeys::ValveIndex_Left_Trigger_Axis);
+			ReturnMap.FindOrAdd(FViewportActionKeyInput::Axis(UVREditorInteractor::TrackpadPositionX)).Add(EKeys::ValveIndex_Left_Thumbstick_X);
+			ReturnMap.FindOrAdd(FViewportActionKeyInput::Axis(UVREditorInteractor::TrackpadPositionY)).Add(EKeys::ValveIndex_Left_Thumbstick_Y);
+			ReturnMap.FindOrAdd(FViewportActionKeyInput(VRActionTypes::ConfirmRadialSelection)).Add(EKeys::ValveIndex_Left_Thumbstick_Click);
+		}
+
+		if (InHMDDeviceType == OculusDeviceType || InHMDDeviceType == OpenXRDeviceType || InHMDDeviceType == NAME_None)
+		{
+			ReturnMap.FindOrAdd(FViewportActionKeyInput(ViewportWorldActionTypes::WorldMovement)).Add(EKeys::OculusTouch_Left_Grip_Click);
+			ReturnMap.FindOrAdd(FViewportActionKeyInput(VRActionTypes::Modifier)).Add(EKeys::OculusTouch_Left_X_Click);
+			ReturnMap.FindOrAdd(FViewportActionKeyInput(VRActionTypes::Modifier2)).Add(EKeys::OculusTouch_Left_Y_Click);
+			ReturnMap.FindOrAdd(FViewportActionKeyInput(VRActionTypes::Touch)).Add(EKeys::OculusTouch_Left_Thumbstick_Touch);
+
+			ReturnMap.FindOrAdd(FViewportActionKeyInput(VRActionTypes::TrackpadDown)).Add(EKeys::OculusTouch_Left_Thumbstick_Down); // down
+			ReturnMap.FindOrAdd(FViewportActionKeyInput(VRActionTypes::TrackpadUp)).Add(EKeys::OculusTouch_Left_Thumbstick_Up);
+			ReturnMap.FindOrAdd(FViewportActionKeyInput(VRActionTypes::TrackpadLeft)).Add(EKeys::OculusTouch_Left_Thumbstick_Left);
+			ReturnMap.FindOrAdd(FViewportActionKeyInput(VRActionTypes::TrackpadRight)).Add(EKeys::OculusTouch_Left_Thumbstick_Right);
+
+			ReturnMap.FindOrAdd(FViewportActionKeyInput::Axis(UVREditorInteractor::TriggerAxis)).Add(EKeys::OculusTouch_Left_Trigger_Axis);
+			ReturnMap.FindOrAdd(FViewportActionKeyInput::Axis(UVREditorInteractor::TrackpadPositionX)).Add(EKeys::OculusTouch_Left_Thumbstick_X);
+			ReturnMap.FindOrAdd(FViewportActionKeyInput::Axis(UVREditorInteractor::TrackpadPositionY)).Add(EKeys::OculusTouch_Left_Thumbstick_Y);
+			ReturnMap.FindOrAdd(FViewportActionKeyInput(VRActionTypes::ConfirmRadialSelection)).Add(EKeys::OculusTouch_Left_Thumbstick_Click);
+		}
+	}
+
+	if (InHand == EControllerHand::Right || InHand == EControllerHand::AnyHand)
+	{
+		if (InHMDDeviceType == SteamVRDeviceType || InHMDDeviceType == OpenXRDeviceType || InHMDDeviceType == NAME_None)
+		{
+			// HTC Vive
+			ReturnMap.FindOrAdd(FViewportActionKeyInput(ViewportWorldActionTypes::WorldMovement)).Add(EKeys::Vive_Right_Grip_Click);
+			ReturnMap.FindOrAdd(FViewportActionKeyInput(VRActionTypes::Modifier)).Add(EKeys::Vive_Right_Menu_Click);
+
+			ReturnMap.FindOrAdd(FViewportActionKeyInput(VRActionTypes::TrackpadDown)).Add(EKeys::Vive_Right_Trackpad_Down); // down
+			ReturnMap.FindOrAdd(FViewportActionKeyInput(VRActionTypes::TrackpadLeft)).Add(EKeys::Vive_Right_Trackpad_Left);
+			ReturnMap.FindOrAdd(FViewportActionKeyInput(VRActionTypes::TrackpadRight)).Add(EKeys::Vive_Right_Trackpad_Right);
+			ReturnMap.FindOrAdd(FViewportActionKeyInput(VRActionTypes::TrackpadUp)).Add(EKeys::Vive_Right_Trackpad_Up);
+
+			ReturnMap.FindOrAdd(FViewportActionKeyInput::Axis(UVREditorInteractor::TriggerAxis)).Add(EKeys::Vive_Right_Trigger_Axis);
+			ReturnMap.FindOrAdd(FViewportActionKeyInput::Axis(UVREditorInteractor::TrackpadPositionX)).Add(EKeys::Vive_Right_Trackpad_X);
+			ReturnMap.FindOrAdd(FViewportActionKeyInput::Axis(UVREditorInteractor::TrackpadPositionY)).Add(EKeys::Vive_Right_Trackpad_Y);
+			ReturnMap.FindOrAdd(FViewportActionKeyInput(VRActionTypes::ConfirmRadialSelection)).Add(EKeys::Vive_Right_Trackpad_Click);
+			ReturnMap.FindOrAdd(FViewportActionKeyInput(VRActionTypes::Touch)).Add(EKeys::Vive_Right_Trackpad_Touch);
+
+			// Valve Index
+			ReturnMap.FindOrAdd(FViewportActionKeyInput(ViewportWorldActionTypes::WorldMovement)).Add(EKeys::ValveIndex_Right_Trackpad_Touch);
+			ReturnMap.FindOrAdd(FViewportActionKeyInput(VRActionTypes::Modifier)).Add(EKeys::ValveIndex_Right_A_Click);
+			ReturnMap.FindOrAdd(FViewportActionKeyInput(VRActionTypes::Modifier2)).Add(EKeys::ValveIndex_Right_B_Click);
+			ReturnMap.FindOrAdd(FViewportActionKeyInput(VRActionTypes::Touch)).Add(EKeys::ValveIndex_Right_Thumbstick_Touch);
+
+			ReturnMap.FindOrAdd(FViewportActionKeyInput(VRActionTypes::TrackpadDown)).Add(EKeys::ValveIndex_Right_Thumbstick_Down); // down
+			ReturnMap.FindOrAdd(FViewportActionKeyInput(VRActionTypes::TrackpadUp)).Add(EKeys::ValveIndex_Right_Thumbstick_Up);
+			ReturnMap.FindOrAdd(FViewportActionKeyInput(VRActionTypes::TrackpadLeft)).Add(EKeys::ValveIndex_Right_Thumbstick_Left);
+			ReturnMap.FindOrAdd(FViewportActionKeyInput(VRActionTypes::TrackpadRight)).Add(EKeys::ValveIndex_Right_Thumbstick_Right);
+
+			ReturnMap.FindOrAdd(FViewportActionKeyInput::Axis(UVREditorInteractor::TriggerAxis)).Add(EKeys::ValveIndex_Right_Trigger_Axis);
+			ReturnMap.FindOrAdd(FViewportActionKeyInput::Axis(UVREditorInteractor::TrackpadPositionX)).Add(EKeys::ValveIndex_Right_Thumbstick_X);
+			ReturnMap.FindOrAdd(FViewportActionKeyInput::Axis(UVREditorInteractor::TrackpadPositionY)).Add(EKeys::ValveIndex_Right_Thumbstick_Y);
+			ReturnMap.FindOrAdd(FViewportActionKeyInput(VRActionTypes::ConfirmRadialSelection)).Add(EKeys::ValveIndex_Right_Thumbstick_Click);
+		}
+
+		if (InHMDDeviceType == OculusDeviceType || InHMDDeviceType == OpenXRDeviceType || InHMDDeviceType == NAME_None)
+		{
+			ReturnMap.FindOrAdd(FViewportActionKeyInput(ViewportWorldActionTypes::WorldMovement)).Add(EKeys::OculusTouch_Right_Grip_Click);
+			ReturnMap.FindOrAdd(FViewportActionKeyInput(VRActionTypes::Modifier)).Add(EKeys::OculusTouch_Right_A_Click);
+			ReturnMap.FindOrAdd(FViewportActionKeyInput(VRActionTypes::Modifier2)).Add(EKeys::OculusTouch_Right_B_Click);
+			ReturnMap.FindOrAdd(FViewportActionKeyInput(VRActionTypes::Touch)).Add(EKeys::OculusTouch_Right_Thumbstick_Touch);
+
+			ReturnMap.FindOrAdd(FViewportActionKeyInput(VRActionTypes::TrackpadDown)).Add(EKeys::OculusTouch_Right_Thumbstick_Down); // down
+			ReturnMap.FindOrAdd(FViewportActionKeyInput(VRActionTypes::TrackpadUp)).Add(EKeys::OculusTouch_Right_Thumbstick_Up);
+			ReturnMap.FindOrAdd(FViewportActionKeyInput(VRActionTypes::TrackpadLeft)).Add(EKeys::OculusTouch_Right_Thumbstick_Left);
+			ReturnMap.FindOrAdd(FViewportActionKeyInput(VRActionTypes::TrackpadRight)).Add(EKeys::OculusTouch_Right_Thumbstick_Right);
+
+			ReturnMap.FindOrAdd(FViewportActionKeyInput::Axis(UVREditorInteractor::TriggerAxis)).Add(EKeys::OculusTouch_Right_Trigger_Axis);
+			ReturnMap.FindOrAdd(FViewportActionKeyInput::Axis(UVREditorInteractor::TrackpadPositionX)).Add(EKeys::OculusTouch_Right_Thumbstick_X);
+			ReturnMap.FindOrAdd(FViewportActionKeyInput::Axis(UVREditorInteractor::TrackpadPositionY)).Add(EKeys::OculusTouch_Right_Thumbstick_Y);
+			ReturnMap.FindOrAdd(FViewportActionKeyInput(VRActionTypes::ConfirmRadialSelection)).Add(EKeys::OculusTouch_Right_Thumbstick_Click);
+		}
+	}
+
+	return ReturnMap;
+}
+
 void UVREditorInteractor::Init_Implementation( UVREditorMode* InVRMode )
 {
 	VRMode = InVRMode;
 	KeyToActionMap.Reset();
 
 	const FName HMDDeviceType = InVRMode->GetHMDDeviceType();
+
 	// Setup keys
 	if (ControllerMotionSource == FXRMotionControllerBase::LeftHandSourceId)
 	{
 		AddKeyAction( UVREditorInteractor::MotionController_Left_FullyPressedTriggerAxis, FViewportActionKeyInput( ViewportWorldActionTypes::SelectAndMove_FullyPressed ) );
 		AddKeyAction( UVREditorInteractor::MotionController_Left_PressedTriggerAxis, FViewportActionKeyInput( ViewportWorldActionTypes::SelectAndMove ) );
 
-		if (HMDDeviceType == SteamVRDeviceType)
-		{
-			// HTC Vive
-			AddKeyAction( EKeys::Vive_Left_Grip_Click, FViewportActionKeyInput( ViewportWorldActionTypes::WorldMovement ) );
-			AddKeyAction( EKeys::Vive_Left_Menu_Click, FViewportActionKeyInput( VRActionTypes::Modifier ) );
-
-			AddKeyAction( EKeys::Vive_Left_Trackpad_Down, FViewportActionKeyInput( VRActionTypes::TrackpadDown ) ); // down
-			AddKeyAction( EKeys::Vive_Left_Trackpad_Left, FViewportActionKeyInput( VRActionTypes::TrackpadLeft ) );
-			AddKeyAction( EKeys::Vive_Left_Trackpad_Right, FViewportActionKeyInput( VRActionTypes::TrackpadRight ) );
-			AddKeyAction( EKeys::Vive_Left_Trackpad_Up, FViewportActionKeyInput( VRActionTypes::TrackpadUp ) );
-
-			AddKeyAction( EKeys::Vive_Left_Trigger_Axis, FViewportActionKeyInput( UVREditorInteractor::TriggerAxis ) );
-			AddKeyAction( EKeys::Vive_Left_Trackpad_X, FViewportActionKeyInput( UVREditorInteractor::TrackpadPositionX ) );
-			AddKeyAction( EKeys::Vive_Left_Trackpad_Y, FViewportActionKeyInput( UVREditorInteractor::TrackpadPositionY ) );
-			AddKeyAction( EKeys::Vive_Left_Trackpad_Click, FViewportActionKeyInput( VRActionTypes::ConfirmRadialSelection ) );
-			AddKeyAction( EKeys::Vive_Left_Trackpad_Touch, FViewportActionKeyInput( VRActionTypes::Touch ) );
-
-			// Valve Index
-			AddKeyAction( EKeys::ValveIndex_Left_Trackpad_Touch, FViewportActionKeyInput( ViewportWorldActionTypes::WorldMovement ) );
-			AddKeyAction( EKeys::ValveIndex_Left_A_Click, FViewportActionKeyInput( VRActionTypes::Modifier ) );
-			AddKeyAction( EKeys::ValveIndex_Left_B_Click, FViewportActionKeyInput( VRActionTypes::Modifier2 ) );
-			AddKeyAction( EKeys::ValveIndex_Left_Thumbstick_Touch, FViewportActionKeyInput( VRActionTypes::Touch ) );
-
-			AddKeyAction( EKeys::ValveIndex_Left_Thumbstick_Down, FViewportActionKeyInput( VRActionTypes::TrackpadDown ) ); // down
-			AddKeyAction( EKeys::ValveIndex_Left_Thumbstick_Up, FViewportActionKeyInput( VRActionTypes::TrackpadUp ) );
-			AddKeyAction( EKeys::ValveIndex_Left_Thumbstick_Left, FViewportActionKeyInput( VRActionTypes::TrackpadLeft ) );
-			AddKeyAction( EKeys::ValveIndex_Left_Thumbstick_Right, FViewportActionKeyInput( VRActionTypes::TrackpadRight ) );
-
-			AddKeyAction( EKeys::ValveIndex_Left_Trigger_Axis, FViewportActionKeyInput( UVREditorInteractor::TriggerAxis ) );
-			AddKeyAction( EKeys::ValveIndex_Left_Thumbstick_X, FViewportActionKeyInput( UVREditorInteractor::TrackpadPositionX ) );
-			AddKeyAction( EKeys::ValveIndex_Left_Thumbstick_Y, FViewportActionKeyInput( UVREditorInteractor::TrackpadPositionY ) );
-			AddKeyAction( EKeys::ValveIndex_Left_Thumbstick_Click, FViewportActionKeyInput( VRActionTypes::ConfirmRadialSelection ) );
-		}
-		else if (HMDDeviceType == OculusDeviceType)
+		for (const TPair<FViewportActionKeyInput, TArray<FKey>>& ActionMappings : GetKnownActionMappings(EControllerHand::Left, HMDDeviceType))
 		{
-			AddKeyAction( EKeys::OculusTouch_Left_Grip_Click, FViewportActionKeyInput( ViewportWorldActionTypes::WorldMovement ) );
-			AddKeyAction( EKeys::OculusTouch_Left_X_Click, FViewportActionKeyInput( VRActionTypes::Modifier ) );
-			AddKeyAction( EKeys::OculusTouch_Left_Y_Click, FViewportActionKeyInput( VRActionTypes::Modifier2 ) );
-			AddKeyAction( EKeys::OculusTouch_Left_Thumbstick_Touch, FViewportActionKeyInput( VRActionTypes::Touch ) );
-
-			AddKeyAction( EKeys::OculusTouch_Left_Thumbstick_Down, FViewportActionKeyInput( VRActionTypes::TrackpadDown ) ); // down
-			AddKeyAction( EKeys::OculusTouch_Left_Thumbstick_Up, FViewportActionKeyInput( VRActionTypes::TrackpadUp ) );
-			AddKeyAction( EKeys::OculusTouch_Left_Thumbstick_Left, FViewportActionKeyInput( VRActionTypes::TrackpadLeft ) );
-			AddKeyAction( EKeys::OculusTouch_Left_Thumbstick_Right, FViewportActionKeyInput( VRActionTypes::TrackpadRight ) );
-
-			AddKeyAction( EKeys::OculusTouch_Left_Trigger_Axis, FViewportActionKeyInput( UVREditorInteractor::TriggerAxis ) );
-			AddKeyAction( EKeys::OculusTouch_Left_Thumbstick_X, FViewportActionKeyInput( UVREditorInteractor::TrackpadPositionX ) );
-			AddKeyAction( EKeys::OculusTouch_Left_Thumbstick_Y, FViewportActionKeyInput( UVREditorInteractor::TrackpadPositionY ) );
-			AddKeyAction( EKeys::OculusTouch_Left_Thumbstick_Click, FViewportActionKeyInput( VRActionTypes::ConfirmRadialSelection ) );
+			for (const FKey& Key : ActionMappings.Value)
+			{
+				AddKeyAction(Key, ActionMappings.Key);
+			}
 		}
 	}
 	else if (ControllerMotionSource == FXRMotionControllerBase::RightHandSourceId)
@@ -209,57 +288,15 @@ void UVREditorInteractor::Init_Implementation( UVREditorMode* InVRMode )
 		AddKeyAction( UVREditorInteractor::MotionController_Right_FullyPressedTriggerAxis, FViewportActionKeyInput( ViewportWorldActionTypes::SelectAndMove_FullyPressed ) );
 		AddKeyAction( UVREditorInteractor::MotionController_Right_PressedTriggerAxis, FViewportActionKeyInput( ViewportWorldActionTypes::SelectAndMove ) );
 
-		if (HMDDeviceType == SteamVRDeviceType)
-		{
-			// HTC Vive
-			AddKeyAction( EKeys::Vive_Right_Grip_Click, FViewportActionKeyInput( ViewportWorldActionTypes::WorldMovement ) );
-			AddKeyAction( EKeys::Vive_Right_Menu_Click, FViewportActionKeyInput( VRActionTypes::Modifier ) );
-
-			AddKeyAction( EKeys::Vive_Right_Trackpad_Down, FViewportActionKeyInput( VRActionTypes::TrackpadDown ) ); // down
-			AddKeyAction( EKeys::Vive_Right_Trackpad_Left, FViewportActionKeyInput( VRActionTypes::TrackpadLeft ) );
-			AddKeyAction( EKeys::Vive_Right_Trackpad_Right, FViewportActionKeyInput( VRActionTypes::TrackpadRight ) );
-			AddKeyAction( EKeys::Vive_Right_Trackpad_Up, FViewportActionKeyInput( VRActionTypes::TrackpadUp ) );
-			
-			AddKeyAction( EKeys::Vive_Right_Trigger_Axis, FViewportActionKeyInput( UVREditorInteractor::TriggerAxis ) );
-			AddKeyAction( EKeys::Vive_Right_Trackpad_X, FViewportActionKeyInput( UVREditorInteractor::TrackpadPositionX ) );
-			AddKeyAction( EKeys::Vive_Right_Trackpad_Y, FViewportActionKeyInput( UVREditorInteractor::TrackpadPositionY ) );
-			AddKeyAction( EKeys::Vive_Right_Trackpad_Click, FViewportActionKeyInput( VRActionTypes::ConfirmRadialSelection ) );
-			AddKeyAction( EKeys::Vive_Right_Trackpad_Touch, FViewportActionKeyInput( VRActionTypes::Touch ) );
-			
-			// Valve Index
-			AddKeyAction( EKeys::ValveIndex_Right_Trackpad_Touch, FViewportActionKeyInput( ViewportWorldActionTypes::WorldMovement ) );
-			AddKeyAction( EKeys::ValveIndex_Right_A_Click, FViewportActionKeyInput( VRActionTypes::Modifier ) );
-			AddKeyAction( EKeys::ValveIndex_Right_B_Click, FViewportActionKeyInput( VRActionTypes::Modifier2 ) );
-			AddKeyAction( EKeys::ValveIndex_Right_Thumbstick_Touch, FViewportActionKeyInput( VRActionTypes::Touch ) );
-
-			AddKeyAction( EKeys::ValveIndex_Right_Thumbstick_Down, FViewportActionKeyInput( VRActionTypes::TrackpadDown ) ); // down
-			AddKeyAction( EKeys::ValveIndex_Right_Thumbstick_Up, FViewportActionKeyInput( VRActionTypes::TrackpadUp ) );
-			AddKeyAction( EKeys::ValveIndex_Right_Thumbstick_Left, FViewportActionKeyInput( VRActionTypes::TrackpadLeft ) );
-			AddKeyAction( EKeys::ValveIndex_Right_Thumbstick_Right, FViewportActionKeyInput( VRActionTypes::TrackpadRight ) );
-			
-			AddKeyAction( EKeys::ValveIndex_Right_Trigger_Axis, FViewportActionKeyInput( UVREditorInteractor::TriggerAxis ) );
-			AddKeyAction( EKeys::ValveIndex_Right_Thumbstick_X, FViewportActionKeyInput( UVREditorInteractor::TrackpadPositionX ) );
-			AddKeyAction( EKeys::ValveIndex_Right_Thumbstick_Y, FViewportActionKeyInput( UVREditorInteractor::TrackpadPositionY ) );
-			AddKeyAction( EKeys::ValveIndex_Right_Thumbstick_Click, FViewportActionKeyInput( VRActionTypes::ConfirmRadialSelection ) );
-		}
-		else if (HMDDeviceType == OculusDeviceType)
+		for (const TPair<FViewportActionKeyInput, TArray<FKey>>& ActionMappings : GetKnownActionMappings(EControllerHand::Right, HMDDeviceType))
 		{
-			AddKeyAction( EKeys::OculusTouch_Right_Grip_Click, FViewportActionKeyInput( ViewportWorldActionTypes::WorldMovement ) );
-			AddKeyAction( EKeys::OculusTouch_Right_A_Click, FViewportActionKeyInput( VRActionTypes::Modifier ) );
-			AddKeyAction( EKeys::OculusTouch_Right_B_Click, FViewportActionKeyInput( VRActionTypes::Modifier2 ) );
-			AddKeyAction( EKeys::OculusTouch_Right_Thumbstick_Touch, FViewportActionKeyInput( VRActionTypes::Touch ) );
-
-			AddKeyAction( EKeys::OculusTouch_Right_Thumbstick_Down, FViewportActionKeyInput( VRActionTypes::TrackpadDown ) ); // down
-			AddKeyAction( EKeys::OculusTouch_Right_Thumbstick_Up, FViewportActionKeyInput( VRActionTypes::TrackpadUp ) );
-			AddKeyAction( EKeys::OculusTouch_Right_Thumbstick_Left, FViewportActionKeyInput( VRActionTypes::TrackpadLeft ) );
-			AddKeyAction( EKeys::OculusTouch_Right_Thumbstick_Right, FViewportActionKeyInput( VRActionTypes::TrackpadRight ) );
-			
-			AddKeyAction( EKeys::OculusTouch_Right_Trigger_Axis, FViewportActionKeyInput( UVREditorInteractor::TriggerAxis ) );
-			AddKeyAction( EKeys::OculusTouch_Right_Thumbstick_X, FViewportActionKeyInput( UVREditorInteractor::TrackpadPositionX ) );
-			AddKeyAction( EKeys::OculusTouch_Right_Thumbstick_Y, FViewportActionKeyInput( UVREditorInteractor::TrackpadPositionY ) );
-			AddKeyAction( EKeys::OculusTouch_Right_Thumbstick_Click, FViewportActionKeyInput( VRActionTypes::ConfirmRadialSelection ) );
+			for (const FKey& Key : ActionMappings.Value)
+			{
+				AddKeyAction(Key, ActionMappings.Key);
+			}
 		}
 	}
+
 	bHaveMotionController = true;
 }
 
@@ -726,8 +763,6 @@ void UVREditorInteractor::CalculateDragRay( float& InOutDragRayLength, float& In
 
 		if (!FMath::IsNearlyZero( SlideDelta ))
 		{
-
-
 			InOutDragRayLength += SlideDelta;
 
 			InOutDragRayVelocity = 0.0f;
@@ -1204,7 +1239,7 @@ void UVREditorInteractor::ApplyButtonPressColors( const FViewportActionKeyInput&
 		SetMotionControllerButtonPressedVisuals( Event, StaticModifierParameter, PressStrength );
 	}
 
-	if (GetVRMode().GetHMDDeviceType() == OculusDeviceType && ActionType == VRActionTypes::Modifier2)
+	if (ActionType == VRActionTypes::Modifier2)
 	{
 		static FName StaticModifierParameter( "B5" );
 		SetMotionControllerButtonPressedVisuals( Event, StaticModifierParameter, PressStrength );
diff --git a/Engine/Source/Editor/VREditor/VREditorModeManager.cpp b/Engine/Source/Editor/VREditor/VREditorModeManager.cpp
--- a/Engine/Source/Editor/VREditor/VREditorModeManager.cpp
+++ b/Engine/Source/Editor/VREditor/VREditorModeManager.cpp
@@ -134,14 +134,13 @@ bool FVREditorModeManager::IsVREditorActive() const
 }
 
 const static FName WMRSytemName = FName(TEXT("WindowsMixedRealityHMD"));
-const static FName OXRSytemName = FName(TEXT("OpenXR"));
 bool FVREditorModeManager::IsVREditorAvailable() const
 {
 	if (GEngine->XRSystem.IsValid() && GEngine->XRSystem->GetHMDDevice() && GEngine->XRSystem->GetHMDDevice()->IsHMDEnabled())
 	{
-		// TODO: UE-71871/UE-73237 Work around for avoiding starting VRMode when using WMR
+		// TODO: UE-71871 Work around for avoiding starting VRMode when using WMR
 		FName SystemName = GEngine->XRSystem->GetSystemName();
-		const bool bIsWMR = SystemName == WMRSytemName || SystemName == OXRSytemName;
+		const bool bIsWMR = SystemName == WMRSytemName;
 		return !bIsWMR && !GEditor->IsPlayingSessionInEditor();
 	}
 	else
diff --git a/Engine/Source/Editor/ViewportInteraction/Public/ViewportInteractionTypes.h b/Engine/Source/Editor/ViewportInteraction/Public/ViewportInteractionTypes.h
--- a/Engine/Source/Editor/ViewportInteraction/Public/ViewportInteractionTypes.h
+++ b/Engine/Source/Editor/ViewportInteraction/Public/ViewportInteractionTypes.h
@@ -15,18 +15,27 @@ struct VIEWPORTINTERACTION_API FViewportActionKeyInput
 {
 	GENERATED_BODY()
 	
-	FViewportActionKeyInput() :
-		ActionType( NAME_None ),
-		Event( IE_Pressed ),
-		bIsInputCaptured( false )
+	FViewportActionKeyInput()
+		: ActionType( NAME_None )
+		, Event( IE_Pressed )
+		, bIsInputCaptured( false )
+		, bIsAxis( false )
 	{}
 
-	FViewportActionKeyInput( const FName& InActionType ) : 
-		ActionType( InActionType ),
-		Event( IE_Pressed ),
-		bIsInputCaptured( false )
+	FViewportActionKeyInput( const FName& InActionType )
+		: ActionType( InActionType )
+		, Event( IE_Pressed )
+		, bIsInputCaptured( false )
+		, bIsAxis( false )
 	{}
 
+	static FViewportActionKeyInput Axis(const FName& InActionType)
+	{
+		FViewportActionKeyInput Input(InActionType);
+		Input.bIsAxis = true;
+		return Input;
+	}
+
 	/** The name of this action */
 	UPROPERTY( BlueprintReadOnly, Category = Action )
 	FName ActionType;
@@ -39,6 +48,26 @@ struct VIEWPORTINTERACTION_API FViewportActionKeyInput
 	handle input events until it is no longer captured.  It's the captors responsibility to set this using OnVRAction(), or clear it when finished with capturing. */
 	UPROPERTY()
 	bool bIsInputCaptured;
+
+	UPROPERTY()
+	bool bIsAxis;
+
+	bool operator==(const FViewportActionKeyInput& Other) const
+	{
+		return ActionType == Other.ActionType &&
+			Event == Other.Event &&
+			bIsInputCaptured == Other.bIsInputCaptured &&
+			bIsAxis == Other.bIsAxis;
+	}
+
+	friend uint32 GetTypeHash(const FViewportActionKeyInput& Input)
+	{
+		uint32 Hash = GetTypeHash(Input.ActionType);
+		Hash = HashCombine(Hash, GetTypeHash(Input.Event));
+		Hash = HashCombine(Hash, GetTypeHash(Input.bIsInputCaptured));
+		Hash = HashCombine(Hash, GetTypeHash(Input.bIsAxis));
+		return Hash;
+	}
 };
 
 
diff --git a/Engine/Source/Editor/ViewportInteraction/Public/ViewportInteractor.h b/Engine/Source/Editor/ViewportInteraction/Public/ViewportInteractor.h
--- a/Engine/Source/Editor/ViewportInteraction/Public/ViewportInteractor.h
+++ b/Engine/Source/Editor/ViewportInteraction/Public/ViewportInteractor.h
@@ -133,7 +133,6 @@ public:
 	/**
 	 * Gets the start and end point of the laser pointer for the specified hand
 	 *
-	 * @param HandIndex				Index of the hand to use
 	 * @param LasertPointerStart	(Out) The start location of the laser pointer in world space
 	 * @param LasertPointerEnd		(Out) The end location of the laser pointer in world space
 	 * @param bEvenIfBlocked		If true, returns a laser pointer even if the hand has UI in front of it (defaults to false)
@@ -198,7 +197,6 @@ public:
 	/**
 	 * Creates a hand transform and forward vector for a laser pointer for a given hand
 	 *
-	 * @param HandIndex			Index of the hand to use
 	 * @param OutHandTransform	The created hand transform
 	 * @param OutForwardVector	The forward vector of the hand
 	 *
diff --git a/Engine/Source/Editor/ViewportInteraction/ViewportInteractor.cpp b/Engine/Source/Editor/ViewportInteraction/ViewportInteractor.cpp
--- a/Engine/Source/Editor/ViewportInteraction/ViewportInteractor.cpp
+++ b/Engine/Source/Editor/ViewportInteraction/ViewportInteractor.cpp
@@ -119,6 +119,12 @@ bool UViewportInteractor::HandleInputKey( FEditorViewportClient& ViewportClient,
 	FViewportActionKeyInput* Action = KeyToActionMap.Find( Key );
 	if ( Action != nullptr )	// Ignore key repeats
 	{
+		if ( Action->bIsAxis )
+		{
+			UE_LOG(LogInput, Warning, TEXT("UViewportInteractor: %s registered as axis but handled as key"), *Action->ActionType.ToString());
+			Action->bIsAxis = false;
+		}
+
 		Action->Event = Event;
 
 		if( !bHandled )
@@ -524,6 +530,12 @@ bool UViewportInteractor::HandleInputAxis( FEditorViewportClient& ViewportClient
 	FViewportActionKeyInput* KnownAction = KeyToActionMap.Find( Key );
 	if ( KnownAction != nullptr )	// Ignore key repeats
 	{
+		if (!KnownAction->bIsAxis)
+		{
+			UE_LOG(LogInput, Warning, TEXT("UViewportInteractor: %s registered as key but handled as axis"), *KnownAction->ActionType.ToString());
+			KnownAction->bIsAxis = true;
+		}
+
 		FViewportActionKeyInput Action( KnownAction->ActionType );
 
 		if( !bHandled )
