VREditor: Change UIBrightness upper bound from ClampMax to UIMax to permit values >10.0 to be typed in.

diff --git a/Engine/Source/Editor/VREditor/Public/VRModeSettings.h b/Engine/Source/Editor/VREditor/Public/VRModeSettings.h
--- a/Engine/Source/Editor/VREditor/Public/VRModeSettings.h
+++ b/Engine/Source/Editor/VREditor/Public/VRModeSettings.h
@@ -61,7 +61,7 @@ public:
 	uint32 bShowWorldScaleProgressBar : 1;
 
 	/** Adjusts the brightness of the UI panels */
-	UPROPERTY(EditAnywhere, config, Category = "UI Customization", meta = (DisplayName = "UI Panel Brightness", ClampMin = 0.01, ClampMax = 10.0))
+	UPROPERTY(EditAnywhere, config, Category = "UI Customization", meta = (DisplayName = "UI Panel Brightness", ClampMin = 0.01, UIMax = 10.0))
 	float UIBrightness;
 
 	/** The size of the transform gizmo */
