Virtual Scouting: Bind "Touch" (trackpad / joystick capacitive sense) as axis for consistency with Oculus plugin behavior, even though underlying input component and key are boolean.

diff --git a/Engine/Source/Editor/VREditor/VREditorInteractor.cpp b/Engine/Source/Editor/VREditor/VREditorInteractor.cpp
--- a/Engine/Source/Editor/VREditor/VREditorInteractor.cpp
+++ b/Engine/Source/Editor/VREditor/VREditorInteractor.cpp
@@ -166,13 +166,13 @@ TMap<FViewportActionKeyInput, TArray<FKey>> UVREditorInteractor::GetKnownActionM
 			ReturnMap.FindOrAdd(FViewportActionKeyInput::Axis(UVREditorInteractor::TrackpadPositionX)).Add(EKeys::Vive_Left_Trackpad_X);
 			ReturnMap.FindOrAdd(FViewportActionKeyInput::Axis(UVREditorInteractor::TrackpadPositionY)).Add(EKeys::Vive_Left_Trackpad_Y);
 			ReturnMap.FindOrAdd(FViewportActionKeyInput(VRActionTypes::ConfirmRadialSelection)).Add(EKeys::Vive_Left_Trackpad_Click);
-			ReturnMap.FindOrAdd(FViewportActionKeyInput(VRActionTypes::Touch)).Add(EKeys::Vive_Left_Trackpad_Touch);
+			ReturnMap.FindOrAdd(FViewportActionKeyInput::Axis(VRActionTypes::Touch)).Add(EKeys::Vive_Left_Trackpad_Touch);
 
 			// Valve Index
 			ReturnMap.FindOrAdd(FViewportActionKeyInput(ViewportWorldActionTypes::WorldMovement)).Add(EKeys::ValveIndex_Left_Trackpad_Touch);
 			ReturnMap.FindOrAdd(FViewportActionKeyInput(VRActionTypes::Modifier)).Add(EKeys::ValveIndex_Left_A_Click);
 			ReturnMap.FindOrAdd(FViewportActionKeyInput(VRActionTypes::Modifier2)).Add(EKeys::ValveIndex_Left_B_Click);
-			ReturnMap.FindOrAdd(FViewportActionKeyInput(VRActionTypes::Touch)).Add(EKeys::ValveIndex_Left_Thumbstick_Touch);
+			ReturnMap.FindOrAdd(FViewportActionKeyInput::Axis(VRActionTypes::Touch)).Add(EKeys::ValveIndex_Left_Thumbstick_Touch);
 
 			ReturnMap.FindOrAdd(FViewportActionKeyInput(VRActionTypes::TrackpadDown)).Add(EKeys::ValveIndex_Left_Thumbstick_Down); // down
 			ReturnMap.FindOrAdd(FViewportActionKeyInput(VRActionTypes::TrackpadUp)).Add(EKeys::ValveIndex_Left_Thumbstick_Up);
@@ -190,7 +190,7 @@ TMap<FViewportActionKeyInput, TArray<FKey>> UVREditorInteractor::GetKnownActionM
 			ReturnMap.FindOrAdd(FViewportActionKeyInput(ViewportWorldActionTypes::WorldMovement)).Add(EKeys::OculusTouch_Left_Grip_Click);
 			ReturnMap.FindOrAdd(FViewportActionKeyInput(VRActionTypes::Modifier)).Add(EKeys::OculusTouch_Left_X_Click);
 			ReturnMap.FindOrAdd(FViewportActionKeyInput(VRActionTypes::Modifier2)).Add(EKeys::OculusTouch_Left_Y_Click);
-			ReturnMap.FindOrAdd(FViewportActionKeyInput(VRActionTypes::Touch)).Add(EKeys::OculusTouch_Left_Thumbstick_Touch);
+			ReturnMap.FindOrAdd(FViewportActionKeyInput::Axis(VRActionTypes::Touch)).Add(EKeys::OculusTouch_Left_Thumbstick_Touch);
 
 			ReturnMap.FindOrAdd(FViewportActionKeyInput(VRActionTypes::TrackpadDown)).Add(EKeys::OculusTouch_Left_Thumbstick_Down); // down
 			ReturnMap.FindOrAdd(FViewportActionKeyInput(VRActionTypes::TrackpadUp)).Add(EKeys::OculusTouch_Left_Thumbstick_Up);
@@ -221,13 +221,13 @@ TMap<FViewportActionKeyInput, TArray<FKey>> UVREditorInteractor::GetKnownActionM
 			ReturnMap.FindOrAdd(FViewportActionKeyInput::Axis(UVREditorInteractor::TrackpadPositionX)).Add(EKeys::Vive_Right_Trackpad_X);
 			ReturnMap.FindOrAdd(FViewportActionKeyInput::Axis(UVREditorInteractor::TrackpadPositionY)).Add(EKeys::Vive_Right_Trackpad_Y);
 			ReturnMap.FindOrAdd(FViewportActionKeyInput(VRActionTypes::ConfirmRadialSelection)).Add(EKeys::Vive_Right_Trackpad_Click);
-			ReturnMap.FindOrAdd(FViewportActionKeyInput(VRActionTypes::Touch)).Add(EKeys::Vive_Right_Trackpad_Touch);
+			ReturnMap.FindOrAdd(FViewportActionKeyInput::Axis(VRActionTypes::Touch)).Add(EKeys::Vive_Right_Trackpad_Touch);
 
 			// Valve Index
 			ReturnMap.FindOrAdd(FViewportActionKeyInput(ViewportWorldActionTypes::WorldMovement)).Add(EKeys::ValveIndex_Right_Trackpad_Touch);
 			ReturnMap.FindOrAdd(FViewportActionKeyInput(VRActionTypes::Modifier)).Add(EKeys::ValveIndex_Right_A_Click);
 			ReturnMap.FindOrAdd(FViewportActionKeyInput(VRActionTypes::Modifier2)).Add(EKeys::ValveIndex_Right_B_Click);
-			ReturnMap.FindOrAdd(FViewportActionKeyInput(VRActionTypes::Touch)).Add(EKeys::ValveIndex_Right_Thumbstick_Touch);
+			ReturnMap.FindOrAdd(FViewportActionKeyInput::Axis(VRActionTypes::Touch)).Add(EKeys::ValveIndex_Right_Thumbstick_Touch);
 
 			ReturnMap.FindOrAdd(FViewportActionKeyInput(VRActionTypes::TrackpadDown)).Add(EKeys::ValveIndex_Right_Thumbstick_Down); // down
 			ReturnMap.FindOrAdd(FViewportActionKeyInput(VRActionTypes::TrackpadUp)).Add(EKeys::ValveIndex_Right_Thumbstick_Up);
@@ -245,7 +245,7 @@ TMap<FViewportActionKeyInput, TArray<FKey>> UVREditorInteractor::GetKnownActionM
 			ReturnMap.FindOrAdd(FViewportActionKeyInput(ViewportWorldActionTypes::WorldMovement)).Add(EKeys::OculusTouch_Right_Grip_Click);
 			ReturnMap.FindOrAdd(FViewportActionKeyInput(VRActionTypes::Modifier)).Add(EKeys::OculusTouch_Right_A_Click);
 			ReturnMap.FindOrAdd(FViewportActionKeyInput(VRActionTypes::Modifier2)).Add(EKeys::OculusTouch_Right_B_Click);
-			ReturnMap.FindOrAdd(FViewportActionKeyInput(VRActionTypes::Touch)).Add(EKeys::OculusTouch_Right_Thumbstick_Touch);
+			ReturnMap.FindOrAdd(FViewportActionKeyInput::Axis(VRActionTypes::Touch)).Add(EKeys::OculusTouch_Right_Thumbstick_Touch);
 
 			ReturnMap.FindOrAdd(FViewportActionKeyInput(VRActionTypes::TrackpadDown)).Add(EKeys::OculusTouch_Right_Thumbstick_Down); // down
 			ReturnMap.FindOrAdd(FViewportActionKeyInput(VRActionTypes::TrackpadUp)).Add(EKeys::OculusTouch_Right_Thumbstick_Up);
