Added Enhanced input component lifecycle to VRScoutingInteractor and VRTool

diff --git a/Engine/Plugins/Experimental/VirtualProductionUtilities/Source/VPUtilitiesEditor/Private/VRTool.cpp b/Engine/Plugins/Experimental/VirtualProductionUtilities/Source/VPUtilitiesEditor/Private/VRTool.cpp
--- a/Engine/Plugins/Experimental/VirtualProductionUtilities/Source/VPUtilitiesEditor/Private/VRTool.cpp
+++ b/Engine/Plugins/Experimental/VirtualProductionUtilities/Source/VPUtilitiesEditor/Private/VRTool.cpp
@@ -1,4 +1,61 @@
 // Copyright Epic Games, Inc. All Rights Reserved.
 
 #include "VRTool.h"
+#include "GameFramework/InputSettings.h"
+#include "Engine/InputDelegateBinding.h"
 
+
+
+void AVRTool::OnConstruction(const FTransform& Transform)
+{
+	Super::OnConstruction(Transform);
+
+	CreateEditorInput();
+}
+
+void AVRTool::CreateEditorInput()
+{
+	if (bReceivesEditorInput && !HasAnyFlags(RF_ClassDefaultObject) && !EditorOnlyInputComponent)
+	{
+		EditorOnlyInputComponent = NewObject<UInputComponent>(this, UInputSettings::GetDefaultInputComponentClass(), TEXT("EdUtilActor_InputComponent0"), RF_Transient);
+		UInputDelegateBinding::BindInputDelegatesWithSubojects(this, EditorOnlyInputComponent);
+	}
+}
+
+void AVRTool::RemoveEditorInput()
+{
+	ensure(!bReceivesEditorInput);
+
+	if (EditorOnlyInputComponent)
+	{
+		EditorOnlyInputComponent->DestroyComponent();
+	}
+	EditorOnlyInputComponent = nullptr;
+}
+
+#if WITH_EDITOR
+void AVRTool::PostEditChangeProperty(FPropertyChangedEvent& PropertyChangedEvent)
+{
+	Super::PostEditChangeProperty(PropertyChangedEvent);
+
+	const FName PropertyName = (PropertyChangedEvent.Property != nullptr) ? PropertyChangedEvent.Property->GetFName() : NAME_None;
+	if (PropertyName == GET_MEMBER_NAME_CHECKED(AVRTool, bReceivesEditorInput))
+	{
+		SetReceivesEditorInput(bReceivesEditorInput);
+	}
+}
+#endif
+
+
+void AVRTool::SetReceivesEditorInput(bool bInValue)
+{
+	bReceivesEditorInput = bInValue;
+	if (bReceivesEditorInput)
+	{
+		CreateEditorInput();
+	}
+	else
+	{
+		RemoveEditorInput();
+	}
+}
\ No newline at end of file
diff --git a/Engine/Plugins/Experimental/VirtualProductionUtilities/Source/VPUtilitiesEditor/Public/VRTool.h b/Engine/Plugins/Experimental/VirtualProductionUtilities/Source/VPUtilitiesEditor/Public/VRTool.h
--- a/Engine/Plugins/Experimental/VirtualProductionUtilities/Source/VPUtilitiesEditor/Public/VRTool.h
+++ b/Engine/Plugins/Experimental/VirtualProductionUtilities/Source/VPUtilitiesEditor/Public/VRTool.h
@@ -14,4 +14,40 @@ class AVRTool : public AVPTransientEditorTickableActorBase
 {
 	GENERATED_BODY()
 
+	virtual void OnConstruction(const FTransform& Transform) override;
+#if WITH_EDITOR
+	virtual void PostEditChangeProperty(FPropertyChangedEvent& PropertyChangedEvent) override;
+#endif
+
+	/** Returns the current InputComponent on this utility actor. This will be NULL unless bReceivesEditorInput is set to true. */
+	UFUNCTION(BlueprintCallable, BlueprintPure, Category = "Input|Editor")
+	UInputComponent* GetInputComponent() const
+	{
+		return EditorOnlyInputComponent.Get();
+	}
+
+	UFUNCTION(BlueprintSetter, Category = "Input|Editor")
+	void SetReceivesEditorInput(bool bInValue);
+
+	UFUNCTION(BlueprintGetter, BlueprintPure, Category = "Input|Editor")
+	bool GetReceivesEditorInput() const
+	{
+		return bReceivesEditorInput;
+	}
+
+private:
+
+	/** Creates the EditorOnlyInputComponent if it does not already exist and registers all subobject callbacks to it */
+	void CreateEditorInput();
+
+	/** Removes the EditorOnlyInputComponent from this utility actor */
+	void RemoveEditorInput();
+
+	UPROPERTY(Transient, DuplicateTransient)
+	TObjectPtr<UInputComponent> EditorOnlyInputComponent;
+
+	/** If set to true, then this actor will be able to recieve input delegate callbacks when in the editor. */
+	UPROPERTY(EditAnywhere, Category = "Input|Editor", BlueprintSetter = SetReceivesEditorInput, BlueprintGetter = GetReceivesEditorInput)
+	bool bReceivesEditorInput = false;
+
 };
\ No newline at end of file
diff --git a/Engine/Source/Editor/VREditor/VRScoutingInteractor.cpp b/Engine/Source/Editor/VREditor/VRScoutingInteractor.cpp
--- a/Engine/Source/Editor/VREditor/VRScoutingInteractor.cpp
+++ b/Engine/Source/Editor/VREditor/VRScoutingInteractor.cpp
@@ -6,6 +6,8 @@
 #include "ViewportWorldInteraction.h"
 #include "VREditorMode.h"
 #include "VREditorActions.h"
+#include "GameFramework/InputSettings.h"
+#include "Engine/InputDelegateBinding.h"
 
 #if WITH_EDITOR
 #include "Editor.h"
@@ -22,13 +24,16 @@ UVRScoutingInteractor::UVRScoutingInteractor() :
 
 float UVRScoutingInteractor::GetSlideDelta_Implementation() const
 {
-		return 0.0f;
+	return 0.0f;
 }
 
-void UVRScoutingInteractor::SetupComponent_Implementation( AActor* OwningActor )
+void UVRScoutingInteractor::SetupComponent_Implementation(AActor* OwningActor)
 {
 	Super::SetupComponent_Implementation(OwningActor);
 
+	CreateEditorInput();
+
+
 	// Flying Mesh
 	FlyingIndicatorComponent = NewObject<UStaticMeshComponent>(OwningActor);
 	OwningActor->AddOwnedComponent(FlyingIndicatorComponent);
@@ -41,6 +46,7 @@ void UVRScoutingInteractor::SetupComponent_Implementation( AActor* OwningActor )
 	FlyingIndicatorComponent->SetCollisionResponseToAllChannels(ECollisionResponse::ECR_Ignore);
 	FlyingIndicatorComponent->SetVisibility(false);
 	FlyingIndicatorComponent->SetCastShadow(false);
+
 }
 
 void UVRScoutingInteractor::Shutdown_Implementation()
@@ -48,6 +54,7 @@ void UVRScoutingInteractor::Shutdown_Implementation()
 	Super::Shutdown_Implementation();
 
 	FlyingIndicatorComponent = nullptr;
+	RemoveEditorInput();
 }
 
 void UVRScoutingInteractor::SetGizmoMode(EGizmoHandleTypes InGizmoMode)
@@ -82,4 +89,34 @@ TArray<AActor*> UVRScoutingInteractor::GetSelectedActors()
 	return TArray<AActor*>();
 }
 
-#undef LOCTEXT_NAMESPACE
+void UVRScoutingInteractor::CreateEditorInput()
+{
+	EditorOnlyInputComponent = NewObject<UInputComponent>(this, UInputSettings::GetDefaultInputComponentClass(), TEXT("VScoutInputComponent0"), RF_Transient);
+	UInputDelegateBinding::BindInputDelegates(GetClass(), EditorOnlyInputComponent, this);
+}
+
+void UVRScoutingInteractor::RemoveEditorInput()
+{
+	ensure(!bReceivesEditorInput);
+
+	if (EditorOnlyInputComponent)
+	{
+		EditorOnlyInputComponent->DestroyComponent();
+	}
+	EditorOnlyInputComponent = nullptr;
+}
+
+void UVRScoutingInteractor::SetReceivesEditorInput(bool bInValue)
+{
+	bReceivesEditorInput = bInValue;
+	if (bReceivesEditorInput)
+	{
+		CreateEditorInput();
+	}
+	else
+	{
+		RemoveEditorInput();
+	}
+}
+
+#undef LOCTEXT_NAMESPACE
\ No newline at end of file
diff --git a/Engine/Source/Editor/VREditor/VRScoutingInteractor.h b/Engine/Source/Editor/VREditor/VRScoutingInteractor.h
--- a/Engine/Source/Editor/VREditor/VRScoutingInteractor.h
+++ b/Engine/Source/Editor/VREditor/VRScoutingInteractor.h
@@ -12,12 +12,13 @@
 class AActor;
 class UStaticMesh;
 class UStaticMeshSocket;
+class UInputComponent;
 
 /**
  * Represents the interactor in the world
  */
 UCLASS(Abstract)
-class VREDITOR_API UVRScoutingInteractor: public UVREditorInteractor
+class VREDITOR_API UVRScoutingInteractor : public UVREditorInteractor
 {
 	GENERATED_BODY()
 
@@ -30,7 +31,7 @@ public:
 	virtual float GetSlideDelta_Implementation() const override;
 
 	/** Sets up all components */
-	virtual void SetupComponent_Implementation( AActor* OwningActor ) override;
+	virtual void SetupComponent_Implementation(AActor* OwningActor) override;
 
 	// IViewportInteractorInterface overrides
 	virtual void Shutdown_Implementation() override;
@@ -50,4 +51,36 @@ public:
 	/** Shown in Navigation mode */
 	UPROPERTY(Category = Interactor, EditAnywhere, BlueprintReadOnly)
 	TObjectPtr<class UStaticMeshComponent> FlyingIndicatorComponent;
+
+	/** Returns the current InputComponent. This will be NULL unless bReceivesEditorInput is set to true. */
+	UFUNCTION(BlueprintCallable, BlueprintPure, Category = "Input|Editor")
+	UInputComponent* GetInputComponent() const
+	{
+		return EditorOnlyInputComponent.Get();
+	}
+
+	UFUNCTION(BlueprintSetter, Category = "Input|Editor")
+	void SetReceivesEditorInput(bool bInValue);
+
+	UFUNCTION(BlueprintGetter, BlueprintPure, Category = "Input|Editor")
+	bool GetReceivesEditorInput() const
+	{
+		return bReceivesEditorInput;
+	}
+
+private:
+
+	/** Creates the EditorOnlyInputComponent if it does not already exist and registers all subobject callbacks to it */
+	void CreateEditorInput();
+
+	/** Removes the EditorOnlyInputComponent from this object */
+	void RemoveEditorInput();
+
+	UPROPERTY(Transient, DuplicateTransient)
+	TObjectPtr<UInputComponent> EditorOnlyInputComponent;
+
+	/** If set to true, then this interactor will be able to recieve input delegate callbacks when in the editor. Defaults to true since we will always want this interactor to consume input */
+	UPROPERTY(EditAnywhere, Category = "Input|Editor", BlueprintSetter = SetReceivesEditorInput, BlueprintGetter = GetReceivesEditorInput)
+	bool bReceivesEditorInput = true;
+
 };
