Virtual Scouting: Use OpenXR aim pose for laser origin when available. Adds a second motion controller component to the interactor, with its motion source set to "(Left|Right)Aim" when supported, and uses that to drive FViewportInteractorData::Transform, rather than the grip pose.

This introduced some follow-on issues in places where the grip and aim poses were previously conflated. To address those:

 - The Oculus controller mesh (SM_Controller_RiftS) has had a socket named "OpenXrGrip" added to it, calibrated to coincide roughly with the XR Touch v3 model origin. Under OpenXR, if this socket is present, the controller mesh has its relative transform overridden to treat this socket as the origin. While this significantly improves the usability of this legacy mesh, this is not a long-term solution.

 - The aforementioned grip socket transform logic occurs in UVREditorInteractor::SetHandMeshComponentProperties, and so in order to enforce it correctly, the VirtualScoutingInteractor blueprint now correctly calls UVREditorInteractor::ReplaceHandMeshComponent, instead of setting the mesh asset directly.

 - UVREditorInteractor::ReplaceHandMeshComponent now additionally takes the controller mesh scale as a parameter, in order to correctly account for the Oculus right-hand mirroring (Y scale of -1.0).

 - The VRTool_BP blueprint's "Update Transform" method now uses the parent interactor HandMeshComponent transform, rather than UViewportInteractor::GetTransform, in order to keep all the hardcoded offsets correctly pointing to their respective portions of the controller mesh. The "Helper Pivot Rift" scene component is now attached to the parent interactor's HandMeshComponent for the same reason.

-----

This changelist also includes a couple of unrelated small fixes:

 - Removes unnecessary ensures from FVirtualScoutingOpenXRExtension::TryGetHmdDeviceType, which could trigger if VR preview was used while the Virtual Scouting plugin is active.

 - Changes VRActionTypes::Touch back to a key binding rather than an axis (per the OpenXR spec v1.0.24 §6.3.2, "'touch' components are always boolean"), and instead modifies the MarkerTool blueprint to handle it as a key instead of an axis.

diff --git a/Engine/Plugins/Experimental/VirtualScouting/Source/VirtualScoutingOpenXR/Private/VirtualScoutingOpenXR.cpp b/Engine/Plugins/Experimental/VirtualScouting/Source/VirtualScoutingOpenXR/Private/VirtualScoutingOpenXR.cpp
--- a/Engine/Plugins/Experimental/VirtualScouting/Source/VirtualScoutingOpenXR/Private/VirtualScoutingOpenXR.cpp
+++ b/Engine/Plugins/Experimental/VirtualScouting/Source/VirtualScoutingOpenXR/Private/VirtualScoutingOpenXR.cpp
@@ -231,6 +231,7 @@ void FVirtualScoutingOpenXRExtension::PostSyncActions(XrSession InSession)
 	TryFulfillDeviceTypePromise();
 }
 
+
 void FVirtualScoutingOpenXRExtension::TryFulfillDeviceTypePromise()
 {
 	if (!DeviceTypeFuture.IsReady())
@@ -247,10 +248,9 @@ void FVirtualScoutingOpenXRExtension::TryFulfillDeviceTypePromise()
 
 TOptional<FName> FVirtualScoutingOpenXRExtension::TryGetHmdDeviceType()
 {
-	UVREditorMode* VrMode = IVREditorModule::Get().GetVRMode();
-	if (!ensure(VrMode) || !ensure(Session != XR_NULL_HANDLE) || !ensure(ActionSet != XR_NULL_HANDLE))
+	if (Instance == XR_NULL_HANDLE || Session == XR_NULL_HANDLE)
 	{
-		return TOptional<FName>();
+		return FName();
 	}
 
 	XrPath LeftHand = XR_NULL_PATH, RightHand = XR_NULL_PATH;
@@ -270,6 +270,7 @@ TOptional<FName> FVirtualScoutingOpenXRExtension::TryGetHmdDeviceType()
 
 	if (InteractionProfile == XR_NULL_PATH)
 	{
+		// Returning an unset optional indicates we don't know _yet_, and the caller can poll again.
 		return TOptional<FName>();
 	}
 
diff --git a/Engine/Source/Editor/VREditor/Public/VREditorInteractor.h b/Engine/Source/Editor/VREditor/Public/VREditorInteractor.h
--- a/Engine/Source/Editor/VREditor/Public/VREditorInteractor.h
+++ b/Engine/Source/Editor/VREditor/Public/VREditorInteractor.h
@@ -256,7 +256,7 @@ public:
 
 	/** Replace the default VR controller mesh with a custom one. */
 	UFUNCTION(BlueprintCallable, Category = "VREditorInteractor")
-	void ReplaceHandMeshComponent(UStaticMesh* NewMesh);
+	void ReplaceHandMeshComponent(UStaticMesh* NewMesh, FVector MeshScale = FVector(1.f, 1.f, 1.f));
 
 protected:
 
@@ -271,6 +271,10 @@ protected:
 	UPROPERTY()
 	TObjectPtr<class UMotionControllerComponent> MotionControllerComponent;
 
+	/** Separate motion controller component set to the "aim" pose motion source, used for the laser pointer. */
+	UPROPERTY()
+	TObjectPtr<class UMotionControllerComponent> LaserMotionControllerComponent;
+
 	//
 	// Graphics
 	//
diff --git a/Engine/Source/Editor/VREditor/VREditorInteractor.cpp b/Engine/Source/Editor/VREditor/VREditorInteractor.cpp
--- a/Engine/Source/Editor/VREditor/VREditorInteractor.cpp
+++ b/Engine/Source/Editor/VREditor/VREditorInteractor.cpp
@@ -98,6 +98,7 @@ UVREditorInteractor::UVREditorInteractor() :
 	Super(),
 	bIsUndoRedoSwipeEnabled( true ),
 	MotionControllerComponent( nullptr ),
+	LaserMotionControllerComponent( nullptr ),
 	HandMeshComponent( nullptr ),
 	LaserSplineComponent( nullptr ),
 	LaserPointerMID( nullptr ),
@@ -166,13 +167,13 @@ TMap<FViewportActionKeyInput, TArray<FKey>> UVREditorInteractor::GetKnownActionM
 			ReturnMap.FindOrAdd(FViewportActionKeyInput::Axis(UVREditorInteractor::TrackpadPositionX)).Add(EKeys::Vive_Left_Trackpad_X);
 			ReturnMap.FindOrAdd(FViewportActionKeyInput::Axis(UVREditorInteractor::TrackpadPositionY)).Add(EKeys::Vive_Left_Trackpad_Y);
 			ReturnMap.FindOrAdd(FViewportActionKeyInput(VRActionTypes::ConfirmRadialSelection)).Add(EKeys::Vive_Left_Trackpad_Click);
-			ReturnMap.FindOrAdd(FViewportActionKeyInput::Axis(VRActionTypes::Touch)).Add(EKeys::Vive_Left_Trackpad_Touch);
+			ReturnMap.FindOrAdd(FViewportActionKeyInput(VRActionTypes::Touch)).Add(EKeys::Vive_Left_Trackpad_Touch);
 
 			// Valve Index
 			ReturnMap.FindOrAdd(FViewportActionKeyInput(ViewportWorldActionTypes::WorldMovement)).Add(EKeys::ValveIndex_Left_Trackpad_Touch);
 			ReturnMap.FindOrAdd(FViewportActionKeyInput(VRActionTypes::Modifier)).Add(EKeys::ValveIndex_Left_A_Click);
 			ReturnMap.FindOrAdd(FViewportActionKeyInput(VRActionTypes::Modifier2)).Add(EKeys::ValveIndex_Left_B_Click);
-			ReturnMap.FindOrAdd(FViewportActionKeyInput::Axis(VRActionTypes::Touch)).Add(EKeys::ValveIndex_Left_Thumbstick_Touch);
+			ReturnMap.FindOrAdd(FViewportActionKeyInput(VRActionTypes::Touch)).Add(EKeys::ValveIndex_Left_Thumbstick_Touch);
 
 			ReturnMap.FindOrAdd(FViewportActionKeyInput(VRActionTypes::TrackpadDown)).Add(EKeys::ValveIndex_Left_Thumbstick_Down); // down
 			ReturnMap.FindOrAdd(FViewportActionKeyInput(VRActionTypes::TrackpadUp)).Add(EKeys::ValveIndex_Left_Thumbstick_Up);
@@ -190,7 +191,7 @@ TMap<FViewportActionKeyInput, TArray<FKey>> UVREditorInteractor::GetKnownActionM
 			ReturnMap.FindOrAdd(FViewportActionKeyInput(ViewportWorldActionTypes::WorldMovement)).Add(EKeys::OculusTouch_Left_Grip_Click);
 			ReturnMap.FindOrAdd(FViewportActionKeyInput(VRActionTypes::Modifier)).Add(EKeys::OculusTouch_Left_X_Click);
 			ReturnMap.FindOrAdd(FViewportActionKeyInput(VRActionTypes::Modifier2)).Add(EKeys::OculusTouch_Left_Y_Click);
-			ReturnMap.FindOrAdd(FViewportActionKeyInput::Axis(VRActionTypes::Touch)).Add(EKeys::OculusTouch_Left_Thumbstick_Touch);
+			ReturnMap.FindOrAdd(FViewportActionKeyInput(VRActionTypes::Touch)).Add(EKeys::OculusTouch_Left_Thumbstick_Touch);
 
 			ReturnMap.FindOrAdd(FViewportActionKeyInput(VRActionTypes::TrackpadDown)).Add(EKeys::OculusTouch_Left_Thumbstick_Down); // down
 			ReturnMap.FindOrAdd(FViewportActionKeyInput(VRActionTypes::TrackpadUp)).Add(EKeys::OculusTouch_Left_Thumbstick_Up);
@@ -221,13 +222,13 @@ TMap<FViewportActionKeyInput, TArray<FKey>> UVREditorInteractor::GetKnownActionM
 			ReturnMap.FindOrAdd(FViewportActionKeyInput::Axis(UVREditorInteractor::TrackpadPositionX)).Add(EKeys::Vive_Right_Trackpad_X);
 			ReturnMap.FindOrAdd(FViewportActionKeyInput::Axis(UVREditorInteractor::TrackpadPositionY)).Add(EKeys::Vive_Right_Trackpad_Y);
 			ReturnMap.FindOrAdd(FViewportActionKeyInput(VRActionTypes::ConfirmRadialSelection)).Add(EKeys::Vive_Right_Trackpad_Click);
-			ReturnMap.FindOrAdd(FViewportActionKeyInput::Axis(VRActionTypes::Touch)).Add(EKeys::Vive_Right_Trackpad_Touch);
+			ReturnMap.FindOrAdd(FViewportActionKeyInput(VRActionTypes::Touch)).Add(EKeys::Vive_Right_Trackpad_Touch);
 
 			// Valve Index
 			ReturnMap.FindOrAdd(FViewportActionKeyInput(ViewportWorldActionTypes::WorldMovement)).Add(EKeys::ValveIndex_Right_Trackpad_Touch);
 			ReturnMap.FindOrAdd(FViewportActionKeyInput(VRActionTypes::Modifier)).Add(EKeys::ValveIndex_Right_A_Click);
 			ReturnMap.FindOrAdd(FViewportActionKeyInput(VRActionTypes::Modifier2)).Add(EKeys::ValveIndex_Right_B_Click);
-			ReturnMap.FindOrAdd(FViewportActionKeyInput::Axis(VRActionTypes::Touch)).Add(EKeys::ValveIndex_Right_Thumbstick_Touch);
+			ReturnMap.FindOrAdd(FViewportActionKeyInput(VRActionTypes::Touch)).Add(EKeys::ValveIndex_Right_Thumbstick_Touch);
 
 			ReturnMap.FindOrAdd(FViewportActionKeyInput(VRActionTypes::TrackpadDown)).Add(EKeys::ValveIndex_Right_Thumbstick_Down); // down
 			ReturnMap.FindOrAdd(FViewportActionKeyInput(VRActionTypes::TrackpadUp)).Add(EKeys::ValveIndex_Right_Thumbstick_Up);
@@ -245,7 +246,7 @@ TMap<FViewportActionKeyInput, TArray<FKey>> UVREditorInteractor::GetKnownActionM
 			ReturnMap.FindOrAdd(FViewportActionKeyInput(ViewportWorldActionTypes::WorldMovement)).Add(EKeys::OculusTouch_Right_Grip_Click);
 			ReturnMap.FindOrAdd(FViewportActionKeyInput(VRActionTypes::Modifier)).Add(EKeys::OculusTouch_Right_A_Click);
 			ReturnMap.FindOrAdd(FViewportActionKeyInput(VRActionTypes::Modifier2)).Add(EKeys::OculusTouch_Right_B_Click);
-			ReturnMap.FindOrAdd(FViewportActionKeyInput::Axis(VRActionTypes::Touch)).Add(EKeys::OculusTouch_Right_Thumbstick_Touch);
+			ReturnMap.FindOrAdd(FViewportActionKeyInput(VRActionTypes::Touch)).Add(EKeys::OculusTouch_Right_Thumbstick_Touch);
 
 			ReturnMap.FindOrAdd(FViewportActionKeyInput(VRActionTypes::TrackpadDown)).Add(EKeys::OculusTouch_Right_Thumbstick_Down); // down
 			ReturnMap.FindOrAdd(FViewportActionKeyInput(VRActionTypes::TrackpadUp)).Add(EKeys::OculusTouch_Right_Thumbstick_Up);
@@ -310,12 +311,10 @@ void UVREditorInteractor::SetupComponent_Implementation( AActor* OwningActor )
 	{
 		MotionControllerComponent = NewObject<UMotionControllerComponent>(OwningAvatar);
 		OwningAvatar->AddOwnedComponent( MotionControllerComponent );
-		MotionControllerComponent->SetupAttachment(OwningAvatar->GetRootComponent() );
+		MotionControllerComponent->SetupAttachment( OwningAvatar->GetRootComponent() );
 		MotionControllerComponent->RegisterComponent();
-
 		MotionControllerComponent->SetMobility( EComponentMobility::Movable );
 		MotionControllerComponent->SetCollisionEnabled( ECollisionEnabled::NoCollision );
-
 		MotionControllerComponent->MotionSource = ControllerMotionSource;
 
 		// @todo vreditor: Reenable late frame updates after we've sorted out why they cause popping artifacts on Rift
@@ -336,6 +335,24 @@ void UVREditorInteractor::SetupComponent_Implementation( AActor* OwningActor )
 		HandMeshComponent->SetMaterial( 0, HandMeshMID );
 	}
 
+	if (GEngine->XRSystem.IsValid() && GEngine->XRSystem->GetSystemName() == OpenXRDeviceType)
+	{
+		LaserMotionControllerComponent = NewObject<UMotionControllerComponent>(OwningAvatar);
+		OwningAvatar->AddOwnedComponent(LaserMotionControllerComponent);
+		LaserMotionControllerComponent->SetupAttachment(OwningAvatar->GetRootComponent());
+		LaserMotionControllerComponent->RegisterComponent();
+		LaserMotionControllerComponent->SetMobility(EComponentMobility::Movable);
+		LaserMotionControllerComponent->SetCollisionEnabled(ECollisionEnabled::NoCollision);
+		LaserMotionControllerComponent->MotionSource =
+			(ControllerMotionSource == FXRMotionControllerBase::LeftHandSourceId)
+			? FName("LeftAim")
+			: FName("RightAim");
+
+		// FIXME: Copied from above, but I tested disabling this briefly and didn't see any issues.
+		// @todo vreditor: Reenable late frame updates after we've sorted out why they cause popping artifacts on Rift
+		LaserMotionControllerComponent->bDisableLowLatencyUpdate = true;
+	}
+
 	{
 		UMaterialInterface* LaserPointerMaterial = AssetContainer.LaserPointerMaterial;
 		check( LaserPointerMaterial != nullptr );
@@ -394,7 +411,7 @@ void UVREditorInteractor::SetupComponent_Implementation( AActor* OwningActor )
 
 		LaserSplineComponent = NewObject<USplineComponent>(OwningAvatar);
 		OwningAvatar->AddOwnedComponent( LaserSplineComponent );
-		LaserSplineComponent->SetupAttachment( MotionControllerComponent );
+		LaserSplineComponent->SetupAttachment( LaserMotionControllerComponent ? LaserMotionControllerComponent : MotionControllerComponent );
 		LaserSplineComponent->RegisterComponent();
 		LaserSplineComponent->SetVisibility( false );
 		LaserSplineMeshComponents.Empty();
@@ -435,11 +452,15 @@ void UVREditorInteractor::SetupComponent_Implementation( AActor* OwningActor )
 }
 
 
-void UVREditorInteractor::ReplaceHandMeshComponent(UStaticMesh* NewMesh)
+void UVREditorInteractor::ReplaceHandMeshComponent(
+	UStaticMesh* NewMesh,
+	FVector MeshScale /* = FVector(1.f, 1.f, 1.f) */
+)
 {
 	HandMeshComponent->UnregisterComponent();
 	HandMeshComponent->DestroyComponent();
 	HandMeshComponent = VRMode->CreateMotionControllerMesh(OwningAvatar, MotionControllerComponent, NewMesh);
+	HandMeshComponent->SetRelativeScale3D(MeshScale);
 	SetHandMeshComponentProperties();
 }
 
@@ -451,6 +472,19 @@ void UVREditorInteractor::SetHandMeshComponentProperties()
 	HandMeshComponent->SetCollisionEnabled(ECollisionEnabled::PhysicsOnly);
 	HandMeshComponent->SetCollisionResponseToAllChannels(ECR_Overlap);
 	HandMeshComponent->SetGenerateOverlapEvents(true);
+
+	// To help ease the transition to OpenXR, an alternative controller mesh origin can be specified.
+	if (GEngine->XRSystem.IsValid() && GEngine->XRSystem->GetSystemName() == OpenXRDeviceType)
+	{
+		const FName GripSocket("OpenXrGrip");
+		if (HandMeshComponent->DoesSocketExist(GripSocket))
+		{
+			const FTransform GripTransform = HandMeshComponent->GetSocketTransform(GripSocket, RTS_Component);
+			const FTransform ScaleTransform(FRotator::ZeroRotator, FVector::ZeroVector, HandMeshComponent->GetRelativeScale3D());
+			const FTransform NewRelativeTransform = (GripTransform * ScaleTransform).Inverse();
+			HandMeshComponent->SetRelativeTransform(NewRelativeTransform);
+		}
+	}
 }
 
 void UVREditorInteractor::Shutdown_Implementation()
@@ -516,6 +550,7 @@ void UVREditorInteractor::Tick_Implementation( const float DeltaTime )
 
 		// @todo vreditor: Manually ticking motion controller components
 		MotionControllerComponent->TickComponent( DeltaTime, ELevelTick::LEVELTICK_PauseTick, nullptr );
+		LaserMotionControllerComponent->TickComponent( DeltaTime, ELevelTick::LEVELTICK_PauseTick, nullptr );
 
 		// The hands need to stay the same size relative to our tracking space, so we inverse compensate for world to meters scale here
 		// NOTE: We don't need to set the hand mesh location and rotation, as the MotionControllerComponent does that itself
@@ -619,8 +654,15 @@ void UVREditorInteractor::Tick_Implementation( const float DeltaTime )
 		// so if in the next frame it needs to be visible it won't interpolate from a previous location.
 		{
 			// Offset the beginning of the laser pointer a bit, so that it doesn't overlap the hand mesh
-			const float LaserPointerStartOffset = WorldScaleFactor *
-				(GetVRMode().GetHMDDeviceType() == OculusDeviceType ? VREd::OculusLaserPointerStartOffset->GetFloat() : VREd::ViveLaserPointerStartOffset->GetFloat());
+			float LaserPointerStartOffset = WorldScaleFactor *
+				(GetVRMode().GetHMDDeviceType() == OculusDeviceType
+					? VREd::OculusLaserPointerStartOffset->GetFloat()
+					: VREd::ViveLaserPointerStartOffset->GetFloat());
+
+			if (LaserMotionControllerComponent)
+			{
+				LaserPointerStartOffset = 0.0f;
+			}
 
 			// Get the hand transform and forward vector.
 			FTransform InteractorTransform;
@@ -1165,7 +1207,7 @@ void UVREditorInteractor::PollInput()
 	TArray<IMotionController*> MotionControllers = IModularFeatures::Get().GetModularFeatureImplementations<IMotionController>( IMotionController::GetModularFeatureName() );
 	for (auto MotionController : MotionControllers)	// @todo viewportinteraction: Needs support for multiple pairs of motion controllers
 	{
-		if (MotionController != nullptr && !bHaveMotionController)
+		if (MotionController != nullptr)
 		{
 			FVector Location = FVector::ZeroVector;
 			FRotator Rotation = FRotator::ZeroRotator;
@@ -1175,11 +1217,13 @@ void UVREditorInteractor::PollInput()
 				WorldScale = GetVRMode().GetWorldScaleFactor() *100.0f; // WorldScaleFactor is worldscale / 100.0
 			}
 
-			if (MotionController->GetControllerOrientationAndPosition( WorldInteraction->GetMotionControllerID(), ControllerMotionSource, /* Out */ Rotation, /* Out */ Location, WorldScale ))
+			const FName MotionSource = LaserMotionControllerComponent ? LaserMotionControllerComponent->MotionSource : ControllerMotionSource;
+			if (MotionController->GetControllerOrientationAndPosition( WorldInteraction->GetMotionControllerID(), MotionSource,/* Out */ Rotation, /* Out */ Location, WorldScale ))
 			{
 				bHaveMotionController = true;
 				InteractorData.RoomSpaceTransform = FTransform( Rotation.Quaternion(), Location, FVector( 1.0f ) );
 				InteractorData.Transform = InteractorData.RoomSpaceTransform * WorldInteraction->GetRoomTransform();
+				break;
 			}
 		}
 	}
@@ -1190,10 +1234,22 @@ bool UVREditorInteractor::GetTransformAndForwardVector( FTransform& OutHandTrans
 {
 	if (bHaveMotionController)
 	{
-		OutHandTransform = InteractorData.Transform;
-
-		const float LaserPointerRotationOffset = GetHMDDeviceType() == OculusDeviceType ? VREd::OculusLaserPointerRotationOffset->GetFloat() : VREd::ViveLaserPointerRotationOffset->GetFloat();
-		OutForwardVector = OutHandTransform.GetRotation().RotateVector( FRotator( LaserPointerRotationOffset, 0.0f, 0.0f ).RotateVector( FVector( 1.0f, 0.0f, 0.0f ) ) );
+		if (LaserMotionControllerComponent)
+		{
+			OutHandTransform = InteractorData.Transform;
+			OutForwardVector = InteractorData.Transform.TransformVector(FVector::ForwardVector);
+		}
+		else
+		{
+			const float LaserPointerRotationOffset = GetHMDDeviceType() == OculusDeviceType
+				? VREd::OculusLaserPointerRotationOffset->GetFloat()
+				: VREd::ViveLaserPointerRotationOffset->GetFloat();
+
+			OutHandTransform = InteractorData.Transform;
+			OutForwardVector = OutHandTransform.GetRotation()
+				.RotateVector(FRotator(LaserPointerRotationOffset, 0.0f, 0.0f)
+				.RotateVector(FVector::ForwardVector));
+		}
 
 		return true;
 	}
diff --git a/Engine/Source/Editor/ViewportInteraction/ViewportInteractorData.h b/Engine/Source/Editor/ViewportInteraction/ViewportInteractorData.h
--- a/Engine/Source/Editor/ViewportInteraction/ViewportInteractorData.h
+++ b/Engine/Source/Editor/ViewportInteraction/ViewportInteractorData.h
@@ -15,10 +15,10 @@ struct FViewportInteractorData
 	// Positional data
 	//
 
-	/** Your hand in the virtual world in world space, usually driven by VR motion controllers */
+	/** Your hand (laser origin) in the virtual world in world space, usually driven by VR motion controllers */
 	FTransform Transform;
 
-	/** Your hand transform, in the local tracking space */
+	/** Your hand transform, in the local tracking space (unaffected by world scaling) */
 	FTransform RoomSpaceTransform;
 
 	/** Hand transform in world space from the previous frame */
