Virtual Scouting: When the OpenXR XR_EXT_dpad_binding extension is missing, synthesize Trackpad(Up|Down|Left|Right) events based on axis inputs instead. This is currently required for feature parity under the Oculus OpenXR runtime.

UVREditorInteractor now keeps track of whether action keys are held down in the ActionKeysPressed map, and bIsModifierPressed/bIsTouchingTrackpad/bIsPressingTrackpad have been marked deprecated (and will always remain false).

diff --git a/Engine/Plugins/Experimental/VirtualScouting/Source/VirtualScoutingEditor/Private/VirtualScoutingMode.cpp b/Engine/Plugins/Experimental/VirtualScouting/Source/VirtualScoutingEditor/Private/VirtualScoutingMode.cpp
--- a/Engine/Plugins/Experimental/VirtualScouting/Source/VirtualScoutingEditor/Private/VirtualScoutingMode.cpp
+++ b/Engine/Plugins/Experimental/VirtualScouting/Source/VirtualScoutingEditor/Private/VirtualScoutingMode.cpp
@@ -8,6 +8,7 @@
 #include "Editor.h"
 #endif
 
+#include "IOpenXRHMDModule.h"
 #include "IXRTrackingSystem.h"
 #include "Logging/LogMacros.h"
 #include "Modules/ModuleManager.h"
@@ -20,6 +21,9 @@ DECLARE_LOG_CATEGORY_EXTERN(LogVirtualScouting, Log, All);
 DEFINE_LOG_CATEGORY(LogVirtualScouting);
 
 
+static const FName OpenXRSystemName = FName(TEXT("OpenXR"));
+
+
 class FVirtualScoutingEditorModule : public IModuleInterface
 {
 private:
@@ -36,9 +40,20 @@ private:
 IMPLEMENT_MODULE(FVirtualScoutingEditorModule, VirtualScoutingEditor);
 
 
+bool UVirtualScoutingMode::NeedsSyntheticDpad()
+{
+	if (!GEngine->XRSystem.IsValid() || GEngine->XRSystem->GetSystemName() != OpenXRSystemName)
+	{
+		return Super::NeedsSyntheticDpad();
+	}
+
+	return !IOpenXRHMDModule::Get().IsExtensionEnabled("XR_EXT_dpad_binding");
+}
+
+
 void UVirtualScoutingMode::Enter()
 {
-	if (!GEngine->XRSystem.IsValid() || GEngine->XRSystem->GetSystemName() != FName(TEXT("OpenXR")))
+	if (!GEngine->XRSystem.IsValid() || GEngine->XRSystem->GetSystemName() != OpenXRSystemName)
 	{
 		Super::Enter();
 		return;
diff --git a/Engine/Plugins/Experimental/VirtualScouting/Source/VirtualScoutingEditor/Private/VirtualScoutingMode.h b/Engine/Plugins/Experimental/VirtualScouting/Source/VirtualScoutingEditor/Private/VirtualScoutingMode.h
--- a/Engine/Plugins/Experimental/VirtualScouting/Source/VirtualScoutingEditor/Private/VirtualScoutingMode.h
+++ b/Engine/Plugins/Experimental/VirtualScouting/Source/VirtualScoutingEditor/Private/VirtualScoutingMode.h
@@ -13,6 +13,8 @@ class UVirtualScoutingMode : public UVREditorMode
 
 public:
 	//~ Begin UVREditorMode interface
+	virtual bool NeedsSyntheticDpad() override;
+
 	virtual void Enter() override;
 	//~ End UVREditorMode interface
 };
diff --git a/Engine/Plugins/Experimental/VirtualScouting/Source/VirtualScoutingEditor/VirtualScoutingEditor.Build.cs b/Engine/Plugins/Experimental/VirtualScouting/Source/VirtualScoutingEditor/VirtualScoutingEditor.Build.cs
--- a/Engine/Plugins/Experimental/VirtualScouting/Source/VirtualScoutingEditor/VirtualScoutingEditor.Build.cs
+++ b/Engine/Plugins/Experimental/VirtualScouting/Source/VirtualScoutingEditor/VirtualScoutingEditor.Build.cs
@@ -25,6 +25,7 @@ public class VirtualScoutingEditor : ModuleRules
 
 		PrivateDependencyModuleNames.AddRange(new string[]
 			{
+				"OpenXRHMD",
 				"UnrealEd",
 				"VREditor",
 				"VirtualScoutingOpenXR",
diff --git a/Engine/Source/Editor/VREditor/Public/VREditorInteractor.h b/Engine/Source/Editor/VREditor/Public/VREditorInteractor.h
--- a/Engine/Source/Editor/VREditor/Public/VREditorInteractor.h
+++ b/Engine/Source/Editor/VREditor/Public/VREditorInteractor.h
@@ -93,9 +93,7 @@ public:
 	virtual void PreviewInputKey( class FEditorViewportClient& ViewportClient, FViewportActionKeyInput& Action, const FKey Key, const EInputEvent Event, bool& bOutWasHandled ) override;
 	virtual void HandleInputKey( class FEditorViewportClient& ViewportClient, FViewportActionKeyInput& Action, const FKey Key, const EInputEvent Event, bool& bOutWasHandled ) override;
 	virtual bool GetTransformAndForwardVector( FTransform& OutHandTransform, FVector& OutForwardVector ) const override;
-
-
-	void HandleInputAxis( FEditorViewportClient& ViewportClient, FViewportActionKeyInput& Action, const FKey Key, const float Delta, const float DeltaTime, bool& bOutWasHandled );
+	virtual void HandleInputAxis( FEditorViewportClient& ViewportClient, FViewportActionKeyInput& Action, const FKey Key, const float Delta, const float DeltaTime, bool& bOutWasHandled ) override;
 
 	/** Toggles whether or not this controller is being used to scrub sequencer */
 	void ToggleSequencerScrubbingMode();
@@ -258,6 +256,8 @@ public:
 	UFUNCTION(BlueprintCallable, Category = "VREditorInteractor")
 	void ReplaceHandMeshComponent(UStaticMesh* NewMesh, FVector MeshScale = FVector(1.f, 1.f, 1.f));
 
+	bool IsActionKeyPressed(FName ActionName) const;
+
 protected:
 
 	/** Polls input for the motion controllers transforms */
@@ -330,7 +330,11 @@ private:
 	static const FName MotionController_Left_PressedTriggerAxis;
 	static const FName MotionController_Right_PressedTriggerAxis;
 
+	/** Is the button for the specified action currently held down? */
+	TMap<FName, bool> ActionKeysPressed;
+
 	/** Is the Modifier button held down? */
+	UE_DEPRECATED(5.1, "Use IsActionKeyPressed(VRActionTypes::Modifier) instead.")
 	bool bIsModifierPressed;
 
 	/** Current trigger pressed amount for 'select and move' (0.0 - 1.0) */
@@ -420,9 +424,11 @@ protected:
 	//
 
 	/** True if the trackpad is actively being touched */
+	UE_DEPRECATED(5.1, "Use IsActionKeyPressed(VRActionTypes::Touch) instead.")
 	bool bIsTouchingTrackpad;
 
 	/** True if pressing trackpad button (or analog stick button is down) */
+	UE_DEPRECATED(5.1, "Use IsActionKeyPressed(VRActionTypes::ConfirmRadialSelection) instead.")
 	bool bIsPressingTrackpad;
 
 	/** Position of the touched trackpad */
diff --git a/Engine/Source/Editor/VREditor/Public/VREditorMode.h b/Engine/Source/Editor/VREditor/Public/VREditorMode.h
--- a/Engine/Source/Editor/VREditor/Public/VREditorMode.h
+++ b/Engine/Source/Editor/VREditor/Public/VREditorMode.h
@@ -76,6 +76,11 @@ protected:
 	//~ End UEditorWorldExtension interface
 
 public:
+	virtual bool NeedsSyntheticDpad()
+	{
+		return false;
+	}
+
 	/** When the user actually enters the VR Editor mode */
 	virtual void Enter();
 
diff --git a/Engine/Source/Editor/VREditor/VREditorInteractor.cpp b/Engine/Source/Editor/VREditor/VREditorInteractor.cpp
--- a/Engine/Source/Editor/VREditor/VREditorInteractor.cpp
+++ b/Engine/Source/Editor/VREditor/VREditorInteractor.cpp
@@ -23,6 +23,7 @@
 #include "VREditorDockableWindow.h"
 #include "VREditorFloatingText.h"
 #include "VREditorFloatingUI.h"
+#include "VREditorModule.h"
 #include "VREditorPlacement.h"
 #include "VREditorRadialFloatingUI.h"
 #include "VREditorUISystem.h"
@@ -80,9 +81,9 @@ namespace VREditorKeyNames
 	static const FName MotionController_Right_FullyPressedTriggerAxis( "MotionController_Right_FullyPressedTriggerAxis" );
 }
 
-static const FName OculusDeviceType( TEXT( "OculusHMD" ) );
-static const FName SteamVRDeviceType( TEXT( "SteamVR" ) );
-static const FName OpenXRDeviceType( TEXT( "OpenXR" ) );
+static const FName OculusDeviceType = TEXT( "OculusHMD" );
+static const FName SteamVRDeviceType = TEXT( "SteamVR" );
+static const FName OpenXRDeviceType = TEXT( "OpenXR" );
 
 const FName UVREditorInteractor::TrackpadPositionX = FName( "TrackpadPosition_X" );
 const FName UVREditorInteractor::TrackpadPositionY = FName( "TrackpadPosition_Y" );
@@ -316,9 +317,6 @@ void UVREditorInteractor::SetupComponent_Implementation( AActor* OwningActor )
 		MotionControllerComponent->SetMobility( EComponentMobility::Movable );
 		MotionControllerComponent->SetCollisionEnabled( ECollisionEnabled::NoCollision );
 		MotionControllerComponent->MotionSource = ControllerMotionSource;
-
-		// @todo vreditor: Reenable late frame updates after we've sorted out why they cause popping artifacts on Rift
-		MotionControllerComponent->bDisableLowLatencyUpdate = true;
 	}
 
 	const UVREditorAssetContainer& AssetContainer = VRMode->GetAssetContainer();
@@ -347,10 +345,6 @@ void UVREditorInteractor::SetupComponent_Implementation( AActor* OwningActor )
 			(ControllerMotionSource == FXRMotionControllerBase::LeftHandSourceId)
 			? FName("LeftAim")
 			: FName("RightAim");
-
-		// FIXME: Copied from above, but I tested disabling this briefly and didn't see any issues.
-		// @todo vreditor: Reenable late frame updates after we've sorted out why they cause popping artifacts on Rift
-		LaserMotionControllerComponent->bDisableLowLatencyUpdate = true;
 	}
 
 	{
@@ -908,14 +902,11 @@ void UVREditorInteractor::PreviewInputKey( class FEditorViewportClient& Viewport
 	{
 		if (Event == IE_Pressed)
 		{
-			bIsTouchingTrackpad = true;
-
 			// Set initial position when starting to touch the trackpad
 			InitialTouchPosition = TrackpadPosition;
 		}
 		else if (Event == IE_Released)
 		{
-			bIsTouchingTrackpad = false;
 			bIsTrackpadPositionValid[0] = false;
 			bIsTrackpadPositionValid[1] = false;
 
@@ -956,23 +947,7 @@ void UVREditorInteractor::PreviewInputKey( class FEditorViewportClient& Viewport
 		}
 	}
 
-	if (Action.ActionType == VRActionTypes::ConfirmRadialSelection)
-	{
-		bIsPressingTrackpad = Event == IE_Released ? false : true;
-	}
-
-	// Update modifier state
-	if (Action.ActionType == VRActionTypes::Modifier)
-	{
-		if (Event == IE_Pressed)
-		{
-			bIsModifierPressed = true;
-		}
-		else if (Event == IE_Released)
-		{
-			bIsModifierPressed = false;
-		}
-	}
+	ActionKeysPressed.FindOrAdd(Action.ActionType) = (Event == IE_Released) ? false : true;
 
 	if (!bOutWasHandled)
 	{
@@ -1018,10 +993,12 @@ void UVREditorInteractor::HandleInputKey( class FEditorViewportClient& ViewportC
 
 void UVREditorInteractor::HandleInputAxis( FEditorViewportClient& ViewportClient, FViewportActionKeyInput& Action, const FKey Key, const float Delta, const float DeltaTime, bool& bOutWasHandled )
 {
+	const FName HmdDeviceType = GetHMDDeviceType();
+
 	if (!bOutWasHandled && Action.ActionType == TriggerAxis)
 	{
-		const float TriggerPressedThreshold = (GetHMDDeviceType() == OculusDeviceType) ? GetDefault<UVRModeSettings>()->TriggerPressedThreshold_Rift : GetDefault<UVRModeSettings>()->TriggerPressedThreshold_Vive;
-		const float TriggerDeadZone = (GetHMDDeviceType() == OculusDeviceType) ? VREd::TriggerDeadZone_Rift->GetFloat() : VREd::TriggerDeadZone_Vive->GetFloat();
+		const float TriggerPressedThreshold = (HmdDeviceType == OculusDeviceType) ? GetDefault<UVRModeSettings>()->TriggerPressedThreshold_Rift : GetDefault<UVRModeSettings>()->TriggerPressedThreshold_Vive;
+		const float TriggerDeadZone = (HmdDeviceType == OculusDeviceType) ? VREd::TriggerDeadZone_Rift->GetFloat() : VREd::TriggerDeadZone_Vive->GetFloat();
 
 		// Synthesize "lightly pressed" events for the trigger
 		{
@@ -1055,7 +1032,7 @@ void UVREditorInteractor::HandleInputAxis( FEditorViewportClient& ViewportClient
 
 		// Synthesize "fully pressed" events for the trigger
 		{
-			const float TriggerFullyPressedThreshold = (GetHMDDeviceType() == OculusDeviceType) ? VREd::TriggerFullyPressedThreshold_Rift->GetFloat() : VREd::TriggerFullyPressedThreshold_Vive->GetFloat();
+			const float TriggerFullyPressedThreshold = (HmdDeviceType == OculusDeviceType) ? VREd::TriggerFullyPressedThreshold_Rift->GetFloat() : VREd::TriggerFullyPressedThreshold_Vive->GetFloat();
 
 			if (!bIsTriggerFullyPressed &&	// Don't fire if we are already pressed
 				Delta >= TriggerFullyPressedThreshold)
@@ -1103,6 +1080,89 @@ void UVREditorInteractor::HandleInputAxis( FEditorViewportClient& ViewportClient
 		}
 	}
 
+	// Check and generate trackpad d-pad key events ourselves, which may be required if the XR runtime doesn't support them.
+	if (Action.ActionType == TrackpadPositionX || Action.ActionType == TrackpadPositionY)
+	{
+		const bool bTrackpadValid = bIsTrackpadPositionValid[0] && bIsTrackpadPositionValid[1];
+		if (VRMode->NeedsSyntheticDpad() && bTrackpadValid)
+		{
+			//const float WedgeAngle = 0.5f * UE_PI; // Quadrants / 4-way
+			const float WedgeAngle = 3.0f * UE_PI / 4.0f; // Overlapping octants / 8-way
+			const float CenterRegion = 0.75f;
+
+			const float WedgeHalfAngle = WedgeAngle / 2.0f;
+
+			enum EDpadDirection
+			{
+				Up,
+				Down,
+				Left,
+				Right,
+			};
+
+			const float Magnitude = TrackpadPosition.Size();
+			const float Angle = FMath::Atan2(TrackpadPosition.Y, TrackpadPosition.X);
+
+			auto IsDpadDirectionPressed = [&](EDpadDirection Dir) -> bool
+			{
+				if (Magnitude <= CenterRegion)
+				{
+					return false;
+				}
+
+				switch (Dir)
+				{
+					case Up:    return Angle >= (UE_PI * .5f - WedgeHalfAngle)  && Angle <= (UE_PI * .5f + WedgeHalfAngle);
+					case Down:  return Angle >= (-UE_PI * .5f - WedgeHalfAngle) && Angle <= (-UE_PI * .5f + WedgeHalfAngle);
+					case Left:  return Angle >= (UE_PI - WedgeHalfAngle)        || Angle <= (-UE_PI + WedgeHalfAngle);
+					case Right: return Angle >= (-WedgeHalfAngle)               && Angle <= (WedgeHalfAngle);
+					default: checkNoEntry(); return false;
+				}
+			};
+
+			auto CheckSendDpadKeyEvent = [&](EDpadDirection Dir)
+			{
+				FName ActionName;
+				switch (Dir)
+				{
+					case Up:    ActionName = VRActionTypes::TrackpadUp; break;
+					case Down:  ActionName = VRActionTypes::TrackpadDown; break;
+					case Left:  ActionName = VRActionTypes::TrackpadLeft; break;
+					case Right: ActionName = VRActionTypes::TrackpadRight; break;
+					default: checkNoEntry(); return;
+				}
+
+				const bool bWasPressed = IsActionKeyPressed(ActionName);
+				const bool bPressedNow = IsDpadDirectionPressed(Dir);
+
+				if (bPressedNow != bWasPressed)
+				{
+					FKey InputKey;
+					for (const TPair<FKey, FViewportActionKeyInput>& ActionPair : KeyToActionMap)
+					{
+						if (ActionPair.Value.ActionType == ActionName)
+						{
+							InputKey = ActionPair.Key;
+						}
+					}
+
+					if (ensure(InputKey.IsValid()))
+					{
+						const EInputEvent InputEvent = bPressedNow ? IE_Pressed : IE_Released;
+						UE_LOG(LogVREditor, Verbose, TEXT("Synthesized D-pad event: Key %s %s"),
+							*InputKey.ToString(), InputEvent == IE_Pressed ? TEXT("pressed") : TEXT("released"));
+						UViewportInteractor::HandleInputKey(ViewportClient, InputKey, InputEvent);
+					}
+				}
+			};
+
+			CheckSendDpadKeyEvent(Up);
+			CheckSendDpadKeyEvent(Down);
+			CheckSendDpadKeyEvent(Left);
+			CheckSendDpadKeyEvent(Right);
+		}
+	}
+
 	Super::HandleInputAxis( ViewportClient, Action, Key, Delta, DeltaTime, bOutWasHandled );
 }
 
@@ -1180,7 +1240,7 @@ float UVREditorInteractor::GetTrackpadSlideDelta( const bool Axis /*= 1*/ ) cons
 
 	const bool bIsAbsolute = (GetVRMode().GetHMDDeviceType() == SteamVRDeviceType);
 	float SlideDelta = 0.0f;
-	if (bIsTouchingTrackpad || !bIsAbsolute)
+	if (IsActionKeyPressed(VRActionTypes::Touch) || !bIsAbsolute)
 	{
 		if (bIsAbsolute)
 		{
@@ -1197,6 +1257,17 @@ float UVREditorInteractor::GetTrackpadSlideDelta( const bool Axis /*= 1*/ ) cons
 }
 
 
+bool UVREditorInteractor::IsActionKeyPressed(FName ActionName) const
+{
+	if (const bool* bResult = ActionKeysPressed.Find(ActionName))
+	{
+		return *bResult;
+	}
+
+	return false;
+}
+
+
 void UVREditorInteractor::PollInput()
 {
 	bHaveMotionController = false;
@@ -1706,7 +1777,7 @@ void UVREditorInteractor::UpdateRadialMenuInput( const float DeltaTime )
 			}
 		}
 		// If we are not currently touching the Vive touchpad, reset the highlighted button and pause sequencer playback if scrubbing
-		else if (HMDDeviceType == SteamVRDeviceType && !bIsTouchingTrackpad)
+		else if (HMDDeviceType == SteamVRDeviceType && !IsActionKeyPressed(VRActionTypes::Touch))
 		{
 			if (bIsScrubbingSequence)
 			{
@@ -1797,14 +1868,13 @@ bool UVREditorInteractor::GetIsLaserBlocked() const
 void UVREditorInteractor::ResetTrackpad()
 {
 	TrackpadPosition = FVector2D::ZeroVector;
-	bIsTouchingTrackpad = false;
 	bIsTrackpadPositionValid[0] = false;
 	bIsTrackpadPositionValid[1] = false;
 }
 
 bool UVREditorInteractor::IsTouchingTrackpad() const
 {
-	return bIsTouchingTrackpad;
+	return IsActionKeyPressed(VRActionTypes::Touch);
 }
 
 FVector2D UVREditorInteractor::GetTrackpadPosition() const
@@ -1879,7 +1949,7 @@ void UVREditorInteractor::SetLastHoveredWidgetComponent( UWidgetComponent* NewHo
 
 bool UVREditorInteractor::IsModifierPressed() const
 {
-	return bIsModifierPressed;
+	return IsActionKeyPressed(VRActionTypes::Modifier);
 }
 
 void UVREditorInteractor::SetIsClickingOnUI( const bool bInIsClickingOnUI )
diff --git a/Engine/Source/Editor/VREditor/VREditorModule.cpp b/Engine/Source/Editor/VREditor/VREditorModule.cpp
--- a/Engine/Source/Editor/VREditor/VREditorModule.cpp
+++ b/Engine/Source/Editor/VREditor/VREditorModule.cpp
@@ -13,6 +13,10 @@
 #include "HeadMountedDisplayTypes.h"
 #include "VREditorFloatingUI.h"
 
+
+DEFINE_LOG_CATEGORY(LogVREditor);
+
+
 class FVREditorModule : public IVREditorModule
 {
 public:
diff --git a/Engine/Source/Editor/VREditor/VREditorModule.h b/Engine/Source/Editor/VREditor/VREditorModule.h
--- a/Engine/Source/Editor/VREditor/VREditorModule.h
+++ b/Engine/Source/Editor/VREditor/VREditorModule.h
@@ -3,3 +3,7 @@
 #pragma once
 
 #include "CoreMinimal.h"
+#include "Logging/LogMacros.h"
+
+
+DECLARE_LOG_CATEGORY_EXTERN(LogVREditor, Log, All);
