Force all VREditor widget components to not cast shadows.

diff --git a/Engine/Source/Editor/VREditor/VREditorWidgetComponent.cpp b/Engine/Source/Editor/VREditor/VREditorWidgetComponent.cpp
--- a/Engine/Source/Editor/VREditor/VREditorWidgetComponent.cpp
+++ b/Engine/Source/Editor/VREditor/VREditorWidgetComponent.cpp
@@ -9,6 +9,7 @@ UVREditorWidgetComponent::UVREditorWidgetComponent()
 	bHasEverDrawn(false)
 {
 	bSelectable = false;
+	SetCastShadow(false);
 }
 
 bool UVREditorWidgetComponent::ShouldDrawWidget() const
