Updated VREditorWidgetComponent and VREditorCameraWidgetComponent to use materials that ignore pre-exposure so VR menus have same brightness regardless of scene lighting setup.

Added new materials to VREditorAssetContainer to avoid modifying the default materials used for WidgetComponents

diff --git a/Engine/Source/Editor/VREditor/Public/VREditorAssetContainer.h b/Engine/Source/Editor/VREditor/Public/VREditorAssetContainer.h
--- a/Engine/Source/Editor/VREditor/Public/VREditorAssetContainer.h
+++ b/Engine/Source/Editor/VREditor/Public/VREditorAssetContainer.h
@@ -163,6 +163,14 @@ public:
 	UPROPERTY(EditAnywhere, Category = Material)
 	TObjectPtr<UMaterialInterface> TranslucentTextMaterial;
 
+	UPROPERTY(EditAnywhere, Category = Material)
+	TObjectPtr<UMaterialInterface> WidgetMaterial;
+
+	//Specific material for camera widgets that operates in linear color space
+	UPROPERTY(EditAnywhere, Category = Material)
+	TObjectPtr<UMaterialInterface> CameraWidgetMaterial;
+
+
 	//
 	// Fonts
 	//
diff --git a/Engine/Source/Editor/VREditor/VREditorCameraWidgetComponent.cpp b/Engine/Source/Editor/VREditor/VREditorCameraWidgetComponent.cpp
--- a/Engine/Source/Editor/VREditor/VREditorCameraWidgetComponent.cpp
+++ b/Engine/Source/Editor/VREditor/VREditorCameraWidgetComponent.cpp
@@ -1,17 +1,36 @@
 // Copyright Epic Games, Inc. All Rights Reserved.
 
 #include "VREditorCameraWidgetComponent.h"
-#include "UObject/ConstructorHelpers.h"
+#include "IVREditorModule.h"
+#include "VREditorAssetContainer.h"
+#include "VREditorMode.h"
 
 UVREditorCameraWidgetComponent::UVREditorCameraWidgetComponent()
 	: Super()	
 {	
 	// Override this shader for VR camera viewfinders so that we get color-correct images.
-	// This shader does an sRGB -> Linear conversion and doesn't apply the "UI Brightness" setting .
+	 //This shader does an sRGB -> Linear conversion and doesn't apply the "UI Brightness" setting .
+
+	if (UNLIKELY(IsRunningDedicatedServer()) || HasAnyFlags(RF_ClassDefaultObject))   // @todo vreditor: early out to avoid loading font assets in the cooker on Linux
+	{
+		return;
+	}
+
 	if (!IsRunningCommandlet())
 	{
-		static ConstructorHelpers::FObjectFinder<UMaterialInterface> OpaqueMaterial_OneSided_Finder(TEXT("/Engine/EngineMaterials/Widget3DCameraPassThrough_Opaque_OneSided"));
-		OpaqueMaterial_OneSided = OpaqueMaterial_OneSided_Finder.Object;
+		if (IVREditorModule::IsAvailable())
+		{
+			IVREditorModule& VRModule = IVREditorModule::Get();
+
+			if (UVREditorMode* VRMode = VRModule.GetVRMode())
+			{
+				const UVREditorAssetContainer& VRAssetContainer = VRMode->GetAssetContainer();
+				{
+					OpaqueMaterial_OneSided = VRAssetContainer.CameraWidgetMaterial;
+				}
+			}
+		}
+
 	}
 }
 
diff --git a/Engine/Source/Editor/VREditor/VREditorWidgetComponent.cpp b/Engine/Source/Editor/VREditor/VREditorWidgetComponent.cpp
--- a/Engine/Source/Editor/VREditor/VREditorWidgetComponent.cpp
+++ b/Engine/Source/Editor/VREditor/VREditorWidgetComponent.cpp
@@ -1,6 +1,9 @@
 // Copyright Epic Games, Inc. All Rights Reserved.
 
 #include "VREditorWidgetComponent.h"
+#include "IVREditorModule.h"
+#include "VREditorAssetContainer.h"
+#include "VREditorMode.h"
 
 UVREditorWidgetComponent::UVREditorWidgetComponent()
 	: Super(),
@@ -10,8 +13,32 @@ UVREditorWidgetComponent::UVREditorWidgetComponent()
 {
 	bSelectable = false;
 	SetCastShadow(false);
+
+	if (UNLIKELY(IsRunningDedicatedServer()) || HasAnyFlags(RF_ClassDefaultObject))   // @todo vreditor: early out to avoid loading font assets in the cooker on Linux
+	{
+		return;
+	}
+
+	if (!IsRunningCommandlet())
+	{
+		if (IVREditorModule::IsAvailable())
+		{
+			IVREditorModule& VRModule = IVREditorModule::Get();
+
+			if (UVREditorMode* VRMode = VRModule.GetVRMode())
+			{
+				const UVREditorAssetContainer& VRAssetContainer = VRMode->GetAssetContainer();
+				{
+					TranslucentMaterial = VRAssetContainer.WidgetMaterial;
+					TranslucentMaterial_OneSided = VRAssetContainer.WidgetMaterial;
+				}
+			}
+		}
+
+	}
 }
 
+
 bool UVREditorWidgetComponent::ShouldDrawWidget() const
 {
 	if ( DrawingPolicy == EVREditorWidgetDrawingPolicy::Always ||
