Virtual Scouting / VR Editor: Rather than specifying interactor or teleporter class overrides directly in editor preferences, these are now determined by choosing a VREditorMode subclass. This change lays the groundwork for further simplification of first-time setup / onboarding.

This change also suppresses the "experimental" warning popup if the Virtual Scouting class is selected.

diff --git a/Engine/Config/BaseEditorSettings.ini b/Engine/Config/BaseEditorSettings.ini
--- a/Engine/Config/BaseEditorSettings.ini
+++ b/Engine/Config/BaseEditorSettings.ini
@@ -60,8 +60,6 @@ DoubleClickTime=0.250000
 TriggerPressedThreshold_Vive=0.330000
 TriggerPressedThreshold_Rift=0.500000
 ModeClass=Class'/Script/VREditor.VREditorMode'
-InteractorClass=Class'/Script/VREditor.VREditorInteractor'
-TeleporterClass=Class'/Script/VREditor.VREditorTeleporter'
 
 [/Script/TakeRecorder.TakeRecorderProjectSettings]
 +Settings=(DefaultTracks=((MatchingActorClass=/Script/CinematicCamera.CineCameraActor,DefaultPropertyTracks=((ComponentPath="CameraComponent",PropertyPath="CurrentFocalLength"),(ComponentPath="CameraComponent",PropertyPath="FocusSettings.ManualFocusDistance"),(ComponentPath="CameraComponent",PropertyPath="CurrentAperture"),(ComponentPath="CameraComponent",PropertyPath="FilmbackSettings.SensorWidth"),(ComponentPath="CameraComponent",PropertyPath="FilmbackSettings.SensorHeight")), ExcludePropertyTracks=((ComponentPath="CameraComponent",PropertyPath="FilmbackSettings.SensorAspectRatio"),(ComponentPath="CameraComponent",PropertyPath="AspectRatio"),(ComponentPath="CameraComponent",PropertyPath="FieldOfView"))),(MatchingActorClass=/Script/Engine.Light,DefaultPropertyTracks=((ComponentPath="LightComponent0",PropertyPath="Intensity"),(ComponentPath="LightComponent0",PropertyPath="LightColor")))))
diff --git a/Engine/Plugins/Experimental/VirtualScouting/Source/VirtualScoutingEditor/Private/VirtualScoutingMode.cpp b/Engine/Plugins/Experimental/VirtualScouting/Source/VirtualScoutingEditor/Private/VirtualScoutingMode.cpp
--- a/Engine/Plugins/Experimental/VirtualScouting/Source/VirtualScoutingEditor/Private/VirtualScoutingMode.cpp
+++ b/Engine/Plugins/Experimental/VirtualScouting/Source/VirtualScoutingEditor/Private/VirtualScoutingMode.cpp
@@ -40,6 +40,13 @@ private:
 IMPLEMENT_MODULE(FVirtualScoutingEditorModule, VirtualScoutingEditor);
 
 
+UVirtualScoutingMode::UVirtualScoutingMode(const FObjectInitializer& ObjectInitializer)
+{
+	InteractorClass = FSoftObjectPath("/VirtualProductionUtilities/VR/VirtualScoutingInteractor.VirtualScoutingInteractor_C");
+	TeleporterClass = FSoftObjectPath("/VirtualProductionUtilities/VR/VirtualScoutingTeleporter.VirtualScoutingTeleporter_C");
+}
+
+
 bool UVirtualScoutingMode::NeedsSyntheticDpad()
 {
 	if (!GEngine->XRSystem.IsValid() || GEngine->XRSystem->GetSystemName() != OpenXRSystemName)
@@ -95,4 +102,5 @@ void UVirtualScoutingMode::Enter()
 	});
 }
 
+
 #undef LOCTEXT_NAMESPACE
diff --git a/Engine/Plugins/Experimental/VirtualScouting/Source/VirtualScoutingEditor/Private/VirtualScoutingMode.h b/Engine/Plugins/Experimental/VirtualScouting/Source/VirtualScoutingEditor/Private/VirtualScoutingMode.h
--- a/Engine/Plugins/Experimental/VirtualScouting/Source/VirtualScoutingEditor/Private/VirtualScoutingMode.h
+++ b/Engine/Plugins/Experimental/VirtualScouting/Source/VirtualScoutingEditor/Private/VirtualScoutingMode.h
@@ -12,9 +12,13 @@ class UVirtualScoutingMode : public UVREditorMode
 	GENERATED_BODY()
 
 public:
+	UVirtualScoutingMode(const FObjectInitializer& ObjectInitializer);
+
 	//~ Begin UVREditorMode interface
 	virtual bool NeedsSyntheticDpad() override;
 
+	virtual bool ShouldDisplayExperimentalWarningOnEntry() const override { return false; }
+
 	virtual void Enter() override;
 	//~ End UVREditorMode interface
 };
diff --git a/Engine/Plugins/Experimental/VirtualScouting/Source/VirtualScoutingOpenXR/Private/VirtualScoutingOpenXR.cpp b/Engine/Plugins/Experimental/VirtualScouting/Source/VirtualScoutingOpenXR/Private/VirtualScoutingOpenXR.cpp
--- a/Engine/Plugins/Experimental/VirtualScouting/Source/VirtualScoutingOpenXR/Private/VirtualScoutingOpenXR.cpp
+++ b/Engine/Plugins/Experimental/VirtualScouting/Source/VirtualScoutingOpenXR/Private/VirtualScoutingOpenXR.cpp
@@ -171,7 +171,13 @@ void FVirtualScoutingOpenXRExtension::AddActions(XrInstance InInstance, FCreateA
 		ActionSet = XR_NULL_HANDLE;
 	}
 
-	UClass* InteractorClass = GetDefault<UVRModeSettings>()->InteractorClass.Get();
+	UVREditorMode* VRMode = IVREditorModule::Get().GetVRMode();
+	if (!VRMode)
+	{
+		return;
+	}
+
+	UClass* InteractorClass = VRMode->InteractorClass.Get();
 	if (!InteractorClass)
 	{
 		return;
diff --git a/Engine/Source/Editor/VREditor/Public/VREditorMode.h b/Engine/Source/Editor/VREditor/Public/VREditorMode.h
--- a/Engine/Source/Editor/VREditor/Public/VREditorMode.h
+++ b/Engine/Source/Editor/VREditor/Public/VREditorMode.h
@@ -49,7 +49,7 @@ namespace VRActionTypes
 /**
  * VR Editor Mode. Extends editor viewports with functionality for VR controls and object manipulation
  */
-UCLASS( BlueprintType, Transient )
+UCLASS( BlueprintType, Blueprintable, Transient )
 class VREDITOR_API UVREditorMode : public UEditorWorldExtension
 {
 	GENERATED_BODY()
@@ -81,6 +81,8 @@ public:
 		return false;
 	}
 
+	virtual bool ShouldDisplayExperimentalWarningOnEntry() const { return true; }
+
 	/** When the user actually enters the VR Editor mode */
 	virtual void Enter();
 
@@ -499,6 +501,14 @@ public:
 	/** The asset container path */
 	static const TCHAR* AssetContainerPath;
 
+	/** The controller to use when UnrealEd is in VR mode. Use VREditorInteractor get default editor behavior, or select a custom controller for special behavior */
+	UPROPERTY( EditAnywhere, BlueprintReadWrite, NoClear, Category="Classes" )
+	TSoftClassPtr<UVREditorInteractor> InteractorClass;
+
+	/** The teleporter to use when UnrealEd is in VR mode. Use VREditorTeleporter to get default editor behavior, or select a custom teleporter */
+	UPROPERTY( EditAnywhere, BlueprintReadWrite, NoClear, Category="Classes" )
+	TSoftClassPtr<AVREditorTeleporter> TeleporterClass;
+
 private:
 
 	// All the colors for this mode
diff --git a/Engine/Source/Editor/VREditor/Public/VRModeSettings.h b/Engine/Source/Editor/VREditor/Public/VRModeSettings.h
--- a/Engine/Source/Editor/VREditor/Public/VRModeSettings.h
+++ b/Engine/Source/Editor/VREditor/Public/VRModeSettings.h
@@ -80,12 +80,12 @@ public:
 	UPROPERTY(EditAnywhere, config, Category = "Motion Controllers", meta = (DisplayName = "Trigger Pressed Threshold (Oculus Touch)", ClampMin = 0.01, ClampMax = 1.0))
 	float TriggerPressedThreshold_Rift;
 
-	/** The controller to use when UnrealEd is in VR mode. Use VREditorInteractor get default editor behavior, or select a custom controller for special behavior */
-	UPROPERTY(EditAnywhere, config, NoClear, Category = "Motion Controllers")
+	UE_DEPRECATED(5.1, "Refer to UVREditorMode::InteractorClass, or create a derived mode to override the interactor class.")
+	UPROPERTY(meta = (DeprecatedProperty, DeprecationMessage = "Refer to UVREditorMode::InteractorClass, or create a derived mode to override the interactor class."))
 	TSoftClassPtr<UVREditorInteractor> InteractorClass;
 
-	/** The teleporter to use when UnrealEd is in VR mode. Use VREditorTeleporter to get default editor behavior, or select a custom teleporter */
-	UPROPERTY( EditAnywhere, config, NoClear, Category = "Motion Controllers" )
+	UE_DEPRECATED(5.1, "Refer to UVREditorMode::TeleporterClass, or create a derived mode to override the teleporter class.")
+	UPROPERTY(meta = (DeprecatedProperty, DeprecationMessage = "Refer to UVREditorMode::TeleporterClass, or create a derived mode to override the teleporter class."))
 	TSoftClassPtr<AVREditorTeleporter> TeleporterClass;
 
 	/** The mode extension to use when UnrealEd is in VR mode. Use VREditorMode to get default editor behavior or select a custom mode. */
diff --git a/Engine/Source/Editor/VREditor/VREditorMode.cpp b/Engine/Source/Editor/VREditor/VREditorMode.cpp
--- a/Engine/Source/Editor/VREditor/VREditorMode.cpp
+++ b/Engine/Source/Editor/VREditor/VREditorMode.cpp
@@ -96,6 +96,8 @@ UVREditorMode::UVREditorMode() :
 	TeleportActor( nullptr ),
 	AutoScalerSystem( nullptr ),
 	WorldInteraction( nullptr ),
+	InteractorClass( UVREditorInteractor::StaticClass() ),
+	TeleporterClass( AVREditorTeleporter::StaticClass() ),
 	bFirstTick( true ),
 	AssetContainer( nullptr )
 {
@@ -190,13 +192,12 @@ void UVREditorMode::AllocateInteractors()
 	class UVREditorInteractor* LeftHandInteractor = nullptr;
 	class UVREditorInteractor* RightHandInteractor = nullptr;
 
-	const TSoftClassPtr<UVREditorInteractor> InteractorClassSoft = GetDefault<UVRModeSettings>()->InteractorClass;
-	InteractorClassSoft.LoadSynchronous();
+	InteractorClass.LoadSynchronous();
 
-	if (InteractorClassSoft.IsValid())
+	if (InteractorClass.IsValid())
 	{
-		LeftHandInteractor = NewObject<UVREditorInteractor>(GetTransientPackage(), InteractorClassSoft.Get());
-		RightHandInteractor = NewObject<UVREditorInteractor>(GetTransientPackage(), InteractorClassSoft.Get());
+		LeftHandInteractor = NewObject<UVREditorInteractor>(GetTransientPackage(), InteractorClass.Get());
+		RightHandInteractor = NewObject<UVREditorInteractor>(GetTransientPackage(), InteractorClass.Get());
 	}
 
 	if (LeftHandInteractor == nullptr)
@@ -346,12 +347,11 @@ void UVREditorMode::SetupSubsystems()
 	PlacementSystem->Init(this);
 
 	// Setup teleporter
-	const TSoftClassPtr<AVREditorTeleporter> TeleporterClassSoft = GetDefault<UVRModeSettings>()->TeleporterClass;
-	TeleporterClassSoft.LoadSynchronous();
+	TeleporterClass.LoadSynchronous();
 
-	if (TeleporterClassSoft.IsValid())
+	if (TeleporterClass.IsValid())
 	{
-		TeleportActor = CastChecked<AVREditorTeleporter>( SpawnTransientSceneActor(TeleporterClassSoft.Get(), TEXT( "Teleporter" ), true ) );
+		TeleportActor = CastChecked<AVREditorTeleporter>( SpawnTransientSceneActor(TeleporterClass.Get(), TEXT( "Teleporter" ), true ) );
 	}
 
 	if( !TeleportActor )
diff --git a/Engine/Source/Editor/VREditor/VREditorModeManager.cpp b/Engine/Source/Editor/VREditor/VREditorModeManager.cpp
--- a/Engine/Source/Editor/VREditor/VREditorModeManager.cpp
+++ b/Engine/Source/Editor/VREditor/VREditorModeManager.cpp
@@ -94,34 +94,31 @@ void FVREditorModeManager::EnableVREditor( const bool bEnable, const bool bForce
 	{
 		if( bEnable && ( IsVREditorAvailable() || bForceWithoutHMD ))
 		{
-			UEditorStyleSettings* StyleSettings = GetMutableDefault<UEditorStyleSettings>();
-			const UVRModeSettings* VRModeSettings = GetDefault<UVRModeSettings>();
-			bool bUsingDefaultInteractors = true;
-			if (VRModeSettings)
+			// Check to see if we should warn the user and potentially early out.
+			if (IsVREditorAvailable())
 			{
-				const TSoftClassPtr<UVREditorInteractor> InteractorClassSoft = VRModeSettings->InteractorClass;
-				InteractorClassSoft.LoadSynchronous();
-
-				if (InteractorClassSoft.IsValid())
+				TSoftClassPtr<UVREditorMode> ModeClassSoft = GetDefault<UVRModeSettings>()->ModeClass;
+				check(ModeClassSoft.LoadSynchronous()); // IsVREditorAvailable() should have returned false otherwise.
+				const UVREditorMode* VRModeCDO = ModeClassSoft->GetDefaultObject<UVREditorMode>();
+				if (VRModeCDO->ShouldDisplayExperimentalWarningOnEntry())
 				{
-					bUsingDefaultInteractors = (InteractorClassSoft.Get() == UVREditorInteractor::StaticClass());
+					FSuppressableWarningDialog::FSetupInfo SetupInfo(
+						LOCTEXT("VRModeEntry_Message", "VR Mode enables you to work on your project in virtual reality using motion controllers. This feature is still under development, so you may experience bugs or crashes while using it."),
+						LOCTEXT("VRModeEntry_Title", "Entering VR Mode - Experimental"), "Warning_VRModeEntry", GEditorSettingsIni);
+
+					SetupInfo.ConfirmText = LOCTEXT("VRModeEntry_ConfirmText", "Continue");
+					SetupInfo.CancelText = LOCTEXT("VRModeEntry_CancelText", "Cancel");
+					SetupInfo.bDefaultToSuppressInTheFuture = true;
+
+					FSuppressableWarningDialog VRModeEntryWarning(SetupInfo);
+					if (VRModeEntryWarning.ShowModal() == FSuppressableWarningDialog::Cancel)
+					{
+						return;
+					}
 				}
 			}
 
-			{
-				FSuppressableWarningDialog::FSetupInfo SetupInfo(LOCTEXT("VRModeEntry_Message", "VR Mode enables you to work on your project in virtual reality using motion controllers. This feature is still under development, so you may experience bugs or crashes while using it."),
-					LOCTEXT("VRModeEntry_Title", "Entering VR Mode - Experimental"), "Warning_VRModeEntry", GEditorSettingsIni);
-
-				SetupInfo.ConfirmText = LOCTEXT("VRModeEntry_ConfirmText", "Continue");
-				SetupInfo.CancelText = LOCTEXT("VRModeEntry_CancelText", "Cancel");
-				SetupInfo.bDefaultToSuppressInTheFuture = true;
-				FSuppressableWarningDialog VRModeEntryWarning(SetupInfo);
-
-				if (VRModeEntryWarning.ShowModal() != FSuppressableWarningDialog::Cancel)
-				{
-					StartVREditorMode(bForceWithoutHMD);
-				}
-			}
+			StartVREditorMode(bForceWithoutHMD);
 		}
 		else if( !bEnable )
 		{
@@ -140,17 +137,28 @@ bool FVREditorModeManager::IsVREditorActive() const
 const static FName WMRSytemName = FName(TEXT("WindowsMixedRealityHMD"));
 bool FVREditorModeManager::IsVREditorAvailable() const
 {
-	if (GEngine->XRSystem.IsValid() && GEngine->XRSystem->GetHMDDevice() && GEngine->XRSystem->GetHMDDevice()->IsHMDEnabled())
+	if (!GetDefault<UVRModeSettings>()->ModeClass.LoadSynchronous())
+	{
+		return false;
+	}
+
+	if (GEditor->IsPlayingSessionInEditor())
 	{
-		// TODO: UE-71871 Work around for avoiding starting VRMode when using WMR
-		FName SystemName = GEngine->XRSystem->GetSystemName();
-		const bool bIsWMR = SystemName == WMRSytemName;
-		return !bIsWMR && !GEditor->IsPlayingSessionInEditor();
+		return false;
 	}
-	else
+
+	if (!GEngine->XRSystem.IsValid())
 	{
 		return false;
 	}
+
+	// TODO: UE-71871 Work around for avoiding starting VRMode when using WMR
+	if (GEngine->XRSystem->GetSystemName() == WMRSytemName)
+	{
+		return false;
+	}
+
+	return GEngine->XRSystem->GetHMDDevice() && GEngine->XRSystem->GetHMDDevice()->IsHMDEnabled();
 }
 
 bool FVREditorModeManager::IsVREditorButtonActive() const
diff --git a/Templates/TP_ME_BlankBP/Config/DefaultEditorSettings.ini b/Templates/TP_ME_BlankBP/Config/DefaultEditorSettings.ini
--- a/Templates/TP_ME_BlankBP/Config/DefaultEditorSettings.ini
+++ b/Templates/TP_ME_BlankBP/Config/DefaultEditorSettings.ini
@@ -11,8 +11,6 @@ GizmoScale=0.800000
 DoubleClickTime=0.250000
 TriggerPressedThreshold_Vive=0.330000
 TriggerPressedThreshold_Rift=0.500000
-InteractorClass=/VirtualProductionUtilities/VR/VirtualScoutingInteractor.VirtualScoutingInteractor_C
-TeleporterClass=BlueprintGeneratedClass'/VirtualProductionUtilities/VR/VirtualScoutingTeleporter.VirtualScoutingTeleporter_C'
 bScaleWorldWithDynamicPivot=True
 bAllowSimultaneousWorldScalingAndRotation=True
 
diff --git a/Templates/TP_ME_VProdBP/Config/DefaultEditor.ini b/Templates/TP_ME_VProdBP/Config/DefaultEditor.ini
--- a/Templates/TP_ME_VProdBP/Config/DefaultEditor.ini
+++ b/Templates/TP_ME_VProdBP/Config/DefaultEditor.ini
@@ -19,8 +19,6 @@ GizmoScale=0.800000
 DoubleClickTime=0.250000
 TriggerPressedThreshold_Vive=0.330000
 TriggerPressedThreshold_Rift=0.500000
-InteractorClass=/VirtualProductionUtilities/VR/VirtualScoutingInteractor.VirtualScoutingInteractor_C
-TeleporterClass=BlueprintGeneratedClass'/VirtualProductionUtilities/VR/VirtualScoutingTeleporter.VirtualScoutingTeleporter_C'
 bScaleWorldWithDynamicPivot=True
 bAllowSimultaneousWorldScalingAndRotation=True
 
diff --git a/Templates/TP_ME_VProdBP/Config/DefaultEditorSettings.ini b/Templates/TP_ME_VProdBP/Config/DefaultEditorSettings.ini
--- a/Templates/TP_ME_VProdBP/Config/DefaultEditorSettings.ini
+++ b/Templates/TP_ME_VProdBP/Config/DefaultEditorSettings.ini
@@ -11,8 +11,6 @@ GizmoScale=0.800000
 DoubleClickTime=0.250000
 TriggerPressedThreshold_Vive=0.330000
 TriggerPressedThreshold_Rift=0.500000
-InteractorClass=/VirtualProductionUtilities/VR/VirtualScoutingInteractor.VirtualScoutingInteractor_C
-TeleporterClass=BlueprintGeneratedClass'/VirtualProductionUtilities/VR/VirtualScoutingTeleporter.VirtualScoutingTeleporter_C'
 bScaleWorldWithDynamicPivot=True
 bAllowSimultaneousWorldScalingAndRotation=True
 
