VREditor:

Removed redundant private include paths from build.cs files.
Fixed include paths to be relative to the private or public folders.
Hid or removed includes that reached into other private module folders.
Updated PublicInclude paths when necessary.

diff --git a/Engine/Source/Editor/VREditor/Teleporter/VREditorAutoScaler.cpp b/Engine/Source/Editor/VREditor/Teleporter/VREditorAutoScaler.cpp
--- a/Engine/Source/Editor/VREditor/Teleporter/VREditorAutoScaler.cpp
+++ b/Engine/Source/Editor/VREditor/Teleporter/VREditorAutoScaler.cpp
@@ -1,6 +1,6 @@
 // Copyright Epic Games, Inc. All Rights Reserved.
 
-#include "VREditorAutoScaler.h"
+#include "Teleporter/VREditorAutoScaler.h"
 #include "VREditorMode.h"
 #include "ViewportInteractionTypes.h"
 #include "ViewportWorldInteraction.h"
diff --git a/Engine/Source/Editor/VREditor/Teleporter/VREditorTeleporter.cpp b/Engine/Source/Editor/VREditor/Teleporter/VREditorTeleporter.cpp
--- a/Engine/Source/Editor/VREditor/Teleporter/VREditorTeleporter.cpp
+++ b/Engine/Source/Editor/VREditor/Teleporter/VREditorTeleporter.cpp
@@ -1,6 +1,6 @@
 // Copyright Epic Games, Inc. All Rights Reserved.
 
-#include "VREditorTeleporter.h"
+#include "Teleporter/VREditorTeleporter.h"
 
 #include "Components/StaticMeshComponent.h"
 #include "Engine/EngineTypes.h"
@@ -12,7 +12,7 @@
 #include "Materials/MaterialInstanceDynamic.h"
 #include "ViewportInteractionTypes.h"
 #include "ViewportWorldInteraction.h"
-#include "VREditorTeleporter.h"
+#include "Teleporter/VREditorTeleporter.h"
 #include "VREditorMode.h"
 #include "VREditorAssetContainer.h"
 #include "VREditorInteractor.h"
diff --git a/Engine/Source/Editor/VREditor/UI/VREditorBaseUserWidget.cpp b/Engine/Source/Editor/VREditor/UI/VREditorBaseUserWidget.cpp
--- a/Engine/Source/Editor/VREditor/UI/VREditorBaseUserWidget.cpp
+++ b/Engine/Source/Editor/VREditor/UI/VREditorBaseUserWidget.cpp
@@ -1,7 +1,7 @@
 // Copyright Epic Games, Inc. All Rights Reserved.
 
-#include "VREditorBaseUserWidget.h"
-#include "VREditorFloatingUI.h"
+#include "UI/VREditorBaseUserWidget.h"
+#include "UI/VREditorFloatingUI.h"
 
 UVREditorBaseUserWidget::UVREditorBaseUserWidget( const FObjectInitializer& ObjectInitializer )
 	: Super( ObjectInitializer ),
diff --git a/Engine/Source/Editor/VREditor/UI/VREditorDockableCameraWindow.cpp b/Engine/Source/Editor/VREditor/UI/VREditorDockableCameraWindow.cpp
--- a/Engine/Source/Editor/VREditor/UI/VREditorDockableCameraWindow.cpp
+++ b/Engine/Source/Editor/VREditor/UI/VREditorDockableCameraWindow.cpp
@@ -1,6 +1,6 @@
 // Copyright Epic Games, Inc. All Rights Reserved.
 
-#include "VREditorDockableCameraWindow.h"
+#include "UI/VREditorDockableCameraWindow.h"
 #include "VREditorCameraWidgetComponent.h"
 
 AVREditorDockableCameraWindow::AVREditorDockableCameraWindow(const FObjectInitializer& ObjectInitializer) :
diff --git a/Engine/Source/Editor/VREditor/UI/VREditorDockableCameraWindow.h b/Engine/Source/Editor/VREditor/UI/VREditorDockableCameraWindow.h
--- a/Engine/Source/Editor/VREditor/UI/VREditorDockableCameraWindow.h
+++ b/Engine/Source/Editor/VREditor/UI/VREditorDockableCameraWindow.h
@@ -2,7 +2,7 @@
 
 #pragma once
 
-#include "VREditorDockableWindow.h"
+#include "UI/VREditorDockableWindow.h"
 #include "VREditorDockableCameraWindow.generated.h"
 
 /**
diff --git a/Engine/Source/Editor/VREditor/UI/VREditorDockableWindow.cpp b/Engine/Source/Editor/VREditor/UI/VREditorDockableWindow.cpp
--- a/Engine/Source/Editor/VREditor/UI/VREditorDockableWindow.cpp
+++ b/Engine/Source/Editor/VREditor/UI/VREditorDockableWindow.cpp
@@ -1,12 +1,12 @@
 // Copyright Epic Games, Inc. All Rights Reserved.
 
-#include "VREditorDockableWindow.h"
+#include "UI/VREditorDockableWindow.h"
 #include "UObject/ConstructorHelpers.h"
 #include "Components/StaticMeshComponent.h"
 #include "Materials/Material.h"
 #include "Engine/StaticMesh.h"
 #include "Materials/MaterialInstanceDynamic.h"
-#include "VREditorUISystem.h"
+#include "UI/VREditorUISystem.h"
 #include "VREditorMode.h"
 #include "VREditorWidgetComponent.h"
 #include "ViewportWorldInteraction.h"
diff --git a/Engine/Source/Editor/VREditor/UI/VREditorDockableWindow.h b/Engine/Source/Editor/VREditor/UI/VREditorDockableWindow.h
--- a/Engine/Source/Editor/VREditor/UI/VREditorDockableWindow.h
+++ b/Engine/Source/Editor/VREditor/UI/VREditorDockableWindow.h
@@ -4,7 +4,7 @@
 
 #include "CoreMinimal.h"
 #include "UObject/ObjectMacros.h"
-#include "VREditorFloatingUI.h"
+#include "UI/VREditorFloatingUI.h"
 #include "ViewportDragOperation.h"
 #include "ViewportInteractableInterface.h"
 #include "VREditorDockableWindow.generated.h"
diff --git a/Engine/Source/Editor/VREditor/UI/VREditorFloatingCameraUI.cpp b/Engine/Source/Editor/VREditor/UI/VREditorFloatingCameraUI.cpp
--- a/Engine/Source/Editor/VREditor/UI/VREditorFloatingCameraUI.cpp
+++ b/Engine/Source/Editor/VREditor/UI/VREditorFloatingCameraUI.cpp
@@ -1,8 +1,8 @@
 // Copyright Epic Games, Inc. All Rights Reserved.
 
-#include "VREditorFloatingCameraUI.h"
-#include "VREditorUISystem.h"
-#include "VREditorBaseUserWidget.h"
+#include "UI/VREditorFloatingCameraUI.h"
+#include "UI/VREditorUISystem.h"
+#include "UI/VREditorBaseUserWidget.h"
 #include "VREditorMode.h"
 #include "VREditorCameraWidgetComponent.h"
 #include "Components/StaticMeshComponent.h"
diff --git a/Engine/Source/Editor/VREditor/UI/VREditorFloatingCameraUI.h b/Engine/Source/Editor/VREditor/UI/VREditorFloatingCameraUI.h
--- a/Engine/Source/Editor/VREditor/UI/VREditorFloatingCameraUI.h
+++ b/Engine/Source/Editor/VREditor/UI/VREditorFloatingCameraUI.h
@@ -5,7 +5,7 @@
 #include "CoreMinimal.h"
 #include "UObject/ObjectMacros.h"
 #include "Widgets/SWidget.h"
-#include "VREditorFloatingUI.h"
+#include "UI/VREditorFloatingUI.h"
 #include "VREditorFloatingCameraUI.generated.h"
 
 class UVREditorBaseUserWidget;
diff --git a/Engine/Source/Editor/VREditor/UI/VREditorFloatingUI.cpp b/Engine/Source/Editor/VREditor/UI/VREditorFloatingUI.cpp
--- a/Engine/Source/Editor/VREditor/UI/VREditorFloatingUI.cpp
+++ b/Engine/Source/Editor/VREditor/UI/VREditorFloatingUI.cpp
@@ -1,8 +1,8 @@
 // Copyright Epic Games, Inc. All Rights Reserved.
 
-#include "VREditorFloatingUI.h"
-#include "VREditorUISystem.h"
-#include "VREditorBaseUserWidget.h"
+#include "UI/VREditorFloatingUI.h"
+#include "UI/VREditorUISystem.h"
+#include "UI/VREditorBaseUserWidget.h"
 #include "VREditorMode.h"
 #include "Components/WidgetComponent.h"
 #include "VREditorWidgetComponent.h"
diff --git a/Engine/Source/Editor/VREditor/UI/VREditorRadialFloatingUI.cpp b/Engine/Source/Editor/VREditor/UI/VREditorRadialFloatingUI.cpp
--- a/Engine/Source/Editor/VREditor/UI/VREditorRadialFloatingUI.cpp
+++ b/Engine/Source/Editor/VREditor/UI/VREditorRadialFloatingUI.cpp
@@ -1,8 +1,8 @@
 // Copyright Epic Games, Inc. All Rights Reserved.
 
-#include "VREditorRadialFloatingUI.h"
-#include "VREditorUISystem.h"
-#include "VREditorBaseUserWidget.h"
+#include "UI/VREditorRadialFloatingUI.h"
+#include "UI/VREditorUISystem.h"
+#include "UI/VREditorBaseUserWidget.h"
 #include "VREditorMode.h"
 #include "Components/WidgetComponent.h"
 #include "VREditorWidgetComponent.h"
diff --git a/Engine/Source/Editor/VREditor/UI/VREditorUISystem.cpp b/Engine/Source/Editor/VREditor/UI/VREditorUISystem.cpp
--- a/Engine/Source/Editor/VREditor/UI/VREditorUISystem.cpp
+++ b/Engine/Source/Editor/VREditor/UI/VREditorUISystem.cpp
@@ -1,6 +1,6 @@
 // Copyright Epic Games, Inc. All Rights Reserved.
 
-#include "VREditorUISystem.h"
+#include "UI/VREditorUISystem.h"
 #include "Misc/CommandLine.h"
 #include "HAL/IConsoleManager.h"
 #include "GenericPlatform/GenericApplication.h"
@@ -33,11 +33,11 @@
 
 #include "VREditorMode.h"
 #include "VREditorBaseActor.h"
-#include "VREditorBaseUserWidget.h"
-#include "VREditorFloatingUI.h"
-#include "VREditorRadialFloatingUI.h"
-#include "VREditorDockableWindow.h"
-#include "VREditorDockableCameraWindow.h"
+#include "UI/VREditorBaseUserWidget.h"
+#include "UI/VREditorFloatingUI.h"
+#include "UI/VREditorRadialFloatingUI.h"
+#include "UI/VREditorDockableWindow.h"
+#include "UI/VREditorDockableCameraWindow.h"
 #include "ViewportInteractionTypes.h"
 #include "IHeadMountedDisplay.h"
 #include "ViewportWorldInteraction.h"
@@ -89,7 +89,7 @@
 #include "VRModeSettings.h"
 #include "SequencerSettings.h"
 #include "Engine/Selection.h"
-#include "VREditorFloatingCameraUI.h"
+#include "UI/VREditorFloatingCameraUI.h"
 #include "ScopedTransaction.h"
 #include "Subsystems/AssetEditorSubsystem.h"
 #include "EditorModes.h"
diff --git a/Engine/Source/Editor/VREditor/UI/VREditorUISystem.h b/Engine/Source/Editor/VREditor/UI/VREditorUISystem.h
--- a/Engine/Source/Editor/VREditor/UI/VREditorUISystem.h
+++ b/Engine/Source/Editor/VREditor/UI/VREditorUISystem.h
@@ -7,7 +7,7 @@
 #include "UObject/ObjectMacros.h"
 #include "UObject/Object.h"
 #include "Widgets/SOverlay.h"
-#include "VRRadialMenuHandler.h"
+#include "UI/VRRadialMenuHandler.h"
 #include "EditorWorldExtension.h"
 #include "VREditorUISystem.generated.h"
 
@@ -440,4 +440,4 @@ protected:
 	bool bPanelCanScale;
 
 	FDelegateHandle AssetEditorCloseDelegate;
-};
\ No newline at end of file
+};
diff --git a/Engine/Source/Editor/VREditor/UI/VRRadialMenuHandler.cpp b/Engine/Source/Editor/VREditor/UI/VRRadialMenuHandler.cpp
--- a/Engine/Source/Editor/VREditor/UI/VRRadialMenuHandler.cpp
+++ b/Engine/Source/Editor/VREditor/UI/VRRadialMenuHandler.cpp
@@ -1,6 +1,6 @@
 // Copyright Epic Games, Inc. All Rights Reserved.
 
-#include "VRRadialMenuHandler.h"
+#include "UI/VRRadialMenuHandler.h"
 
 #include "Widgets/SNullWidget.h"
 #include "Widgets/DeclarativeSyntaxSupport.h"
@@ -28,11 +28,11 @@
 #include "UnrealEdGlobals.h"
 
 #include "VREditorMode.h"
-#include "VREditorUISystem.h"
+#include "UI/VREditorUISystem.h"
 #include "VREditorBaseActor.h"
-#include "VREditorBaseUserWidget.h"
-#include "VREditorFloatingUI.h"
-#include "VREditorDockableWindow.h"
+#include "UI/VREditorBaseUserWidget.h"
+#include "UI/VREditorFloatingUI.h"
+#include "UI/VREditorDockableWindow.h"
 #include "ViewportInteractionTypes.h"
 #include "IHeadMountedDisplay.h"
 #include "ViewportWorldInteraction.h"
@@ -53,9 +53,9 @@
 #include "ContentBrowserModule.h"
 
 // World Outliner
-#include "Editor/SceneOutliner/Public/ISceneOutliner.h"
-#include "Editor/SceneOutliner/Public/SceneOutlinerPublicTypes.h"
-#include "Editor/SceneOutliner/Public/SceneOutlinerModule.h"
+#include "ISceneOutliner.h"
+#include "SceneOutlinerPublicTypes.h"
+#include "SceneOutlinerModule.h"
 
 // Actor Details, Modes
 #include "SEditorViewport.h"
@@ -80,7 +80,7 @@
 #include "VREditorActions.h"
 #include "Framework/Commands/GenericCommands.h"
 #include "ISequencer.h"
-#include "Editor/LevelEditor/Private/SLevelEditor.h"
+#include "SLevelEditor.h"
 #include "Kismet2/DebuggerCommands.h"
 #include "UMGStyle.h"
 
diff --git a/Engine/Source/Editor/VREditor/VREditor.Build.cs b/Engine/Source/Editor/VREditor/VREditor.Build.cs
--- a/Engine/Source/Editor/VREditor/VREditor.Build.cs
+++ b/Engine/Source/Editor/VREditor/VREditor.Build.cs
@@ -6,6 +6,12 @@ namespace UnrealBuildTool.Rules
 	{
         public VREditor(ReadOnlyTargetRules Target) : base(Target)
 		{
+			PrivateIncludePaths.AddRange(
+				new string[] {
+					"Editor/LevelEditor/Private",
+				}
+			);
+
 			PublicIncludePaths.Add(ModuleDirectory);
 
             PrivateDependencyModuleNames.AddRange(
@@ -39,7 +45,9 @@ namespace UnrealBuildTool.Rules
 
 			PrivateIncludePathModuleNames.AddRange(
 				new string[] {
-					"PlacementMode"
+					"AssetTools",
+					"LevelEditor",
+					"PlacementMode",
 				}
 			);
 
@@ -48,14 +56,6 @@ namespace UnrealBuildTool.Rules
                     "PlacementMode"
                 }
 			);
-
-            PrivateIncludePaths.AddRange(
-                new string[] {
-                    "Editor/VREditor/UI",
-                    "Editor/VREditor/Teleporter",
-                }
-            );
-
 		}
 	}
 }
\ No newline at end of file
diff --git a/Engine/Source/Editor/VREditor/VREditorActions.cpp b/Engine/Source/Editor/VREditor/VREditorActions.cpp
--- a/Engine/Source/Editor/VREditor/VREditorActions.cpp
+++ b/Engine/Source/Editor/VREditor/VREditorActions.cpp
@@ -2,11 +2,11 @@
 
 #include "VREditorActions.h"
 #include "VREditorModule.h"
-#include "VREditorUISystem.h"
+#include "UI/VREditorUISystem.h"
 #include "VREditorMode.h"
 #include "ViewportWorldInteraction.h"
 #include "VREditorInteractor.h"
-#include "VREditorFloatingUI.h"
+#include "UI/VREditorFloatingUI.h"
 #include "SLevelViewport.h"
 #include "ImageUtils.h"
 #include "Misc/FileHelper.h"
@@ -14,8 +14,8 @@
 #include "ILevelEditor.h"
 #include "LevelEditor.h"
 
-#include "Developer/AssetTools/Public/IAssetTools.h"
-#include "Developer/AssetTools/Public/AssetToolsModule.h"
+#include "IAssetTools.h"
+#include "AssetToolsModule.h"
 #include "Modules/ModuleManager.h"
 #include "LevelSequence.h"
 #include "LevelSequenceActor.h"
diff --git a/Engine/Source/Editor/VREditor/VREditorActions.h b/Engine/Source/Editor/VREditor/VREditorActions.h
--- a/Engine/Source/Editor/VREditor/VREditorActions.h
+++ b/Engine/Source/Editor/VREditor/VREditorActions.h
@@ -1,7 +1,7 @@
 // Copyright Epic Games, Inc. All Rights Reserved.
 #pragma once
 
-#include "VREditorUISystem.h"
+#include "UI/VREditorUISystem.h"
 #include "Styling/SlateTypes.h"
 #include "InputCoreTypes.h"
 #include "Editor.h"
diff --git a/Engine/Source/Editor/VREditor/VREditorInteractor.cpp b/Engine/Source/Editor/VREditor/VREditorInteractor.cpp
--- a/Engine/Source/Editor/VREditor/VREditorInteractor.cpp
+++ b/Engine/Source/Editor/VREditor/VREditorInteractor.cpp
@@ -20,13 +20,13 @@
 #include "VREditorActions.h"
 #include "VREditorAssetContainer.h"
 #include "VREditorMode.h"
-#include "VREditorDockableWindow.h"
+#include "UI/VREditorDockableWindow.h"
 #include "VREditorFloatingText.h"
-#include "VREditorFloatingUI.h"
 #include "VREditorModule.h"
+#include "UI/VREditorFloatingUI.h"
 #include "VREditorPlacement.h"
-#include "VREditorRadialFloatingUI.h"
-#include "VREditorUISystem.h"
+#include "UI/VREditorRadialFloatingUI.h"
+#include "UI/VREditorUISystem.h"
 #include "VRModeSettings.h"
 #include "XRMotionControllerBase.h" // for FXRMotionControllerBase::Left/RightHandSourceId and GetHandEnumForSourceName()
 
diff --git a/Engine/Source/Editor/VREditor/VREditorMode.cpp b/Engine/Source/Editor/VREditor/VREditorMode.cpp
--- a/Engine/Source/Editor/VREditor/VREditorMode.cpp
+++ b/Engine/Source/Editor/VREditor/VREditorMode.cpp
@@ -15,13 +15,13 @@
 #include "DrawDebugHelpers.h"
 #include "Materials/MaterialInterface.h"
 #include "Materials/Material.h"
-#include "VREditorUISystem.h"
+#include "UI/VREditorUISystem.h"
 #include "VIBaseTransformGizmo.h"
 #include "ViewportWorldInteraction.h"
 #include "VREditorPlacement.h"
 #include "VREditorAvatarActor.h"
-#include "VREditorTeleporter.h"
-#include "VREditorAutoScaler.h"
+#include "Teleporter/VREditorTeleporter.h"
+#include "Teleporter/VREditorAutoScaler.h"
 #include "VREditorStyle.h"
 #include "VREditorAssetContainer.h"
 #include "Framework/Notifications/NotificationManager.h"
@@ -54,7 +54,7 @@
 #include "EngineUtils.h"
 #include "Kismet/GameplayStatics.h"
 #include "XRMotionControllerBase.h" // for FXRMotionControllerBase::Left/RightHandSourceId
-#include "VREditorFloatingUI.h"
+#include "UI/VREditorFloatingUI.h"
 #include "AssetEditorViewportLayout.h"
 
 #define LOCTEXT_NAMESPACE "VREditorMode"
diff --git a/Engine/Source/Editor/VREditor/VREditorModule.cpp b/Engine/Source/Editor/VREditor/VREditorModule.cpp
--- a/Engine/Source/Editor/VREditor/VREditorModule.cpp
+++ b/Engine/Source/Editor/VREditor/VREditorModule.cpp
@@ -11,7 +11,7 @@
 #include "Framework/MultiBox/MultiBoxBuilder.h"
 #include "Framework/MultiBox/MultiBoxExtender.h"
 #include "HeadMountedDisplayTypes.h"
-#include "VREditorFloatingUI.h"
+#include "UI/VREditorFloatingUI.h"
 
 
 DEFINE_LOG_CATEGORY(LogVREditor);
diff --git a/Engine/Source/Editor/VREditor/VREditorPlacement.cpp b/Engine/Source/Editor/VREditor/VREditorPlacement.cpp
--- a/Engine/Source/Editor/VREditor/VREditorPlacement.cpp
+++ b/Engine/Source/Editor/VREditor/VREditorPlacement.cpp
@@ -17,8 +17,8 @@
 #include "Kismet2/ComponentEditorUtils.h"
 #include "Sound/SoundCue.h"
 #include "Editor.h"
-#include "VREditorUISystem.h"
-#include "VREditorFloatingUI.h"
+#include "UI/VREditorUISystem.h"
+#include "UI/VREditorFloatingUI.h"
 #include "VREditorInteractor.h"
 
 // For actor placement
