VR Editor: Enumerate UVREditorMode-derived classes for selection in toolbar split button dropdown, and update toolbar icon. The dropdown menu also indicates whether an HMD is available, and if so its runtime version string.

Historically not having an HMD available has meant this toolbar button would be hidden entirely, and this is still the case if no derived/alternative UVREditorMode classes are available. However, if any additional mode classes are present (for example because the user has enabled the Virtual Scouting plugin), the toolbar button will always be shown, in the hopes that surfacing the HMD status will simplify the onboarding experience for new Virtual Scouting users.

If the configured mode class becomes unavailable (for example because the corresponding plugin was subsequently disabled), then its object path will be shown as selected and disabled in the dropdown, and the primary toggle button will also be disabled until a valid mode is selected.

diff --git a/Engine/Source/Editor/EditorStyle/Private/StarshipStyle.cpp b/Engine/Source/Editor/EditorStyle/Private/StarshipStyle.cpp
--- a/Engine/Source/Editor/EditorStyle/Private/StarshipStyle.cpp
+++ b/Engine/Source/Editor/EditorStyle/Private/StarshipStyle.cpp
@@ -4230,8 +4230,8 @@ void FStarshipEditorStyle::FStyle::SetupLevelEditorStyle()
 		Set( "LevelEditor.CreateClassBlueprint", new IMAGE_BRUSH("Icons/icon_class_Blueprint_New_16x", Icon16x16));
 		Set( "LevelEditor.OpenClassBlueprint", new IMAGE_BRUSH_SVG("Starship/Common/BrowseContent", Icon16x16));
 
-		Set( "LevelEditor.ToggleVR", new IMAGE_BRUSH( "Icons/VREditor/VR_Editor_Toolbar_Icon", Icon40x40 ) );
-		Set( "LevelEditor.ToggleVR.Small", new IMAGE_BRUSH( "Icons/VREditor/VR_Editor_Toolbar_Icon_Small", Icon20x20 ) );
+		Set( "LevelEditor.ToggleVR", new IMAGE_BRUSH_SVG( "Starship/MainToolbar/VRTools", Icon40x40 ) );
+		Set( "LevelEditor.ToggleVR.Small", new IMAGE_BRUSH_SVG( "Starship/MainToolbar/VRTools", Icon20x20 ) );
 
 		Set( "MergeActors.MeshMergingTool", new IMAGE_BRUSH_SVG( "Starship/MergeActors/MeshMerging_16", Icon16x16 ) );
 		Set( "MergeActors.MeshProxyTool", new IMAGE_BRUSH_SVG( "Starship/MergeActors/MeshProxy_16", Icon16x16 ) );
diff --git a/Engine/Source/Editor/LevelEditor/Private/LevelEditorToolBar.cpp b/Engine/Source/Editor/LevelEditor/Private/LevelEditorToolBar.cpp
--- a/Engine/Source/Editor/LevelEditor/Private/LevelEditorToolBar.cpp
+++ b/Engine/Source/Editor/LevelEditor/Private/LevelEditorToolBar.cpp
@@ -1384,8 +1384,6 @@ void FLevelEditorToolBar::RegisterLevelEditorToolBar( const TSharedRef<FUIComman
 			);
 			CinematicsEntry.StyleNameOverride = "AssetEditorToolbar";
 			Section.AddEntry(CinematicsEntry);
-
-			Section.AddEntry(FToolMenuEntry::InitToolBarButton(FLevelEditorCommands::Get().ToggleVR, LOCTEXT("ToggleVR", "VR Mode")));
 		}
 
 	}
diff --git a/Engine/Source/Editor/VREditor/VREditor.Build.cs b/Engine/Source/Editor/VREditor/VREditor.Build.cs
--- a/Engine/Source/Editor/VREditor/VREditor.Build.cs
+++ b/Engine/Source/Editor/VREditor/VREditor.Build.cs
@@ -33,7 +33,8 @@ namespace UnrealBuildTool.Rules
 					"Analytics",
                     "LevelSequence",
                     "Sequencer",
-                    "Projects"
+                    "Projects",
+					"ToolMenus",
 				}
 			);
 
diff --git a/Engine/Source/Editor/VREditor/VREditorModeManager.cpp b/Engine/Source/Editor/VREditor/VREditorModeManager.cpp
--- a/Engine/Source/Editor/VREditor/VREditorModeManager.cpp
+++ b/Engine/Source/Editor/VREditor/VREditorModeManager.cpp
@@ -157,7 +157,12 @@ bool FVREditorModeManager::IsVREditorAvailable() const
 bool FVREditorModeManager::IsVREditorButtonActive() const
 {
 	const bool bHasHMDDevice = GEngine->XRSystem.IsValid() && GEngine->XRSystem->GetHMDDevice() && GEngine->XRSystem->GetHMDDevice()->IsHMDEnabled();
-	return bHasHMDDevice;
+
+	TArray<UClass*> ModeClasses;
+	GetConcreteModeClasses(ModeClasses);
+	const bool bDerivedModeAvailable = ModeClasses.Num() > 1;
+
+	return bHasHMDDevice || bDerivedModeAvailable;
 }
 
 
@@ -288,4 +293,20 @@ void FVREditorModeManager::OnMapChanged( UWorld* World, EMapChangeType MapChange
 	CurrentVREditorMode = nullptr;
 }
 
+void FVREditorModeManager::GetConcreteModeClasses(TArray<UClass*>& OutModeClasses) const
+{
+	const int32 StartIndex = OutModeClasses.Num();
+
+	OutModeClasses.Add(UVREditorMode::StaticClass());
+	GetDerivedClasses(UVREditorMode::StaticClass(), OutModeClasses);
+
+	for (TArray<UClass*>::TIterator Iter = OutModeClasses.CreateIterator() + StartIndex; Iter; ++Iter)
+	{
+		if ((*Iter)->HasAllClassFlags(CLASS_Abstract))
+		{
+			Iter.RemoveCurrentSwap();
+		}
+	}
+}
+
 #undef LOCTEXT_NAMESPACE
diff --git a/Engine/Source/Editor/VREditor/VREditorModeManager.h b/Engine/Source/Editor/VREditor/VREditorModeManager.h
--- a/Engine/Source/Editor/VREditor/VREditorModeManager.h
+++ b/Engine/Source/Editor/VREditor/VREditorModeManager.h
@@ -63,6 +63,9 @@ public:
 	/** Return a multicast delegate which is executed when VR mode stops. */
 	FOnVREditingModeExit& OnVREditingModeExit() { return OnVREditingModeExitHandle; }
 
+	/** Appends all UVREditorMode-derived classes which are not abstract to the provided array. */
+	void GetConcreteModeClasses(TArray<UClass*>& OutModeClasses) const;
+
 private:
 
 	/** Broadcasts when VR mode is started */
diff --git a/Engine/Source/Editor/VREditor/VREditorModule.cpp b/Engine/Source/Editor/VREditor/VREditorModule.cpp
--- a/Engine/Source/Editor/VREditor/VREditorModule.cpp
+++ b/Engine/Source/Editor/VREditor/VREditorModule.cpp
@@ -4,16 +4,25 @@
 #include "Modules/ModuleManager.h"
 #include "Stats/Stats.h"
 #include "HAL/IConsoleManager.h"
+#include "IHeadMountedDisplay.h"
+#include "ISettingsModule.h"
 #include "IVREditorModule.h"
+#include "IXRTrackingSystem.h"
+#include "LevelEditorActions.h"
+#include "ToolMenus.h"
 #include "VREditorModeManager.h"
 #include "VREditorStyle.h"
 #include "VREditorMode.h"
+#include "VRModeSettings.h"
 #include "Framework/MultiBox/MultiBoxBuilder.h"
 #include "Framework/MultiBox/MultiBoxExtender.h"
 #include "HeadMountedDisplayTypes.h"
 #include "UI/VREditorFloatingUI.h"
 
 
+#define LOCTEXT_NAMESPACE "VREditorModule"
+
+
 DEFINE_LOG_CATEGORY(LogVREditor);
 
 
@@ -57,11 +66,16 @@ public:
 	/** Return a multicast delegate which is executed when VR mode stops. */
 	virtual FOnVREditingModeExit& OnVREditingModeExit() override { return ModeManager.OnVREditingModeExit(); }
 
+private:
+	void ExtendToolbarMenu();
+
 private:
 	TSharedPtr<class FExtender> RadialMenuExtender;
 
 	/** Handles turning VR Editor mode on and off */
 	FVREditorModeManager ModeManager;
+
+	FDelegateHandle DeferredInitDelegate;
 };
 
 namespace VREd
@@ -73,10 +87,22 @@ void FVREditorModule::StartupModule()
 {
 	LLM_SCOPE_BYNAME(TEXT("VREditor"));
 	RadialMenuExtender = MakeShareable(new FExtender());
+
+	// UToolMenus::RegisterStartupCallback is still too early.
+	DeferredInitDelegate = FCoreDelegates::OnBeginFrame.AddLambda(
+		[this]()
+		{
+			ExtendToolbarMenu();
+
+			// Must happen last; this implicitly deallocates this lambda's captures.
+			FCoreDelegates::OnBeginFrame.Remove(DeferredInitDelegate);
+		}
+	);
 }
 
 void FVREditorModule::ShutdownModule()
 {
+	UToolMenus::UnregisterOwner(this);
 
 	if (GIsEditor)
 	{
@@ -140,4 +166,148 @@ void FVREditorModule::ToggleForceVRMode()
 	Self.EnableVREditor( !Self.IsVREditorEnabled(), bForceWithoutHMD );
 }
 
+void FVREditorModule::ExtendToolbarMenu()
+{
+	FToolMenuOwnerScoped OwnerScoped(this);
+
+	UToolMenu* EditorToolbarMenu = UToolMenus::Get()->ExtendMenu("LevelEditor.LevelEditorToolBar.AssetsToolBar");
+	if (!ensure(EditorToolbarMenu))
+	{
+		return;
+	}
+
+	FToolMenuSection* EditorToolbarSection = EditorToolbarMenu->FindSection("Content");
+	if (!ensure(EditorToolbarSection))
+	{
+		return;
+	}
+
+	// The primary toggle button.
+	EditorToolbarSection->AddEntry(FToolMenuEntry::InitToolBarButton(
+		FLevelEditorCommands::Get().ToggleVR,
+		LOCTEXT("ToggleVR", "VR Mode")
+	));
+
+	// The "..." split, which summons the tool menu.
+	static const FName OptionsToolMenuName = "VrEditor.ToggleVrOptions";
+
+	EditorToolbarSection->AddEntry(FToolMenuEntry::InitComboButton(
+		"ToggleVrOptionsMenu",
+		FUIAction(),
+		FOnGetContent::CreateLambda([]() { return UToolMenus::Get()->GenerateWidget(OptionsToolMenuName, FToolMenuContext()); }),
+		TAttribute<FText>(),
+		LOCTEXT("ToggleVrOptions_ComboSplit_Tooltip", "Configure VR Editor Mode"),
+		TAttribute<FSlateIcon>(),
+		true
+	));
+
+	// Register the split options tool menu.
+	UToolMenu* OptionsMenu = UToolMenus::Get()->RegisterMenu(OptionsToolMenuName);
+
+	// The "Status" section indicates whether an HMD is available.
+	FToolMenuSection& StatusSection = OptionsMenu->AddSection("Status");
+	StatusSection.AddDynamicEntry("StatusEntry", FNewToolMenuSectionDelegate::CreateLambda(
+		[](FToolMenuSection& InSection)
+		{
+			const bool bHasHMDDevice = GEngine->XRSystem
+				&& GEngine->XRSystem->GetHMDDevice()
+				&& GEngine->XRSystem->GetHMDDevice()->IsHMDEnabled();
+
+			const FText HmdAvailableText = LOCTEXT("ToggleVrOptions_Status_HmdAvailable", "Headset available ({XrVersionString})");
+			const FText HmdNotAvailableText = LOCTEXT("ToggleVrOptions_Status_HmdNotAvailable", "No headset available");
+
+			const FText StatusText = bHasHMDDevice
+				? FText::FormatNamed(HmdAvailableText, TEXT("XrVersionString"), FText::FromString(GEngine->XRSystem->GetVersionString()))
+				: HmdNotAvailableText;
+
+			InSection.AddEntry(FToolMenuEntry::InitWidget(
+				"HeadsetStatus",
+				SNew(SBox)
+				.Padding(FMargin(16.0f, 3.0f))
+				[
+					SNew(STextBlock)
+					.ColorAndOpacity(FSlateColor::UseSubduedForeground())
+					.Text(StatusText)
+				],
+				FText::GetEmpty()
+			));
+		}));
+
+	// The "Modes" section enumerates the UVREditorMode-derived classes for radio selection.
+	OptionsMenu->AddDynamicSection("DynamicModesSection", FNewToolMenuDelegate::CreateLambda(
+		[this](UToolMenu* InMenu)
+		{
+			FToolMenuSection& Section = InMenu->AddSection("Modes", LOCTEXT("ToggleVrOptions_ModesSection_Label", "Modes"));
+
+			TSoftClassPtr<UVREditorMode> CurrentModeClassSoft = GetDefault<UVRModeSettings>()->ModeClass;
+			UClass* CurrentModeClass = CurrentModeClassSoft.LoadSynchronous();
+
+			TArray<UClass*> ModeClasses;
+			ModeManager.GetConcreteModeClasses(ModeClasses);
+
+			Algo::Sort(ModeClasses,
+				[](const UClass* Lhs, const UClass* Rhs)
+				{
+					return Lhs->GetDisplayNameText().CompareTo(Rhs->GetDisplayNameText()) < 0;
+				});
+
+			if (!ModeClasses.Contains(CurrentModeClass))
+			{
+				FToolUIAction UnavailableModeAction;
+				UnavailableModeAction.CanExecuteAction = FToolMenuCanExecuteAction::CreateLambda([](const FToolMenuContext&) { return false; });
+				UnavailableModeAction.GetActionCheckState = FToolMenuGetActionCheckState::CreateLambda([](const FToolMenuContext&) { return ECheckBoxState::Checked; });
+
+				Section.AddMenuEntry("UnavailableMode",
+					CurrentModeClassSoft.IsNull() ? LOCTEXT("NullModeClass", "(Unset)") : FText::AsCultureInvariant(CurrentModeClassSoft.ToString()),
+					TAttribute<FText>(),
+					TAttribute<FSlateIcon>(),
+					UnavailableModeAction,
+					EUserInterfaceActionType::RadioButton);
+			}
+
+			for (UClass* Mode : ModeClasses)
+			{
+				FToolUIAction SelectModeAction;
+				SelectModeAction.ExecuteAction = FToolMenuExecuteAction::CreateLambda(
+					[Mode](const FToolMenuContext&)
+					{
+						GetMutableDefault<UVRModeSettings>()->ModeClass = Mode;
+					});
+				SelectModeAction.CanExecuteAction = FToolMenuCanExecuteAction::CreateLambda(
+					[](const FToolMenuContext&)
+					{
+						return true;
+					});
+				SelectModeAction.GetActionCheckState = FToolMenuGetActionCheckState::CreateLambda(
+					[Mode](const FToolMenuContext&)
+					{
+						return Mode == GetDefault<UVRModeSettings>()->ModeClass.Get() ? ECheckBoxState::Checked : ECheckBoxState::Unchecked;
+					});
+
+				Section.AddMenuEntry(Mode->GetFName(),
+					Mode->GetDisplayNameText(),
+					FText::AsCultureInvariant(Mode->GetClassPathName().ToString()),
+					TAttribute<FSlateIcon>(),
+					SelectModeAction,
+					EUserInterfaceActionType::RadioButton);
+			}
+		}));
+
+	FToolMenuSection& SettingsSection = OptionsMenu->AddSection("Settings");
+
+	SettingsSection.AddSeparator("SettingsStartSeparator");
+
+	SettingsSection.AddMenuEntry(
+		"EditorSettings",
+		LOCTEXT("ToggleVrOptions_EditorSettings_Label", "Settings..."),
+		LOCTEXT("ToggleVrOptions_EditorSettings_Tooltip", "Jump to the VR Editor Mode section in the Editor Preferences"),
+		FSlateIcon(FAppStyle::GetAppStyleSetName(), "LevelEditor.GameSettings"),
+		FUIAction(FExecuteAction::CreateLambda([]() { FModuleManager::LoadModuleChecked<ISettingsModule>("Settings").ShowViewer("Editor", "General", "VR Mode"); }))
+	);
+}
+
+
 IMPLEMENT_MODULE( FVREditorModule, VREditor )
+
+
+#undef LOCTEXT_NAMESPACE
