VR Editor: Workaround for tool menu differences triggering ensure in another stream.

diff --git a/Engine/Source/Editor/VREditor/VREditorModule.cpp b/Engine/Source/Editor/VREditor/VREditorModule.cpp
--- a/Engine/Source/Editor/VREditor/VREditorModule.cpp
+++ b/Engine/Source/Editor/VREditor/VREditorModule.cpp
@@ -177,7 +177,7 @@ void FVREditorModule::ExtendToolbarMenu()
 	}
 
 	FToolMenuSection* EditorToolbarSection = EditorToolbarMenu->FindSection("Content");
-	if (!ensure(EditorToolbarSection))
+	if (!EditorToolbarSection)
 	{
 		return;
 	}
