Virtual Scouting: Fix Oculus controller grip pose correction under varying world scales. Also fixes a scale-related bug for Oculus "help system" tooltip positioning.

diff --git a/Engine/Source/Editor/VREditor/Public/VREditorInteractor.h b/Engine/Source/Editor/VREditor/Public/VREditorInteractor.h
--- a/Engine/Source/Editor/VREditor/Public/VREditorInteractor.h
+++ b/Engine/Source/Editor/VREditor/Public/VREditorInteractor.h
@@ -283,6 +283,12 @@ protected:
 	UPROPERTY(BlueprintReadWrite, Category = "VREditorInteractor")
 	TObjectPtr<class UStaticMeshComponent> HandMeshComponent;
 
+	UFUNCTION(BlueprintNativeEvent)
+	void UpdateHandMeshRelativeTransform();
+
+	FVector3d HandMeshBaseScale;
+	FTransform HandMeshGripTransform;
+
 
 private:
 
diff --git a/Engine/Source/Editor/VREditor/VREditorInteractor.cpp b/Engine/Source/Editor/VREditor/VREditorInteractor.cpp
--- a/Engine/Source/Editor/VREditor/VREditorInteractor.cpp
+++ b/Engine/Source/Editor/VREditor/VREditorInteractor.cpp
@@ -101,6 +101,8 @@ UVREditorInteractor::UVREditorInteractor() :
 	MotionControllerComponent( nullptr ),
 	LaserMotionControllerComponent( nullptr ),
 	HandMeshComponent( nullptr ),
+	HandMeshBaseScale( FVector3d::OneVector ),
+	HandMeshGripTransform( FTransform::Identity ),
 	LaserSplineComponent( nullptr ),
 	LaserPointerMID( nullptr ),
 	TranslucentLaserPointerMID( nullptr ),
@@ -323,7 +325,16 @@ void UVREditorInteractor::SetupComponent_Implementation( AActor* OwningActor )
 
 	// Hand mesh
 	{
-		HandMeshComponent = VRMode->CreateMotionControllerMesh(OwningAvatar, MotionControllerComponent );
+		HandMeshComponent = VRMode->CreateMotionControllerMesh(OwningAvatar, MotionControllerComponent);
+		if (ControllerMotionSource == FXRMotionControllerBase::RightHandSourceId &&
+			GetHMDDeviceType() == OculusDeviceType)	// Oculus has asymmetrical controllers, so we mirror the mesh horizontally
+		{
+			HandMeshBaseScale = FVector3d(1.0, -1.0, 1.0);
+		}
+		else
+		{
+			HandMeshBaseScale = FVector3d(1.0, 1.0, 1.0);
+		}
 		SetHandMeshComponentProperties();
 
 		UMaterialInterface* HandMeshMaterial = GetVRMode().GetHMDDeviceType() == SteamVRDeviceType ? AssetContainer.VivePreControllerMaterial : AssetContainer.OculusControllerMaterial;
@@ -454,7 +465,7 @@ void UVREditorInteractor::ReplaceHandMeshComponent(
 	HandMeshComponent->UnregisterComponent();
 	HandMeshComponent->DestroyComponent();
 	HandMeshComponent = VRMode->CreateMotionControllerMesh(OwningAvatar, MotionControllerComponent, NewMesh);
-	HandMeshComponent->SetRelativeScale3D(MeshScale);
+	HandMeshBaseScale = MeshScale;
 	SetHandMeshComponentProperties();
 }
 
@@ -467,18 +478,28 @@ void UVREditorInteractor::SetHandMeshComponentProperties()
 	HandMeshComponent->SetCollisionResponseToAllChannels(ECR_Overlap);
 	HandMeshComponent->SetGenerateOverlapEvents(true);
 
+	HandMeshGripTransform = FTransform::Identity;
+
 	// To help ease the transition to OpenXR, an alternative controller mesh origin can be specified.
 	if (GEngine->XRSystem.IsValid() && GEngine->XRSystem->GetSystemName() == OpenXRDeviceType)
 	{
 		const FName GripSocket("OpenXrGrip");
 		if (HandMeshComponent->DoesSocketExist(GripSocket))
 		{
-			const FTransform GripTransform = HandMeshComponent->GetSocketTransform(GripSocket, RTS_Component);
-			const FTransform ScaleTransform(FRotator::ZeroRotator, FVector::ZeroVector, HandMeshComponent->GetRelativeScale3D());
-			const FTransform NewRelativeTransform = (GripTransform * ScaleTransform).Inverse();
-			HandMeshComponent->SetRelativeTransform(NewRelativeTransform);
+			HandMeshGripTransform = HandMeshComponent->GetSocketTransform(GripSocket, RTS_Component);
 		}
 	}
+
+	UpdateHandMeshRelativeTransform();
+}
+
+void UVREditorInteractor::UpdateHandMeshRelativeTransform_Implementation()
+{
+	// The hands need to stay the same size relative to our tracking space, so we inverse compensate for world to meters scale here
+	const float WorldScaleFactor = WorldInteraction->GetWorldScaleFactor();
+	const FTransform ScaleTransform(FRotator::ZeroRotator, FVector::ZeroVector, HandMeshBaseScale / WorldScaleFactor);
+	const FTransform NewRelativeTransform = (ScaleTransform * HandMeshGripTransform).Inverse();
+	HandMeshComponent->SetRelativeTransform(NewRelativeTransform);
 }
 
 void UVREditorInteractor::Shutdown_Implementation()
@@ -539,25 +560,11 @@ void UVREditorInteractor::Tick_Implementation( const float DeltaTime )
 {
 	Super::Tick_Implementation( DeltaTime );
 
-	{
-		const float WorldScaleFactor = WorldInteraction->GetWorldScaleFactor();
-
-		// @todo vreditor: Manually ticking motion controller components
-		MotionControllerComponent->TickComponent( DeltaTime, ELevelTick::LEVELTICK_PauseTick, nullptr );
-		LaserMotionControllerComponent->TickComponent( DeltaTime, ELevelTick::LEVELTICK_PauseTick, nullptr );
+	// @todo vreditor: Manually ticking motion controller components
+	MotionControllerComponent->TickComponent( DeltaTime, ELevelTick::LEVELTICK_PauseTick, nullptr );
+	LaserMotionControllerComponent->TickComponent( DeltaTime, ELevelTick::LEVELTICK_PauseTick, nullptr );
 
-		// The hands need to stay the same size relative to our tracking space, so we inverse compensate for world to meters scale here
-		// NOTE: We don't need to set the hand mesh location and rotation, as the MotionControllerComponent does that itself
-		if (ControllerMotionSource == FXRMotionControllerBase::RightHandSourceId &&
-			GetHMDDeviceType() == OculusDeviceType)	// Oculus has asymmetrical controllers, so we mirror the mesh horizontally
-		{
-			HandMeshComponent->SetRelativeScale3D( FVector( WorldScaleFactor, -WorldScaleFactor, WorldScaleFactor ) );
-		}
-		else
-		{
-			HandMeshComponent->SetRelativeScale3D( FVector( WorldScaleFactor ) );
-		}
-	}
+	UpdateHandMeshRelativeTransform();
 
 	UpdateRadialMenuInput( DeltaTime );
 
diff --git a/Engine/Source/Editor/ViewportInteraction/ViewportWorldInteraction.cpp b/Engine/Source/Editor/ViewportInteraction/ViewportWorldInteraction.cpp
--- a/Engine/Source/Editor/ViewportInteraction/ViewportWorldInteraction.cpp
+++ b/Engine/Source/Editor/ViewportInteraction/ViewportWorldInteraction.cpp
@@ -758,7 +758,14 @@ void UViewportWorldInteraction::SetRoomTransformForNextFrame(const FTransform& N
 
 float UViewportWorldInteraction::GetWorldScaleFactor() const
 {
-	return GetWorld()->GetWorldSettings()->WorldToMeters / 100.0f;
+	if (GetWorld() && GetWorld()->GetWorldSettings())
+	{
+		return GetWorld()->GetWorldSettings()->WorldToMeters / 100.0f;
+	}
+	else
+	{
+		return 1.0f;
+	}
 }
 
 FEditorViewportClient* UViewportWorldInteraction::GetDefaultOptionalViewportClient() const
