VR Editor: More robust mode class enumeration, handles blueprinted modes (shows unloaded, hides skeletons) correctly.

diff --git a/Engine/Source/Editor/VREditor/VREditorModeManager.cpp b/Engine/Source/Editor/VREditor/VREditorModeManager.cpp
--- a/Engine/Source/Editor/VREditor/VREditorModeManager.cpp
+++ b/Engine/Source/Editor/VREditor/VREditorModeManager.cpp
@@ -1,6 +1,7 @@
 // Copyright Epic Games, Inc. All Rights Reserved.
 
 #include "VREditorModeManager.h"
+#include "AssetRegistry/AssetRegistryModule.h"
 #include "InputCoreTypes.h"
 #include "VREditorMode.h"
 #include "Modules/ModuleManager.h"
@@ -32,10 +33,29 @@ FVREditorModeManager::FVREditorModeManager() :
 	HMDWornState( EHMDWornState::Unknown ), 
 	bAddedViewportWorldInteractionExtension( false )
 {
+	IAssetRegistry& AssetRegistry = FModuleManager::LoadModuleChecked<FAssetRegistryModule>(AssetRegistryConstants::ModuleName).Get();
+	if (AssetRegistry.IsLoadingAssets())
+	{
+		AssetRegistry.OnFilesLoaded().AddRaw(this, &FVREditorModeManager::HandleAssetFilesLoaded);
+	}
+	else
+	{
+		HandleAssetFilesLoaded();
+	}
 }
 
 FVREditorModeManager::~FVREditorModeManager()
 {
+	if (FAssetRegistryModule* AssetRegistryModule = FModuleManager::GetModulePtr<FAssetRegistryModule>(AssetRegistryConstants::ModuleName))
+	{
+		if (IAssetRegistry* AssetRegistry = AssetRegistryModule->TryGet())
+		{
+			AssetRegistry->OnFilesLoaded().RemoveAll(this);
+			AssetRegistry->OnAssetAdded().RemoveAll(this);
+			AssetRegistry->OnAssetRemoved().RemoveAll(this);
+		}
+	}
+
 	if (CurrentVREditorMode)
 	{
 		CurrentVREditorMode->OnVRModeEntryComplete().RemoveAll(this);
@@ -157,10 +177,7 @@ bool FVREditorModeManager::IsVREditorAvailable() const
 bool FVREditorModeManager::IsVREditorButtonActive() const
 {
 	const bool bHasHMDDevice = GEngine->XRSystem.IsValid() && GEngine->XRSystem->GetHMDDevice() && GEngine->XRSystem->GetHMDDevice()->IsHMDEnabled();
-
-	TArray<UClass*> ModeClasses;
-	GetConcreteModeClasses(ModeClasses);
-	const bool bDerivedModeAvailable = ModeClasses.Num() > 1;
+	const bool bDerivedModeAvailable = CachedModeClassPaths.Num() > 1;
 
 	return bHasHMDDevice || bDerivedModeAvailable;
 }
@@ -295,18 +312,47 @@ void FVREditorModeManager::OnMapChanged( UWorld* World, EMapChangeType MapChange
 
 void FVREditorModeManager::GetConcreteModeClasses(TArray<UClass*>& OutModeClasses) const
 {
-	const int32 StartIndex = OutModeClasses.Num();
-
-	OutModeClasses.Add(UVREditorMode::StaticClass());
-	GetDerivedClasses(UVREditorMode::StaticClass(), OutModeClasses);
-
-	for (TArray<UClass*>::TIterator Iter = OutModeClasses.CreateIterator() + StartIndex; Iter; ++Iter)
+	for (const FTopLevelAssetPath& ClassPath : CachedModeClassPaths)
 	{
-		if ((*Iter)->HasAllClassFlags(CLASS_Abstract))
+		TSoftClassPtr<UVREditorMode> SoftClass = TSoftClassPtr<UVREditorMode>(FSoftObjectPath(ClassPath));
+		if (UClass* Class = SoftClass.LoadSynchronous())
 		{
-			Iter.RemoveCurrentSwap();
+			if (Class->HasAnyClassFlags(CLASS_Abstract) || (Class->GetAuthoritativeClass() != Class))
+			{
+				continue;
+			}
+
+			OutModeClasses.Add(Class);
 		}
 	}
 }
 
+void FVREditorModeManager::GetModeClassPaths(TSet<FTopLevelAssetPath>& OutModeClasses) const
+{
+	TArray<FTopLevelAssetPath> BaseModeClasses;
+	BaseModeClasses.Add(UVREditorMode::StaticClass()->GetClassPathName());
+
+	IAssetRegistry& AssetRegistry = FModuleManager::LoadModuleChecked<FAssetRegistryModule>(AssetRegistryConstants::ModuleName).Get();
+	AssetRegistry.GetDerivedClassNames(BaseModeClasses, TSet<FTopLevelAssetPath>(), OutModeClasses);
+}
+
+void FVREditorModeManager::UpdateCachedModeClassPaths()
+{
+	GetModeClassPaths(CachedModeClassPaths);
+}
+
+void FVREditorModeManager::HandleAssetFilesLoaded()
+{
+	IAssetRegistry& AssetRegistry = FModuleManager::LoadModuleChecked<FAssetRegistryModule>(AssetRegistryConstants::ModuleName).Get();
+	AssetRegistry.OnFilesLoaded().RemoveAll(this);
+	AssetRegistry.OnAssetAdded().AddRaw(this, &FVREditorModeManager::HandleAssetAddedOrRemoved);
+	AssetRegistry.OnAssetRemoved().AddRaw(this, &FVREditorModeManager::HandleAssetAddedOrRemoved);
+	UpdateCachedModeClassPaths();
+}
+
+void FVREditorModeManager::HandleAssetAddedOrRemoved(const FAssetData& NewAssetData)
+{
+	UpdateCachedModeClassPaths();
+}
+
 #undef LOCTEXT_NAMESPACE
diff --git a/Engine/Source/Editor/VREditor/VREditorModeManager.h b/Engine/Source/Editor/VREditor/VREditorModeManager.h
--- a/Engine/Source/Editor/VREditor/VREditorModeManager.h
+++ b/Engine/Source/Editor/VREditor/VREditorModeManager.h
@@ -63,10 +63,23 @@ public:
 	/** Return a multicast delegate which is executed when VR mode stops. */
 	FOnVREditingModeExit& OnVREditingModeExit() { return OnVREditingModeExitHandle; }
 
-	/** Appends all UVREditorMode-derived classes which are not abstract to the provided array. */
+	/**
+	 * Loads each cached UVREditorMode-derived class path, checking and appending only
+	 * those which are suitable for use (not abstract, etc.) to the provided array. */
 	void GetConcreteModeClasses(TArray<UClass*>& OutModeClasses) const;
 
 private:
+	/** Queries the asset registry for asset paths of all classes derived from UVREditorMode. */
+	void GetModeClassPaths(TSet<FTopLevelAssetPath>& OutModeClassPaths) const;
+
+	/** Stores the results of GetModeClassPaths() into CachedModeClassPaths. */
+	void UpdateCachedModeClassPaths();
+
+	/** Cached UVREditorMode-derived assets, updated in response to asset registry events. */
+	TSet<FTopLevelAssetPath> CachedModeClassPaths;
+
+	void HandleAssetFilesLoaded();
+	void HandleAssetAddedOrRemoved(const FAssetData& NewAssetData);
 
 	/** Broadcasts when VR mode is started */
 	FOnVREditingModeEnter OnVREditingModeEnterHandle;
