VR Editor: Resave config when mode changes are made via toolbar dropdown.

diff --git a/Engine/Source/Editor/VREditor/Public/VRModeSettings.h b/Engine/Source/Editor/VREditor/Public/VRModeSettings.h
--- a/Engine/Source/Editor/VREditor/Public/VRModeSettings.h
+++ b/Engine/Source/Editor/VREditor/Public/VRModeSettings.h
@@ -92,8 +92,7 @@ public:
 	UPROPERTY(EditAnywhere, config, NoClear, Category = "General")
 	TSoftClassPtr<UVREditorMode> ModeClass;
 
-private:
 #if WITH_EDITOR
-	void PostEditChangeProperty(struct FPropertyChangedEvent& PropertyChangedEvent);
+	virtual void PostEditChangeProperty(struct FPropertyChangedEvent& PropertyChangedEvent) override;
 #endif
 };
diff --git a/Engine/Source/Editor/VREditor/VREditorModule.cpp b/Engine/Source/Editor/VREditor/VREditorModule.cpp
--- a/Engine/Source/Editor/VREditor/VREditorModule.cpp
+++ b/Engine/Source/Editor/VREditor/VREditorModule.cpp
@@ -5,7 +5,10 @@
 #include "Stats/Stats.h"
 #include "HAL/IConsoleManager.h"
 #include "IHeadMountedDisplay.h"
+#include "ISettingsCategory.h"
+#include "ISettingsContainer.h"
 #include "ISettingsModule.h"
+#include "ISettingsSection.h"
 #include "IVREditorModule.h"
 #include "IXRTrackingSystem.h"
 #include "LevelEditorActions.h"
@@ -76,6 +79,8 @@ private:
 	FVREditorModeManager ModeManager;
 
 	FDelegateHandle DeferredInitDelegate;
+
+	TWeakPtr<ISettingsSection> WeakSettingsSection;
 };
 
 namespace VREd
@@ -86,6 +91,7 @@ namespace VREd
 void FVREditorModule::StartupModule()
 {
 	LLM_SCOPE_BYNAME(TEXT("VREditor"));
+
 	RadialMenuExtender = MakeShareable(new FExtender());
 
 	// UToolMenus::RegisterStartupCallback is still too early.
@@ -94,6 +100,16 @@ void FVREditorModule::StartupModule()
 		{
 			ExtendToolbarMenu();
 
+			// Cache pointer to VR Mode settings section, registered elsewhere.
+			ISettingsModule& SettingsModule = FModuleManager::GetModuleChecked<ISettingsModule>("Settings");
+			if (TSharedPtr<ISettingsContainer> EditorContainer = SettingsModule.GetContainer("Editor"))
+			{
+				if (TSharedPtr<ISettingsCategory> GeneralCategory = EditorContainer->GetCategory("General"))
+				{
+					WeakSettingsSection = GeneralCategory->GetSection("VR Mode");
+				}
+			}
+
 			// Must happen last; this implicitly deallocates this lambda's captures.
 			FCoreDelegates::OnBeginFrame.Remove(DeferredInitDelegate);
 		}
@@ -269,9 +285,23 @@ void FVREditorModule::ExtendToolbarMenu()
 			{
 				FToolUIAction SelectModeAction;
 				SelectModeAction.ExecuteAction = FToolMenuExecuteAction::CreateLambda(
-					[Mode](const FToolMenuContext&)
+					[Mode, WeakSettingsSection = WeakSettingsSection]
+					(const FToolMenuContext&)
 					{
-						GetMutableDefault<UVRModeSettings>()->ModeClass = Mode;
+						UVRModeSettings* Settings = GetMutableDefault<UVRModeSettings>();
+						Settings->ModeClass = Mode;
+
+						if (FProperty* ModeClassProperty = FindFProperty<FProperty>(Settings->GetClass(),
+							GET_MEMBER_NAME_CHECKED(UVRModeSettings, ModeClass)))
+						{
+							FPropertyChangedEvent PropertyUpdateStruct(ModeClassProperty, EPropertyChangeType::ValueSet);
+							Settings->PostEditChangeProperty(PropertyUpdateStruct);
+						}
+
+						if (TSharedPtr<ISettingsSection> SettingsSection = WeakSettingsSection.Pin())
+						{
+							SettingsSection->Save();
+						}
 					});
 				SelectModeAction.CanExecuteAction = FToolMenuCanExecuteAction::CreateLambda(
 					[](const FToolMenuContext&)
@@ -293,11 +323,11 @@ void FVREditorModule::ExtendToolbarMenu()
 			}
 		}));
 
-	FToolMenuSection& SettingsSection = OptionsMenu->AddSection("Settings");
+	FToolMenuSection& SettingsMenuSection = OptionsMenu->AddSection("Settings");
 
-	SettingsSection.AddSeparator("SettingsStartSeparator");
+	SettingsMenuSection.AddSeparator("SettingsStartSeparator");
 
-	SettingsSection.AddMenuEntry(
+	SettingsMenuSection.AddMenuEntry(
 		"EditorSettings",
 		LOCTEXT("ToggleVrOptions_EditorSettings_Label", "Settings..."),
 		LOCTEXT("ToggleVrOptions_EditorSettings_Tooltip", "Jump to the VR Editor Mode section in the Editor Preferences"),
diff --git a/Engine/Source/Editor/VREditor/VRModeSettings.cpp b/Engine/Source/Editor/VREditor/VRModeSettings.cpp
--- a/Engine/Source/Editor/VREditor/VRModeSettings.cpp
+++ b/Engine/Source/Editor/VREditor/VRModeSettings.cpp
@@ -15,6 +15,8 @@ UVRModeSettings::UVRModeSettings()
 #if WITH_EDITOR
 void UVRModeSettings::PostEditChangeProperty(struct FPropertyChangedEvent& PropertyChangedEvent)
 {
+	Super::PostEditChangeProperty(PropertyChangedEvent);
+
 	if (PropertyChangedEvent.Property
 		&& PropertyChangedEvent.Property->GetFName() == GET_MEMBER_NAME_CHECKED(UVRModeSettings, bEnableAutoVREditMode)
 		&& bEnableAutoVREditMode == true)
