VR Editor: Remove access to base VR editor mode; Virtual Scouting is the recommended alternative for all users going forward.

 - The UVREditorMode UCLASS is now marked abstract, and therefore excluded from the toolbar dropdown.
 - The toolbar button is now shown if any eligible (non-abstract) mode class is available (for example, if the user enabled the Virtual Scouting plugin).
 - The mode preference defaults to null, and if it is still set to null, will be automatically assigned the first discovered eligible mode (this has the effect of defaulting to Virtual Scouting as soon as the plugin is enabled).

diff --git a/Engine/Config/BaseEditorSettings.ini b/Engine/Config/BaseEditorSettings.ini
--- a/Engine/Config/BaseEditorSettings.ini
+++ b/Engine/Config/BaseEditorSettings.ini
@@ -51,15 +51,14 @@ InteractorHand=Right
 bShowWorldMovementGrid=True
 bShowWorldMovementPostProcess=True
 bShowWorldScaleProgressBar=True
-bScaleWorldFromFloor = False
-bScaleWorldWithDynamicPivot = True
-bAllowSimultaneousWorldScalingAndRotation = True
+bScaleWorldFromFloor=False
+bScaleWorldWithDynamicPivot=True
+bAllowSimultaneousWorldScalingAndRotation=True
 UIBrightness=1.500000
 GizmoScale=0.800000
 DoubleClickTime=0.250000
 TriggerPressedThreshold_Vive=0.330000
 TriggerPressedThreshold_Rift=0.500000
-ModeClass=Class'/Script/VREditor.VREditorMode'
 
 [/Script/TakeRecorder.TakeRecorderProjectSettings]
 +Settings=(DefaultTracks=((MatchingActorClass=/Script/CinematicCamera.CineCameraActor,DefaultPropertyTracks=((ComponentPath="CameraComponent",PropertyPath="CurrentFocalLength"),(ComponentPath="CameraComponent",PropertyPath="FocusSettings.ManualFocusDistance"),(ComponentPath="CameraComponent",PropertyPath="CurrentAperture"),(ComponentPath="CameraComponent",PropertyPath="FilmbackSettings.SensorWidth"),(ComponentPath="CameraComponent",PropertyPath="FilmbackSettings.SensorHeight")), ExcludePropertyTracks=((ComponentPath="CameraComponent",PropertyPath="FilmbackSettings.SensorAspectRatio"),(ComponentPath="CameraComponent",PropertyPath="AspectRatio"),(ComponentPath="CameraComponent",PropertyPath="FieldOfView"))),(MatchingActorClass=/Script/Engine.Light,DefaultPropertyTracks=((ComponentPath="LightComponent0",PropertyPath="Intensity"),(ComponentPath="LightComponent0",PropertyPath="LightColor")))))
diff --git a/Engine/Source/Editor/VREditor/Public/VREditorMode.h b/Engine/Source/Editor/VREditor/Public/VREditorMode.h
--- a/Engine/Source/Editor/VREditor/Public/VREditorMode.h
+++ b/Engine/Source/Editor/VREditor/Public/VREditorMode.h
@@ -49,7 +49,7 @@ namespace VRActionTypes
 /**
  * VR Editor Mode. Extends editor viewports with functionality for VR controls and object manipulation
  */
-UCLASS( BlueprintType, Blueprintable, Transient )
+UCLASS( Abstract, BlueprintType, Blueprintable, Transient )
 class VREDITOR_API UVREditorMode : public UEditorWorldExtension
 {
 	GENERATED_BODY()
diff --git a/Engine/Source/Editor/VREditor/VREditorModeManager.cpp b/Engine/Source/Editor/VREditor/VREditorModeManager.cpp
--- a/Engine/Source/Editor/VREditor/VREditorModeManager.cpp
+++ b/Engine/Source/Editor/VREditor/VREditorModeManager.cpp
@@ -176,10 +176,8 @@ bool FVREditorModeManager::IsVREditorAvailable() const
 
 bool FVREditorModeManager::IsVREditorButtonActive() const
 {
-	const bool bHasHMDDevice = GEngine->XRSystem.IsValid() && GEngine->XRSystem->GetHMDDevice() && GEngine->XRSystem->GetHMDDevice()->IsHMDEnabled();
-	const bool bDerivedModeAvailable = CachedModeClassPaths.Num() > 1;
-
-	return bHasHMDDevice || bDerivedModeAvailable;
+	const bool bAnyModeAvailable = CachedModeClassPaths.Num() > 0;
+	return bAnyModeAvailable;
 }
 
 
@@ -314,19 +312,27 @@ void FVREditorModeManager::GetConcreteModeClasses(TArray<UClass*>& OutModeClasse
 {
 	for (const FTopLevelAssetPath& ClassPath : CachedModeClassPaths)
 	{
-		TSoftClassPtr<UVREditorMode> SoftClass = TSoftClassPtr<UVREditorMode>(FSoftObjectPath(ClassPath));
-		if (UClass* Class = SoftClass.LoadSynchronous())
+		if (UClass* Class = TryGetConcreteModeClass(ClassPath))
 		{
-			if (Class->HasAnyClassFlags(CLASS_Abstract) || (Class->GetAuthoritativeClass() != Class))
-			{
-				continue;
-			}
-
 			OutModeClasses.Add(Class);
 		}
 	}
 }
 
+UClass* FVREditorModeManager::TryGetConcreteModeClass(const FTopLevelAssetPath& ClassPath) const
+{
+	TSoftClassPtr<UVREditorMode> SoftClass = TSoftClassPtr<UVREditorMode>(FSoftObjectPath(ClassPath));
+	if (UClass* Class = SoftClass.LoadSynchronous())
+	{
+		if (!Class->HasAnyClassFlags(CLASS_Abstract) || (Class->GetAuthoritativeClass() != Class))
+		{
+			return Class;
+		}
+	}
+
+	return nullptr;
+}
+
 void FVREditorModeManager::GetModeClassPaths(TSet<FTopLevelAssetPath>& OutModeClasses) const
 {
 	TArray<FTopLevelAssetPath> BaseModeClasses;
@@ -338,7 +344,33 @@ void FVREditorModeManager::GetModeClassPaths(TSet<FTopLevelAssetPath>& OutModeCl
 
 void FVREditorModeManager::UpdateCachedModeClassPaths()
 {
-	GetModeClassPaths(CachedModeClassPaths);
+	TSet<FTopLevelAssetPath> UpdatedModeClassPaths;
+	GetModeClassPaths(UpdatedModeClassPaths);
+
+	// Exclude any ineligible (abstract, failed to load, etc) modes
+	for (TSet<FTopLevelAssetPath>::TIterator It = UpdatedModeClassPaths.CreateIterator(); It; ++It)
+	{
+		if (!TryGetConcreteModeClass(*It))
+		{
+			It.RemoveCurrent();
+		}
+	}
+
+	// Issue update events.
+	TSet<FTopLevelAssetPath> Union = CachedModeClassPaths.Union(UpdatedModeClassPaths);
+	for (const FTopLevelAssetPath& Path : Union)
+	{
+		if (!CachedModeClassPaths.Contains(Path))
+		{
+			OnModeClassAdded.Broadcast(Path);
+		}
+		else if (!UpdatedModeClassPaths.Contains(Path))
+		{
+			OnModeClassRemoved.Broadcast(Path);
+		}
+	}
+
+	CachedModeClassPaths = MoveTemp(UpdatedModeClassPaths);
 }
 
 void FVREditorModeManager::HandleAssetFilesLoaded()
diff --git a/Engine/Source/Editor/VREditor/VREditorModeManager.h b/Engine/Source/Editor/VREditor/VREditorModeManager.h
--- a/Engine/Source/Editor/VREditor/VREditorModeManager.h
+++ b/Engine/Source/Editor/VREditor/VREditorModeManager.h
@@ -9,17 +9,17 @@
 #include "HeadMountedDisplayTypes.h"
 #include "IVREditorModule.h"
 
+
 class UVREditorMode;
 enum class EMapChangeType : uint8;
 
+
 /**
  * Manages starting and closing the VREditor mode
  */
 class FVREditorModeManager : public FGCObject, public FTickableEditorObject
 {
-
 public:
-
 	/** Default constructor */
 	FVREditorModeManager();
 
@@ -68,7 +68,15 @@ public:
 	 * those which are suitable for use (not abstract, etc.) to the provided array. */
 	void GetConcreteModeClasses(TArray<UClass*>& OutModeClasses) const;
 
+	DECLARE_MULTICAST_DELEGATE_OneParam(FModeClassAdded, const FTopLevelAssetPath&);
+	FModeClassAdded OnModeClassAdded;
+
+	DECLARE_MULTICAST_DELEGATE_OneParam(FModeClassRemoved, const FTopLevelAssetPath&);
+	FModeClassRemoved OnModeClassRemoved;
+
 private:
+	UClass* TryGetConcreteModeClass(const FTopLevelAssetPath& ClassPath) const;
+
 	/** Queries the asset registry for asset paths of all classes derived from UVREditorMode. */
 	void GetModeClassPaths(TSet<FTopLevelAssetPath>& OutModeClassPaths) const;
 
diff --git a/Engine/Source/Editor/VREditor/VREditorModule.cpp b/Engine/Source/Editor/VREditor/VREditorModule.cpp
--- a/Engine/Source/Editor/VREditor/VREditorModule.cpp
+++ b/Engine/Source/Editor/VREditor/VREditorModule.cpp
@@ -71,6 +71,8 @@ public:
 
 private:
 	void ExtendToolbarMenu();
+	void HandleModeClassAdded(const FTopLevelAssetPath& ClassPath);
+	void AssignModeClassAndSave(const FTopLevelAssetPath& ClassPath);
 
 private:
 	TSharedPtr<class FExtender> RadialMenuExtender;
@@ -94,6 +96,8 @@ void FVREditorModule::StartupModule()
 
 	RadialMenuExtender = MakeShareable(new FExtender());
 
+	ModeManager.OnModeClassAdded.AddRaw(this, &FVREditorModule::HandleModeClassAdded);
+
 	// UToolMenus::RegisterStartupCallback is still too early.
 	DeferredInitDelegate = FCoreDelegates::OnBeginFrame.AddLambda(
 		[this]()
@@ -134,7 +138,7 @@ bool FVREditorModule::IsVREditorEnabled() const
 {
 	return ModeManager.IsVREditorActive();
 }
-													
+
 bool FVREditorModule::IsVREditorAvailable() const
 {
 	return ModeManager.IsVREditorAvailable();
@@ -285,23 +289,10 @@ void FVREditorModule::ExtendToolbarMenu()
 			{
 				FToolUIAction SelectModeAction;
 				SelectModeAction.ExecuteAction = FToolMenuExecuteAction::CreateLambda(
-					[Mode, WeakSettingsSection = WeakSettingsSection]
+					[this, Mode]
 					(const FToolMenuContext&)
 					{
-						UVRModeSettings* Settings = GetMutableDefault<UVRModeSettings>();
-						Settings->ModeClass = Mode;
-
-						if (FProperty* ModeClassProperty = FindFProperty<FProperty>(Settings->GetClass(),
-							GET_MEMBER_NAME_CHECKED(UVRModeSettings, ModeClass)))
-						{
-							FPropertyChangedEvent PropertyUpdateStruct(ModeClassProperty, EPropertyChangeType::ValueSet);
-							Settings->PostEditChangeProperty(PropertyUpdateStruct);
-						}
-
-						if (TSharedPtr<ISettingsSection> SettingsSection = WeakSettingsSection.Pin())
-						{
-							SettingsSection->Save();
-						}
+						AssignModeClassAndSave(Mode->GetClassPathName());
 					});
 				SelectModeAction.CanExecuteAction = FToolMenuCanExecuteAction::CreateLambda(
 					[](const FToolMenuContext&)
@@ -336,6 +327,33 @@ void FVREditorModule::ExtendToolbarMenu()
 	);
 }
 
+void FVREditorModule::HandleModeClassAdded(const FTopLevelAssetPath& ClassPath)
+{
+	// If we don't have a mode class yet, assign the first one we discover.
+	if (GetDefault<UVRModeSettings>()->ModeClass.IsNull())
+	{
+		AssignModeClassAndSave(ClassPath);
+	}
+}
+
+void FVREditorModule::AssignModeClassAndSave(const FTopLevelAssetPath& ClassPath)
+{
+	UVRModeSettings* Settings = GetMutableDefault<UVRModeSettings>();
+	Settings->ModeClass = FSoftObjectPath(ClassPath);
+
+	if (FProperty* ModeClassProperty = FindFProperty<FProperty>(Settings->GetClass(),
+		GET_MEMBER_NAME_CHECKED(UVRModeSettings, ModeClass)))
+	{
+		FPropertyChangedEvent PropertyUpdateStruct(ModeClassProperty, EPropertyChangeType::ValueSet);
+		Settings->PostEditChangeProperty(PropertyUpdateStruct);
+	}
+
+	if (TSharedPtr<ISettingsSection> SettingsSection = WeakSettingsSection.Pin())
+	{
+		SettingsSection->Save();
+	}
+}
+
 
 IMPLEMENT_MODULE( FVREditorModule, VREditor )
 
