Virtual Scouting: Delayed initialization fixes.

 - Fix crashes related to interaction profile future remaining unresolved, or being fulfilled at unexpected times (e.g. during editor shutdown).
 - Ensure FVREditorModeManager is cleaned up during ShutdownModule, which prevents a crash in TMulticastDelegateBase::RemoveAll.
 - Fix shift+V hotkey exit during the period prior to determining the interaction profile, where Virtual Scouting mode is only partially initialized.

diff --git a/Engine/Plugins/Experimental/VirtualScouting/Source/VirtualScoutingEditor/Private/VPScoutingMode.cpp b/Engine/Plugins/Experimental/VirtualScouting/Source/VirtualScoutingEditor/Private/VPScoutingMode.cpp
--- a/Engine/Plugins/Experimental/VirtualScouting/Source/VirtualScoutingEditor/Private/VPScoutingMode.cpp
+++ b/Engine/Plugins/Experimental/VirtualScouting/Source/VirtualScoutingEditor/Private/VPScoutingMode.cpp
@@ -92,22 +92,29 @@ void UVPScoutingMode::Enter()
 	// (stereo rendering) has started, and we need to defer creation of the interactors, etc.
 	BeginEntry();
 
-	XrExt->GetHmdDeviceTypeFuture() = XrExt->GetHmdDeviceTypeFuture().Next([this](FName DeviceType)
-	{
-		if (DeviceType != NAME_None)
-		{
-			SetHMDDeviceTypeOverride(DeviceType);
-		}
-		else
+	XrExt->GetHmdDeviceTypeFuture() = XrExt->GetHmdDeviceTypeFuture().Next(
+		[WeakThis = TWeakObjectPtr<UVPScoutingMode>(this)]
+		(FName DeviceType)
 		{
-			UE_LOG(LogVirtualScouting, Error, TEXT("Unable to map legacy HMD device type"));
-		}
-
-		SetupSubsystems();
-		FinishEntry();
-
-		return DeviceType;
-	});
+			UVPScoutingMode* This = WeakThis.Get();
+			if (!This || !This->bIsFullyInitialized)
+			{
+				UE_LOG(LogVirtualScouting, Warning, TEXT("Stale UVPScoutingMode; ignoring HmdDeviceType future"));
+				return DeviceType;
+			}
+
+			if (DeviceType == NAME_None)
+			{
+				UE_LOG(LogVirtualScouting, Error, TEXT("Unable to map legacy HMD device type"));
+				return DeviceType;
+			}
+
+			This->SetHMDDeviceTypeOverride(DeviceType);
+			This->SetupSubsystems();
+			This->FinishEntry();
+
+			return DeviceType;
+		});
 }
 
 
diff --git a/Engine/Plugins/Experimental/VirtualScouting/Source/VirtualScoutingOpenXR/Private/VirtualScoutingOpenXR.cpp b/Engine/Plugins/Experimental/VirtualScouting/Source/VirtualScoutingOpenXR/Private/VirtualScoutingOpenXR.cpp
--- a/Engine/Plugins/Experimental/VirtualScouting/Source/VirtualScoutingOpenXR/Private/VirtualScoutingOpenXR.cpp
+++ b/Engine/Plugins/Experimental/VirtualScouting/Source/VirtualScoutingOpenXR/Private/VirtualScoutingOpenXR.cpp
@@ -170,7 +170,7 @@ void FVirtualScoutingOpenXRExtension::PostSyncActions(XrSession InSession)
 
 void FVirtualScoutingOpenXRExtension::TryFulfillDeviceTypePromise()
 {
-	if (!DeviceTypeFuture.IsReady())
+	if (DeviceTypeFuture.IsValid() && !DeviceTypeFuture.IsReady())
 	{
 		TOptional<FName> MaybeType = TryGetHmdDeviceType();
 		if (MaybeType)
diff --git a/Engine/Source/Editor/LevelEditor/Private/LevelEditorActions.cpp b/Engine/Source/Editor/LevelEditor/Private/LevelEditorActions.cpp
--- a/Engine/Source/Editor/LevelEditor/Private/LevelEditorActions.cpp
+++ b/Engine/Source/Editor/LevelEditor/Private/LevelEditorActions.cpp
@@ -2379,7 +2379,7 @@ void FLevelEditorActionCallbacks::ImportContent()
 void FLevelEditorActionCallbacks::ToggleVR()
 {
 	IVREditorModule& VREditorModule = IVREditorModule::Get();
-	VREditorModule.EnableVREditor( !VREditorModule.IsVREditorEnabled() );
+	VREditorModule.EnableVREditor( VREditorModule.GetVRMode() == nullptr );
 }
 
 
diff --git a/Engine/Source/Editor/VREditor/VREditorModeManager.cpp b/Engine/Source/Editor/VREditor/VREditorModeManager.cpp
--- a/Engine/Source/Editor/VREditor/VREditorModeManager.cpp
+++ b/Engine/Source/Editor/VREditor/VREditorModeManager.cpp
@@ -57,7 +57,7 @@ FVREditorModeManager::~FVREditorModeManager()
 		}
 	}
 
-	if (CurrentVREditorMode)
+	if (IsValid(CurrentVREditorMode))
 	{
 		CurrentVREditorMode->OnVRModeEntryComplete().RemoveAll(this);
 		CurrentVREditorMode = nullptr;
@@ -111,7 +111,8 @@ bool FVREditorModeManager::IsTickable() const
 void FVREditorModeManager::EnableVREditor( const bool bEnable, const bool bForceWithoutHMD )
 {
 	// Don't do anything when the current VR Editor is already in the requested state
-	if( bEnable != IsVREditorActive() )
+	const bool bIsEnabled = IsValid(CurrentVREditorMode);
+	if (bEnable != bIsEnabled)
 	{
 		if( bEnable && ( IsVREditorAvailable() || bForceWithoutHMD ))
 		{
diff --git a/Engine/Source/Editor/VREditor/VREditorModule.cpp b/Engine/Source/Editor/VREditor/VREditorModule.cpp
--- a/Engine/Source/Editor/VREditor/VREditorModule.cpp
+++ b/Engine/Source/Editor/VREditor/VREditorModule.cpp
@@ -65,10 +65,10 @@ public:
 	static void ToggleForceVRMode();
 
 	/** Return a multicast delegate which is executed when VR mode starts. */
-	virtual FOnVREditingModeEnter& OnVREditingModeEnter() override { return ModeManager.OnVREditingModeEnter(); }
+	virtual FOnVREditingModeEnter& OnVREditingModeEnter() override { return ModeManager->OnVREditingModeEnter(); }
 
 	/** Return a multicast delegate which is executed when VR mode stops. */
-	virtual FOnVREditingModeExit& OnVREditingModeExit() override { return ModeManager.OnVREditingModeExit(); }
+	virtual FOnVREditingModeExit& OnVREditingModeExit() override { return ModeManager->OnVREditingModeExit(); }
 
 private:
 	void ExtendToolbarMenu();
@@ -79,7 +79,7 @@ private:
 	TSharedPtr<class FExtender> RadialMenuExtender;
 
 	/** Handles turning VR Editor mode on and off */
-	FVREditorModeManager ModeManager;
+	TUniquePtr<FVREditorModeManager> ModeManager;
 
 	FDelegateHandle DeferredInitDelegate;
 
@@ -97,7 +97,8 @@ void FVREditorModule::StartupModule()
 
 	RadialMenuExtender = MakeShareable(new FExtender());
 
-	ModeManager.OnModeClassAdded.AddRaw(this, &FVREditorModule::HandleModeClassAdded);
+	ModeManager = MakeUnique<FVREditorModeManager>();
+	ModeManager->OnModeClassAdded.AddRaw(this, &FVREditorModule::HandleModeClassAdded);
 
 	// UToolMenus::RegisterStartupCallback is still too early.
 	DeferredInitDelegate = FCoreDelegates::OnBeginFrame.AddLambda(
@@ -129,6 +130,8 @@ void FVREditorModule::ShutdownModule()
 	{
 		FVREditorStyle::Shutdown();
 	}
+
+	ModeManager.Release();
 }
 
 void FVREditorModule::PostLoadCallback()
@@ -137,33 +140,33 @@ void FVREditorModule::PostLoadCallback()
 
 bool FVREditorModule::IsVREditorEnabled() const
 {
-	return ModeManager.IsVREditorActive();
+	return ModeManager->IsVREditorActive();
 }
 
 bool FVREditorModule::IsVREditorAvailable() const
 {
-	return ModeManager.IsVREditorAvailable();
+	return ModeManager->IsVREditorAvailable();
 }
 
 bool FVREditorModule::IsVREditorButtonActive() const
 {
-	return ModeManager.IsVREditorButtonActive();
+	return ModeManager->IsVREditorButtonActive();
 }
 
 void FVREditorModule::EnableVREditor( const bool bEnable, const bool bForceWithoutHMD )
 {
-	ModeManager.EnableVREditor( bEnable, bForceWithoutHMD );
+	ModeManager->EnableVREditor( bEnable, bForceWithoutHMD );
 }
 
 bool FVREditorModule::IsVREditorModeActive()
 {
 	// Deprecated method only returns true for legacy UVREditorMode
-	return GetVRMode() && ModeManager.IsVREditorActive();
+	return GetVRMode() && ModeManager->IsVREditorActive();
 }
 
 UVREditorModeBase* FVREditorModule::GetVRModeBase()
 {
-	return ModeManager.GetCurrentVREditorMode();
+	return ModeManager->GetCurrentVREditorMode();
 }
 
 UVREditorMode* FVREditorModule::GetVRMode()
@@ -271,7 +274,7 @@ void FVREditorModule::ExtendToolbarMenu()
 			UClass* CurrentModeClass = CurrentModeClassSoft.LoadSynchronous();
 
 			TArray<UClass*> ModeClasses;
-			ModeManager.GetConcreteModeClasses(ModeClasses);
+			ModeManager->GetConcreteModeClasses(ModeClasses);
 
 			Algo::Sort(ModeClasses,
 				[](const UClass* Lhs, const UClass* Rhs)
