VR Editor: Fix workaround for SKEL_ classes in asset registry that was inverted during refactoring. Also fixes a possible null pointer dereference if the implementation of UVREditorModeBase::Enter also removes itself as the active VR mode prior to returning.

diff --git a/Engine/Source/Editor/VREditor/VREditorModeManager.cpp b/Engine/Source/Editor/VREditor/VREditorModeManager.cpp
--- a/Engine/Source/Editor/VREditor/VREditorModeManager.cpp
+++ b/Engine/Source/Editor/VREditor/VREditorModeManager.cpp
@@ -214,7 +214,7 @@ void FVREditorModeManager::StartVREditorMode( const bool bForceWithoutHMD )
 			FEngineAnalytics::GetProvider().RecordEvent(TEXT("Editor.Usage.InitVREditorMode"), FAnalyticsEventAttribute(TEXT("Enterprise"), IProjectManager::Get().IsEnterpriseProject()));
 
 			// Take note of VREditor entering (only if actually in VR)
-			if (CurrentVREditorMode->IsActuallyUsingVR())
+			if (CurrentVREditorMode && CurrentVREditorMode->IsActuallyUsingVR())
 			{
 				TArray<FAnalyticsEventAttribute> Attributes;
 				FString HMDName = GEditor->XRSystem->GetSystemName().ToString();
@@ -295,7 +295,7 @@ UClass* FVREditorModeManager::TryGetConcreteModeClass(const FTopLevelAssetPath&
 	TSoftClassPtr<UVREditorModeBase> SoftClass = TSoftClassPtr<UVREditorModeBase>(FSoftObjectPath(ClassPath));
 	if (UClass* Class = SoftClass.LoadSynchronous())
 	{
-		if (!Class->HasAnyClassFlags(CLASS_Abstract) || (Class->GetAuthoritativeClass() != Class))
+		if (!Class->HasAnyClassFlags(CLASS_Abstract) && (Class->GetAuthoritativeClass() == Class))
 		{
 			return Class;
 		}
