VR Editor: Fix for incorrectly enabled button (and crash when attempting to launch) with disconnected HMD.

diff --git a/Engine/Source/Editor/VREditor/VREditorModeManager.cpp b/Engine/Source/Editor/VREditor/VREditorModeManager.cpp
--- a/Engine/Source/Editor/VREditor/VREditorModeManager.cpp
+++ b/Engine/Source/Editor/VREditor/VREditorModeManager.cpp
@@ -149,7 +149,7 @@ bool FVREditorModeManager::IsVREditorAvailable() const
 		return false;
 	}
 
-	return GEngine->XRSystem->GetHMDDevice() && GEngine->XRSystem->GetHMDDevice()->IsHMDEnabled();
+	return GEngine->XRSystem->GetHMDDevice() && GEngine->XRSystem->GetHMDDevice()->IsHMDConnected();
 }
 
 bool FVREditorModeManager::IsVREditorButtonActive() const
diff --git a/Engine/Source/Editor/VREditor/VREditorModule.cpp b/Engine/Source/Editor/VREditor/VREditorModule.cpp
--- a/Engine/Source/Editor/VREditor/VREditorModule.cpp
+++ b/Engine/Source/Editor/VREditor/VREditorModule.cpp
@@ -242,7 +242,7 @@ void FVREditorModule::ExtendToolbarMenu()
 		{
 			const bool bHasHMDDevice = GEngine->XRSystem
 				&& GEngine->XRSystem->GetHMDDevice()
-				&& GEngine->XRSystem->GetHMDDevice()->IsHMDEnabled();
+				&& GEngine->XRSystem->GetHMDDevice()->IsHMDConnected();
 
 			const FText HmdAvailableText = LOCTEXT("ToggleVrOptions_Status_HmdAvailable", "Headset available ({XrVersionString})");
 			const FText HmdNotAvailableText = LOCTEXT("ToggleVrOptions_Status_HmdNotAvailable", "No headset available");
