Moving loose files in VREditor into Private and Public.

diff --git a/Engine/Source/Editor/VREditor/VREditor.Build.cs b/Engine/Source/Editor/VREditor/VREditor.Build.cs
--- a/Engine/Source/Editor/VREditor/VREditor.Build.cs
+++ b/Engine/Source/Editor/VREditor/VREditor.Build.cs
@@ -12,8 +12,6 @@ namespace UnrealBuildTool.Rules
 				}
 			);
 
-			PublicIncludePaths.Add(ModuleDirectory);
-
             PrivateDependencyModuleNames.AddRange(
                 new string[] {
                     "AppFramework",
