Prevent VR Mode when no viewport is active.

diff --git a/Engine/Source/Editor/VREditor/Private/VREditorModeManager.cpp b/Engine/Source/Editor/VREditor/Private/VREditorModeManager.cpp
--- a/Engine/Source/Editor/VREditor/Private/VREditorModeManager.cpp
+++ b/Engine/Source/Editor/VREditor/Private/VREditorModeManager.cpp
@@ -25,6 +25,7 @@
 #include "UnrealEdMisc.h"
 #include "Settings/EditorStyleSettings.h"
 #include "VREditorInteractor.h"
+#include "SLevelViewport.h"
 
 #define LOCTEXT_NAMESPACE "VREditor"
 
@@ -149,6 +150,13 @@ bool FVREditorModeManager::IsVREditorAvailable() const
 		return false;
 	}
 
+	const TSharedRef<ILevelEditor>& LevelEditor = FModuleManager::GetModuleChecked<FLevelEditorModule>("LevelEditor").GetFirstLevelEditor().ToSharedRef();
+	TSharedPtr<IAssetViewport> ActiveLevelViewport = LevelEditor->GetActiveViewportInterface();
+	if (!ActiveLevelViewport.IsValid())
+	{
+		return false;
+	}
+
 	return GEngine->XRSystem->GetHMDDevice() && GEngine->XRSystem->GetHMDDevice()->IsHMDConnected();
 }
 
