VR Editor: Force cinematic viewport back to default prior to entering VR.

diff --git a/Engine/Source/Editor/VREditor/Private/VREditorMode.cpp b/Engine/Source/Editor/VREditor/Private/VREditorMode.cpp
--- a/Engine/Source/Editor/VREditor/Private/VREditorMode.cpp
+++ b/Engine/Source/Editor/VREditor/Private/VREditorMode.cpp
@@ -54,6 +54,7 @@
 #include "XRMotionControllerBase.h" // for FXRMotionControllerBase::Left/RightHandSourceId
 #include "UI/VREditorFloatingUI.h"
 #include "AssetEditorViewportLayout.h"
+#include "LevelViewportActions.h"
 
 #define LOCTEXT_NAMESPACE "VREditorMode"
 
@@ -239,8 +240,16 @@ void UVREditorMode::Enter()
 	FinishEntry();
 }
 
+namespace UE::VREditor::Private
+{
+	// Defined in VREditorModeBase.cpp.
+	TSharedPtr<SLevelViewport> TryGetActiveViewport();
+}
+
 void UVREditorMode::BeginEntry()
 {
+	using namespace UE::VREditor::Private;
+
 	FSavedEditorState& SavedEditorState = static_cast<FSavedEditorState&>(SavedEditorStateChecked());
 
 	bWantsToExitMode = false;
@@ -260,16 +269,25 @@ void UVREditorMode::BeginEntry()
 		const TSharedRef< ILevelEditor >& LevelEditor = FModuleManager::GetModuleChecked<FLevelEditorModule>("LevelEditor").GetFirstLevelEditor().ToSharedRef();
 
 		// Do we have an active perspective viewport that is valid for VR?  If so, go ahead and use that.
-		TSharedPtr<SLevelViewport> ExistingActiveLevelViewport;
+		TSharedPtr<SLevelViewport> ExistingActiveLevelViewport = TryGetActiveViewport();
+
+		if (ExistingActiveLevelViewport && ExistingActiveLevelViewport->GetCommandList() && FLevelViewportCommands::Get().SetDefaultViewportType)
 		{
-			TSharedPtr<IAssetViewport> ActiveLevelViewport = LevelEditor->GetActiveViewportInterface();
-			if (ActiveLevelViewport.IsValid())
-			{
-				ExistingActiveLevelViewport = StaticCastSharedRef< SLevelViewport >(ActiveLevelViewport->AsWidget());
-				ExistingActiveLevelViewport->RemoveAllPreviews(true);
-			}
+			ExistingActiveLevelViewport->GetCommandList()->TryExecuteAction(
+				FLevelViewportCommands::Get().SetDefaultViewportType.ToSharedRef());
+
+			// If the active viewport was e.g. a cinematic viewport, changing it
+			// back to default recreated it and our pointer might be stale.
+			ExistingActiveLevelViewport = TryGetActiveViewport();
+		}
+
+		if (!ensure(ExistingActiveLevelViewport))
+		{
+			return;
 		}
 
+		ExistingActiveLevelViewport->RemoveAllPreviews(true);
+
 		StartViewport(ExistingActiveLevelViewport);
 
 		if (WorldInteraction != nullptr)
diff --git a/Engine/Source/Editor/VREditor/Private/VREditorModeBase.cpp b/Engine/Source/Editor/VREditor/Private/VREditorModeBase.cpp
--- a/Engine/Source/Editor/VREditor/Private/VREditorModeBase.cpp
+++ b/Engine/Source/Editor/VREditor/Private/VREditorModeBase.cpp
@@ -8,6 +8,7 @@
 #include "ILevelEditor.h"
 #include "IXRTrackingSystem.h"
 #include "LevelEditor.h"
+#include "LevelViewportActions.h"
 #include "Modules/ModuleManager.h"
 #include "SLevelViewport.h"
 
@@ -41,22 +42,46 @@ void UVREditorModeBase::Shutdown()
 }
 
 
+namespace UE::VREditor::Private
+{
+	TSharedPtr<SLevelViewport> TryGetActiveViewport()
+	{
+		const TSharedRef<ILevelEditor>& LevelEditor = FModuleManager::GetModuleChecked<FLevelEditorModule>("LevelEditor").GetFirstLevelEditor().ToSharedRef();
+		TSharedPtr<IAssetViewport> ActiveViewport = LevelEditor->GetActiveViewportInterface();
+		if (!ActiveViewport)
+		{
+			return nullptr;
+		}
+
+		return StaticCastSharedRef<SLevelViewport>(ActiveViewport->AsWidget());
+	}
+}
+
+
 void UVREditorModeBase::Enter()
 {
-	const TSharedRef<ILevelEditor>& LevelEditor = FModuleManager::GetModuleChecked<FLevelEditorModule>("LevelEditor").GetFirstLevelEditor().ToSharedRef();
-	if (TSharedPtr<IAssetViewport> ActiveLevelViewport = LevelEditor->GetActiveViewportInterface())
+	using namespace UE::VREditor::Private;
+
+	TSharedPtr<SLevelViewport> ActiveLevelViewport = TryGetActiveViewport();
+
+	if (ActiveLevelViewport && ActiveLevelViewport->GetCommandList() && FLevelViewportCommands::Get().SetDefaultViewportType)
 	{
-		TSharedRef<SLevelViewport> ExistingActiveLevelViewport =
-			StaticCastSharedRef<SLevelViewport>(ActiveLevelViewport->AsWidget());
+		ActiveLevelViewport->GetCommandList()->TryExecuteAction(
+			FLevelViewportCommands::Get().SetDefaultViewportType.ToSharedRef());
 
-		ExistingActiveLevelViewport->RemoveAllPreviews(true);
-		StartViewport(ExistingActiveLevelViewport);
+		// If the active viewport was e.g. a cinematic viewport, changing it
+		// back to default recreated it and our pointer might be stale.
+		ActiveLevelViewport = TryGetActiveViewport();
 	}
-	else
+
+	if (!ensure(ActiveLevelViewport))
 	{
-		// TODO handle error
+		return;
 	}
 
+	ActiveLevelViewport->RemoveAllPreviews(true);
+	StartViewport(ActiveLevelViewport);
+
 	OnVRModeEntryCompleteEvent.Broadcast();
 }
 
