UE-181824 Move HeadMountedDisplay files to XRBase plugin to reduce UE overhead.
-Moving much of engine/source/runtime/headmounteddisplay to a plugin mainly to reduce minimum UE executable size and memory use.  But this is also nice for organization.  The ideal would be to move everything except interfaces used by core engine. VREditor pulling in UMotionControllerComponent is the main blocker at this time.  That dependency is expected to go away eventually at which point we can do the rest of this transfer.
-Added dependencies to XRBase plugin or module as necessary, refactored to avoid some dependencies, removed a few unnecessary dependencies.
-Stripped vestigial VR references from a few project templates.  Mostly just unused HeadMountedDisplay module refs, but a tiny bit of code from TP_Puzzle.
-Fixed dependency cycle ignores related to MaterialShaderQualitySettings PIEPreviewDeviceProfileSelector and UnrealEd.  For reasons I do not understand the HeadMountedDisplay dependency on AugmentedReality was preventing the detection of those, long existent, cycles.

diff --git a/Engine/Config/BaseEngine.ini b/Engine/Config/BaseEngine.ini
--- a/Engine/Config/BaseEngine.ini
+++ b/Engine/Config/BaseEngine.ini
@@ -522,7 +522,6 @@ DefaultWorldPartitionHLODLayer=/Engine/Maps/Templates/HLODs/HLODLayer_Instanced.
 +ClassRedirects=(OldName="ARBlueprintFunctionLibrary", NewName="/Script/AugmentedReality.ARBlueprintLibrary")
 
 ;SpectatorScreen Refactoring Redirects
-+ClassRedirects=(OldName="HeadMountedDisplayFunctionLibrary",NewName="/Script/HeadMountedDisplay.HeadMountedDisplayFunctionLibrary")
 +EnumRedirects=(OldName="EWidgetClipping",ValueChanges=(("EWidgetClipping::No","EWidgetClipping::Inherit"),("EWidgetClipping::Yes","EWidgetClipping::ClipToBounds"),("EWidgetClipping::YesWithoutIntersecting","EWidgetClipping::ClipToBoundsWithoutIntersecting"),("EWidgetClipping::YesAlways","EWidgetClipping::ClipToBoundsAlways")) )
 +EnumRedirects=(OldName="EOrientPositionSelector",NewName="/Script/HeadMountedDisplay.EOrientPositionSelector")
 +EnumRedirects=(OldName="EHMDTrackingOrigin",NewName="/Script/HeadMountedDisplay.EHMDTrackingOrigin")
@@ -1143,6 +1142,15 @@ DefaultWorldPartitionHLODLayer=/Engine/Maps/Templates/HLODs/HLODLayer_Instanced.
 +StructRedirects=(OldName="/Script/UMG.FieldNotificationId", NewName="/Script/FieldNotification.FieldNotificationId")
 +ClassRedirects=(OldName="/Script/UMG.NotifyFieldValueChanged", NewName="/Script/FieldNotification.NotifyFieldValueChanged")
 
+; Refactoring most of HeadMountedDisplay engine module into XRBase plugin.
++ClassRedirects=(OldName="HeadMountedDisplayFunctionLibrary", NewName="/Script/XRBase.HeadMountedDisplayFunctionLibrary")
++ClassRedirects=(OldName="MotionTrackedDeviceFunctionLibrary", NewName="/Script/XRBase.MotionTrackedDeviceFunctionLibrary")
++ClassRedirects=(OldName="VRNotificationsComponent", NewName="/Script/XRBase.VRNotificationsComponent")
++ClassRedirects=(OldName="XRAssetFunctionLibrary", NewName="/Script/XRBase.XRAssetFunctionLibrary")
++ClassRedirects=(OldName="AsyncTask_LoadXRDeviceVisComponent", NewName="/Script/XRBase.AsyncTask_LoadXRDeviceVisComponent")
++ClassRedirects=(OldName="XRDeviceVisualizationComponent", NewName="/Script/XRBase.XRDeviceVisualizationComponent")
++ClassRedirects=(OldName="XRLoadingScreenFunctionLibrary", NewName="/Script/XRBase.XRLoadingScreenFunctionLibrary")
+
 [CoreUObject.Metadata]
 ; Note: UnrealHeaderTool should be rerun after making changes to MetadataRedirects to catch any keys specified in class headers
 +MetadataRedirects=(OldKey="K2Protected", NewKey="BlueprintProtected")
diff --git a/Engine/Plugins/Developer/Concert/ConcertSync/ConcertSyncClient/Source/ConcertSyncClient/Private/ConcertClientPresenceMode.cpp b/Engine/Plugins/Developer/Concert/ConcertSync/ConcertSyncClient/Source/ConcertSyncClient/Private/ConcertClientPresenceMode.cpp
--- a/Engine/Plugins/Developer/Concert/ConcertSync/ConcertSyncClient/Source/ConcertSyncClient/Private/ConcertClientPresenceMode.cpp
+++ b/Engine/Plugins/Developer/Concert/ConcertSync/ConcertSyncClient/Source/ConcertSyncClient/Private/ConcertClientPresenceMode.cpp
@@ -19,7 +19,7 @@
 #include "ConcertClientVRPresenceActor.h"
 #include "EngineUtils.h"
 #include "Features/IModularFeatures.h"
-#include "XRMotionControllerBase.h"
+#include "IMotionController.h"
 
 #if WITH_EDITOR
 #include "Editor.h"
@@ -243,8 +243,8 @@ void FConcertClientVRPresenceMode::SendEvents(IConcertClientSession& Session)
 			return MotionControllerToWorld;
 		};
 
-		FTransform LeftMotionControllerTransform = GetControllerTransform(FXRMotionControllerBase::LeftHandSourceId);
-		FTransform RightMotionControllerTransform = GetControllerTransform(FXRMotionControllerBase::RightHandSourceId);
+		FTransform LeftMotionControllerTransform = GetControllerTransform(IMotionController::LeftHandSourceId);
+		FTransform RightMotionControllerTransform = GetControllerTransform(IMotionController::RightHandSourceId);
 		
 		FConcertClientVRPresenceUpdateEvent Event;
 		Event.LeftMotionControllerOrientation = LeftMotionControllerTransform.GetRotation();
diff --git a/Engine/Plugins/Experimental/RemoteSession/RemoteSession.uplugin b/Engine/Plugins/Experimental/RemoteSession/RemoteSession.uplugin
--- a/Engine/Plugins/Experimental/RemoteSession/RemoteSession.uplugin
+++ b/Engine/Plugins/Experimental/RemoteSession/RemoteSession.uplugin
@@ -46,5 +46,9 @@
 			"Name": "MediaIOFramework",
 			"Enabled": true
 		},
+		{
+			"Name": "XRBase",
+			"Enabled": true
+		}
 	]
 }
diff --git a/Engine/Plugins/Experimental/RemoteSession/Source/RemoteSession/RemoteSession.Build.cs b/Engine/Plugins/Experimental/RemoteSession/Source/RemoteSession/RemoteSession.Build.cs
--- a/Engine/Plugins/Experimental/RemoteSession/Source/RemoteSession/RemoteSession.Build.cs
+++ b/Engine/Plugins/Experimental/RemoteSession/Source/RemoteSession/RemoteSession.Build.cs
@@ -20,6 +20,7 @@ public class RemoteSession : ModuleRules
 				"MediaIOCore",
 				"BackChannel",
 				"ApplicationCore",
+				"XRBase",
 				// ... add other public dependencies that you statically link with here ...
 			}
 		);
diff --git a/Engine/Plugins/Experimental/SimpleHMD/SimpleHMD.uplugin b/Engine/Plugins/Experimental/SimpleHMD/SimpleHMD.uplugin
--- a/Engine/Plugins/Experimental/SimpleHMD/SimpleHMD.uplugin
+++ b/Engine/Plugins/Experimental/SimpleHMD/SimpleHMD.uplugin
@@ -22,12 +22,19 @@
 		"Mac",
 		"IOS"
 	],
-	"Modules" :
+	"Modules" : 
 	[
 		{
 			"Name" : "SimpleHMD",
 			"Type" : "Runtime",
 			"LoadingPhase" : "Default"
 		}
+	],
+	"Plugins" : 
+	[
+		{
+			"Name" : "XRBase",
+			"Enabled" : true
+		}
 	]
 }
\ No newline at end of file
diff --git a/Engine/Plugins/Experimental/SimpleHMD/Source/SimpleHMD/SimpleHMD.Build.cs b/Engine/Plugins/Experimental/SimpleHMD/Source/SimpleHMD/SimpleHMD.Build.cs
--- a/Engine/Plugins/Experimental/SimpleHMD/Source/SimpleHMD/SimpleHMD.Build.cs
+++ b/Engine/Plugins/Experimental/SimpleHMD/Source/SimpleHMD/SimpleHMD.Build.cs
@@ -21,7 +21,8 @@ namespace UnrealBuildTool.Rules
 					"RHI",
 					"RenderCore",
 					"Renderer",
-					"HeadMountedDisplay"
+					"HeadMountedDisplay",
+					"XRBase"
 				}
 				);
 		}
diff --git a/Engine/Plugins/Experimental/XRCreativeFramework/Source/XRCreative/XRCreative.Build.cs b/Engine/Plugins/Experimental/XRCreativeFramework/Source/XRCreative/XRCreative.Build.cs
--- a/Engine/Plugins/Experimental/XRCreativeFramework/Source/XRCreative/XRCreative.Build.cs
+++ b/Engine/Plugins/Experimental/XRCreativeFramework/Source/XRCreative/XRCreative.Build.cs
@@ -24,6 +24,7 @@ public class XRCreative : ModuleRules
 				"Engine",
 				"EnhancedInput",
 				"HeadMountedDisplay",
+				"XRBase",
 				"InputCore",
 				"InteractiveToolsFramework",
 				"LevelSequence",
diff --git a/Engine/Plugins/Experimental/XRCreativeFramework/XRCreativeFramework.uplugin b/Engine/Plugins/Experimental/XRCreativeFramework/XRCreativeFramework.uplugin
--- a/Engine/Plugins/Experimental/XRCreativeFramework/XRCreativeFramework.uplugin
+++ b/Engine/Plugins/Experimental/XRCreativeFramework/XRCreativeFramework.uplugin
@@ -32,6 +32,10 @@
 		}
 	],
 	"Plugins": [
+		{
+			"Name": "XRBase",
+			"Enabled": true
+		},
 		{
 			"Name": "OpenXR",
 			"Enabled": true
diff --git a/Engine/Plugins/FX/Niagara/Source/Niagara/Niagara.Build.cs b/Engine/Plugins/FX/Niagara/Source/Niagara/Niagara.Build.cs
--- a/Engine/Plugins/FX/Niagara/Source/Niagara/Niagara.Build.cs
+++ b/Engine/Plugins/FX/Niagara/Source/Niagara/Niagara.Build.cs
@@ -30,7 +30,6 @@ public class Niagara : ModuleRules
                 "Core",
                 "DeveloperSettings",
                 "Engine",
-                "HeadMountedDisplay",
                 "Json",
                 "JsonUtilities",
                 "Landscape",
diff --git a/Engine/Plugins/Media/PixelStreaming/PixelStreaming.uplugin b/Engine/Plugins/Media/PixelStreaming/PixelStreaming.uplugin
--- a/Engine/Plugins/Media/PixelStreaming/PixelStreaming.uplugin
+++ b/Engine/Plugins/Media/PixelStreaming/PixelStreaming.uplugin
@@ -91,6 +91,10 @@
 		{
 			"Name": "WebSocketNetworking",
 			"Enabled": true
+		},
+		{
+			"Name": "XRBase",
+			"Enabled": true
 		}
 	]
 }
diff --git a/Engine/Plugins/Media/PixelStreaming/Source/PixelStreamingInput/PixelStreamingInput.build.cs b/Engine/Plugins/Media/PixelStreaming/Source/PixelStreamingInput/PixelStreamingInput.build.cs
--- a/Engine/Plugins/Media/PixelStreaming/Source/PixelStreamingInput/PixelStreamingInput.build.cs
+++ b/Engine/Plugins/Media/PixelStreaming/Source/PixelStreamingInput/PixelStreamingInput.build.cs
@@ -17,7 +17,8 @@ namespace UnrealBuildTool.Rules
 					"SlateCore",
 					"Slate",
 					"DeveloperSettings",
-					"HeadMountedDisplay"
+					"HeadMountedDisplay",
+					"XRBase",
 				}
 			);
 
diff --git a/Engine/Plugins/Runtime/AR/AppleAR/AppleARKit/AppleARKit.uplugin b/Engine/Plugins/Runtime/AR/AppleAR/AppleARKit/AppleARKit.uplugin
--- a/Engine/Plugins/Runtime/AR/AppleAR/AppleARKit/AppleARKit.uplugin
+++ b/Engine/Plugins/Runtime/AR/AppleAR/AppleARKit/AppleARKit.uplugin
@@ -31,14 +31,18 @@
 
 		}
 	],
-    "Plugins": [
-        {
-            "Name": "AppleImageUtils",
-            "Enabled": true
+	"Plugins": [
+		{
+			"Name": "XRBase",
+			"Enabled": true
+		},
+		{
+			"Name": "AppleImageUtils",
+			"Enabled": true
 		},
 		{
-            "Name": "ARUtilities",
-            "Enabled": true
-        }
-    ]
+			"Name": "ARUtilities",
+			"Enabled": true
+		}
+	]
 }
diff --git a/Engine/Plugins/Runtime/AR/AzureSpatialAnchorsForARCore/Source/AzureSpatialAnchorsForARCore/Private/AzureSpatialAnchorsAndroidInterop.cpp b/Engine/Plugins/Runtime/AR/AzureSpatialAnchorsForARCore/Source/AzureSpatialAnchorsForARCore/Private/AzureSpatialAnchorsAndroidInterop.cpp
--- a/Engine/Plugins/Runtime/AR/AzureSpatialAnchorsForARCore/Source/AzureSpatialAnchorsForARCore/Private/AzureSpatialAnchorsAndroidInterop.cpp
+++ b/Engine/Plugins/Runtime/AR/AzureSpatialAnchorsForARCore/Source/AzureSpatialAnchorsForARCore/Private/AzureSpatialAnchorsAndroidInterop.cpp
@@ -1,6 +1,7 @@
 // Copyright Epic Games, Inc. All Rights Reserved.
 
 #include "AzureSpatialAnchorsAndroidInterop.h"
+#include "IXRTrackingSystem.h"
 
 // AzureSpatialAnchorsNDK.hpp includes AzureSpatialAnchorsLibrary_Exports.h.  Both are autogenerated and could change with 
 // ASA version updates, so I don't want to edit them.
@@ -83,7 +84,7 @@ bool FAzureSpatialAnchorsAndroidInterop::CreateSession()
 	}
 
     UE_LOG(LogAzureSpatialAnchorsAndroidInterop, Log, TEXT("Creating an Azure Spatial Anchors session"));
-	auto TempARSystem = TSharedPtr<FARSupportInterface, ESPMode::ThreadSafe>{ StaticCastSharedPtr<FXRTrackingSystemBase>(GEngine->XRSystem)->GetARCompositionComponent() };
+	auto TempARSystem = TSharedPtr<FARSupportInterface, ESPMode::ThreadSafe>{ (GEngine->XRSystem)->GetARCompositionComponent() };
     if (TempARSystem.IsValid())
     {
         UE_LOG(LogAzureSpatialAnchorsAndroidInterop, Log, TEXT("Azure Spatial Anchors located a valid AR System"));    
diff --git a/Engine/Plugins/Runtime/AR/AzureSpatialAnchorsForARKit/Source/AzureSpatialAnchorsForARKit/AzureSpatialAnchorsForARKit.Build.cs b/Engine/Plugins/Runtime/AR/AzureSpatialAnchorsForARKit/Source/AzureSpatialAnchorsForARKit/AzureSpatialAnchorsForARKit.Build.cs
--- a/Engine/Plugins/Runtime/AR/AzureSpatialAnchorsForARKit/Source/AzureSpatialAnchorsForARKit/AzureSpatialAnchorsForARKit.Build.cs
+++ b/Engine/Plugins/Runtime/AR/AzureSpatialAnchorsForARKit/Source/AzureSpatialAnchorsForARKit/AzureSpatialAnchorsForARKit.Build.cs
@@ -22,7 +22,8 @@ namespace UnrealBuildTool.Rules
 					"Core",
 					"CoreUObject",
 					"Engine",
-					"AugmentedReality"
+					"AugmentedReality",
+					"HeadMountedDisplay"
 				}
 			);
 
diff --git a/Engine/Plugins/Runtime/AR/AzureSpatialAnchorsForARKit/Source/AzureSpatialAnchorsForARKit/Private/AzureSpatialAnchorsARKitInterop.cpp b/Engine/Plugins/Runtime/AR/AzureSpatialAnchorsForARKit/Source/AzureSpatialAnchorsForARKit/Private/AzureSpatialAnchorsARKitInterop.cpp
--- a/Engine/Plugins/Runtime/AR/AzureSpatialAnchorsForARKit/Source/AzureSpatialAnchorsForARKit/Private/AzureSpatialAnchorsARKitInterop.cpp
+++ b/Engine/Plugins/Runtime/AR/AzureSpatialAnchorsForARKit/Source/AzureSpatialAnchorsForARKit/Private/AzureSpatialAnchorsARKitInterop.cpp
@@ -1,6 +1,7 @@
 // Copyright Epic Games, Inc. All Rights Reserved.
 
 #include "AzureSpatialAnchorsARKitInterop.h"
+#include "IXRTrackingSystem.h"
 #include <sstream>
 
 DEFINE_LOG_CATEGORY(LogAzureSpatialAnchorsARKitInterop);
@@ -116,7 +117,7 @@ bool FAzureSpatialAnchorsARKitInterop::CreateSession()
     
     UE_LOG(LogAzureSpatialAnchorsARKitInterop, Log, TEXT("Creating an Azure Spatial Anchors session"));
     
-    auto TempARSystem = TSharedPtr<FARSupportInterface, ESPMode::ThreadSafe>{ StaticCastSharedPtr<FXRTrackingSystemBase>(GEngine->XRSystem)->GetARCompositionComponent() };
+    auto TempARSystem = TSharedPtr<FARSupportInterface, ESPMode::ThreadSafe>{ (GEngine->XRSystem)->GetARCompositionComponent() };
     if (!TempARSystem.IsValid())
     {
         UE_LOG(LogAzureSpatialAnchorsARKitInterop, Error, TEXT("Azure Spatial Anchors failed to obtain a valid AR System"));
diff --git a/Engine/Plugins/Runtime/AR/Google/GoogleARCore/GoogleARCore.uplugin b/Engine/Plugins/Runtime/AR/Google/GoogleARCore/GoogleARCore.uplugin
--- a/Engine/Plugins/Runtime/AR/Google/GoogleARCore/GoogleARCore.uplugin
+++ b/Engine/Plugins/Runtime/AR/Google/GoogleARCore/GoogleARCore.uplugin
@@ -28,18 +28,22 @@
             "PlatformAllowList": [ "Win64", "Mac", "Android", "Linux" ]
         },
     ],
-    "Plugins": [
-        {
-            "Name": "AndroidPermission",
-            "Enabled": true
-        },
+	"Plugins": [
 		{
-            "Name": "ProceduralMeshComponent",
-            "Enabled": true
-        },
+			"Name": "XRBase",
+			"Enabled": true
+		},
+		{
+			"Name": "AndroidPermission",
+			"Enabled": true
+		},
+		{
+			"Name": "ProceduralMeshComponent",
+			"Enabled": true
+		},
 		{
-            "Name": "ARUtilities",
-            "Enabled": true
-        }
-    ]
+			"Name": "ARUtilities",
+			"Enabled": true
+		}
+	]
 }
diff --git a/Engine/Plugins/Runtime/MixedRealityCaptureFramework/MixedRealityCaptureFramework.uplugin b/Engine/Plugins/Runtime/MixedRealityCaptureFramework/MixedRealityCaptureFramework.uplugin
--- a/Engine/Plugins/Runtime/MixedRealityCaptureFramework/MixedRealityCaptureFramework.uplugin
+++ b/Engine/Plugins/Runtime/MixedRealityCaptureFramework/MixedRealityCaptureFramework.uplugin
@@ -31,6 +31,10 @@
 	{
 		"Name": "OpenCVLensDistortion",
 		"Enabled": true
+	},
+	{
+		"Name": "XRBase",
+		"Enabled": true
 	}
-	]
+  ]
 }
\ No newline at end of file
diff --git a/Engine/Plugins/Runtime/OpenXR/OpenXR.uplugin b/Engine/Plugins/Runtime/OpenXR/OpenXR.uplugin
--- a/Engine/Plugins/Runtime/OpenXR/OpenXR.uplugin
+++ b/Engine/Plugins/Runtime/OpenXR/OpenXR.uplugin
@@ -61,6 +61,10 @@
 		}
 	],
 	"Plugins": [
+		{
+			"Name": "XRBase",
+			"Enabled": true
+		},
 		{
 			"Name": "EnhancedInput",
 			"Enabled": true
diff --git a/Engine/Plugins/Runtime/OpenXR/Source/OpenXRHMD/OpenXRHMD.Build.cs b/Engine/Plugins/Runtime/OpenXR/Source/OpenXRHMD/OpenXRHMD.Build.cs
--- a/Engine/Plugins/Runtime/OpenXR/Source/OpenXRHMD/OpenXRHMD.Build.cs
+++ b/Engine/Plugins/Runtime/OpenXR/Source/OpenXRHMD/OpenXRHMD.Build.cs
@@ -6,9 +6,20 @@ namespace UnrealBuildTool.Rules
 	{
 		public OpenXRHMD(ReadOnlyTargetRules Target) : base(Target)
         {
-			PublicIncludePathModuleNames.Add("OpenXR");
+			PublicIncludePathModuleNames.AddRange(
+				new string[]
+				{
+					"OpenXR"
+				}
+				);
 
-            PublicDependencyModuleNames.Add("HeadMountedDisplay");
+			PublicDependencyModuleNames.AddRange(
+				new string[]
+				{
+					"HeadMountedDisplay",
+					"XRBase",
+				}
+				);
 
 			PrivateDependencyModuleNames.AddRange(
 				new string[]
diff --git a/Engine/Plugins/Runtime/OpenXR/Source/OpenXRHMD/Public/IOpenXRExtensionPlugin.h b/Engine/Plugins/Runtime/OpenXR/Source/OpenXRHMD/Public/IOpenXRExtensionPlugin.h
--- a/Engine/Plugins/Runtime/OpenXR/Source/OpenXRHMD/Public/IOpenXRExtensionPlugin.h
+++ b/Engine/Plugins/Runtime/OpenXR/Source/OpenXRHMD/Public/IOpenXRExtensionPlugin.h
@@ -12,6 +12,8 @@
 #include "GenericPlatform/IInputInterface.h"
 #include <openxr/openxr.h>
 
+class FHeadMountedDisplayBase;
+
 class IOpenXRCustomAnchorSupport
 {
 public:
diff --git a/Engine/Plugins/Runtime/OpenXR/Source/OpenXRInput/OpenXRInput.Build.cs b/Engine/Plugins/Runtime/OpenXR/Source/OpenXRInput/OpenXRInput.Build.cs
--- a/Engine/Plugins/Runtime/OpenXR/Source/OpenXRInput/OpenXRInput.Build.cs
+++ b/Engine/Plugins/Runtime/OpenXR/Source/OpenXRInput/OpenXRInput.Build.cs
@@ -20,7 +20,8 @@ namespace UnrealBuildTool.Rules
                     "InputDevice",
                     "InputCore",
                     "HeadMountedDisplay",
-                    "OpenXRHMD"
+					"XRBase",
+                    "OpenXRHMD",
                 }
                 );
 
diff --git a/Engine/Plugins/Runtime/OpenXREyeTracker/Source/OpenXREyeTracker/OpenXREyeTracker.Build.cs b/Engine/Plugins/Runtime/OpenXREyeTracker/Source/OpenXREyeTracker/OpenXREyeTracker.Build.cs
--- a/Engine/Plugins/Runtime/OpenXREyeTracker/Source/OpenXREyeTracker/OpenXREyeTracker.Build.cs
+++ b/Engine/Plugins/Runtime/OpenXREyeTracker/Source/OpenXREyeTracker/OpenXREyeTracker.Build.cs
@@ -20,7 +20,8 @@ namespace UnrealBuildTool.Rules
                     "InputCore",
                     "HeadMountedDisplay",
                     "OpenXRHMD",
-					"OpenXRInput"
+					"OpenXRInput",
+					"XRBase"
 				}
 				);
 
diff --git a/Engine/Source/Editor/DetailCustomizations/Private/DetailCustomizations.cpp b/Engine/Source/Editor/DetailCustomizations/Private/DetailCustomizations.cpp
--- a/Engine/Source/Editor/DetailCustomizations/Private/DetailCustomizations.cpp
+++ b/Engine/Source/Editor/DetailCustomizations/Private/DetailCustomizations.cpp
@@ -161,7 +161,6 @@
 #include "VectorStructCustomization.h"
 #include "WindowsTargetSettingsDetails.h"
 #include "WorldSettingsDetails.h"
-#include "XRDeviceVisualizationDetails.h"
 #include "LandscapeGrassTypeDetails.h"
 
 struct FPerPlatformBool;
@@ -471,7 +470,6 @@ void FDetailCustomizationsModule::RegisterObjectCustomizations()
 	RegisterCustomClassLayout("Skeleton", FOnGetDetailCustomizationInstance::CreateStatic(&FSkeletonDetails::MakeInstance));
 
 	RegisterCustomClassLayout("MotionControllerComponent", FOnGetDetailCustomizationInstance::CreateStatic(&FMotionControllerDetails::MakeInstance));
-	RegisterCustomClassLayout("XRDeviceVisualizationComponent", FOnGetDetailCustomizationInstance::CreateStatic(&FXRDeviceVisualizationDetails::MakeInstance));
 
 	RegisterCustomClassLayout("Landscape", FOnGetDetailCustomizationInstance::CreateStatic(&FLandscapeUIDetails::MakeInstance));
 	RegisterCustomClassLayout("LandscapeProxy", FOnGetDetailCustomizationInstance::CreateStatic(&FLandscapeProxyUIDetails::MakeInstance));
diff --git a/Engine/Source/Editor/PIEPreviewDeviceProfileSelector/PIEPreviewDeviceProfileSelector.Build.cs b/Engine/Source/Editor/PIEPreviewDeviceProfileSelector/PIEPreviewDeviceProfileSelector.Build.cs
--- a/Engine/Source/Editor/PIEPreviewDeviceProfileSelector/PIEPreviewDeviceProfileSelector.Build.cs
+++ b/Engine/Source/Editor/PIEPreviewDeviceProfileSelector/PIEPreviewDeviceProfileSelector.Build.cs
@@ -36,6 +36,7 @@ namespace UnrealBuildTool.Rules
 			CircularlyReferencedDependentModules.AddRange(
 				new string[]
 				{
+					"MaterialShaderQualitySettings",
 					"UnrealEd"
 				}
 			);
diff --git a/Engine/Source/Editor/UnrealEd/UnrealEd.Build.cs b/Engine/Source/Editor/UnrealEd/UnrealEd.Build.cs
--- a/Engine/Source/Editor/UnrealEd/UnrealEd.Build.cs
+++ b/Engine/Source/Editor/UnrealEd/UnrealEd.Build.cs
@@ -306,7 +306,7 @@ public class UnrealEd : ModuleRules
 				"InputBindingEditor",
 				"ClothingSystemEditor",
 				"PluginWarden",
-				//"PIEPreviewDeviceProfileSelector",
+				"MaterialShaderQualitySettings",
 				"EditorInteractiveToolsFramework"
 			}
 		);
diff --git a/Engine/Source/Editor/VREditor/Private/VREditorInteractor.cpp b/Engine/Source/Editor/VREditor/Private/VREditorInteractor.cpp
--- a/Engine/Source/Editor/VREditor/Private/VREditorInteractor.cpp
+++ b/Engine/Source/Editor/VREditor/Private/VREditorInteractor.cpp
@@ -28,7 +28,6 @@
 #include "UI/VREditorRadialFloatingUI.h"
 #include "UI/VREditorUISystem.h"
 #include "VRModeSettings.h"
-#include "XRMotionControllerBase.h" // for FXRMotionControllerBase::Left/RightHandSourceId and GetHandEnumForSourceName()
 
 namespace VREd
 {
@@ -274,7 +273,7 @@ void UVREditorInteractor::Init_Implementation( UVREditorMode* InVRMode )
 	const FName HMDDeviceType = InVRMode->GetHMDDeviceType();
 
 	// Setup keys
-	if (ControllerMotionSource == FXRMotionControllerBase::LeftHandSourceId)
+	if (ControllerMotionSource == IMotionController::LeftHandSourceId)
 	{
 		AddKeyAction( UVREditorInteractor::MotionController_Left_FullyPressedTriggerAxis, FViewportActionKeyInput( ViewportWorldActionTypes::SelectAndMove_FullyPressed ) );
 		AddKeyAction( UVREditorInteractor::MotionController_Left_PressedTriggerAxis, FViewportActionKeyInput( ViewportWorldActionTypes::SelectAndMove ) );
@@ -287,7 +286,7 @@ void UVREditorInteractor::Init_Implementation( UVREditorMode* InVRMode )
 			}
 		}
 	}
-	else if (ControllerMotionSource == FXRMotionControllerBase::RightHandSourceId)
+	else if (ControllerMotionSource == IMotionController::RightHandSourceId)
 	{
 		AddKeyAction( UVREditorInteractor::MotionController_Right_FullyPressedTriggerAxis, FViewportActionKeyInput( ViewportWorldActionTypes::SelectAndMove_FullyPressed ) );
 		AddKeyAction( UVREditorInteractor::MotionController_Right_PressedTriggerAxis, FViewportActionKeyInput( ViewportWorldActionTypes::SelectAndMove ) );
@@ -326,7 +325,7 @@ void UVREditorInteractor::SetupComponent_Implementation( AActor* OwningActor )
 	// Hand mesh
 	{
 		HandMeshComponent = VRMode->CreateMotionControllerMesh(OwningAvatar, MotionControllerComponent);
-		if (ControllerMotionSource == FXRMotionControllerBase::RightHandSourceId &&
+		if (ControllerMotionSource == IMotionController::RightHandSourceId &&
 			GetHMDDeviceType() == OculusDeviceType)	// Oculus has asymmetrical controllers, so we mirror the mesh horizontally
 		{
 			HandMeshBaseScale = FVector3d(1.0, -1.0, 1.0);
@@ -353,7 +352,7 @@ void UVREditorInteractor::SetupComponent_Implementation( AActor* OwningActor )
 		LaserMotionControllerComponent->SetMobility(EComponentMobility::Movable);
 		LaserMotionControllerComponent->SetCollisionEnabled(ECollisionEnabled::NoCollision);
 		LaserMotionControllerComponent->MotionSource =
-			(ControllerMotionSource == FXRMotionControllerBase::LeftHandSourceId)
+			(ControllerMotionSource == IMotionController::LeftHandSourceId)
 			? FName("LeftAim")
 			: FName("RightAim");
 	}
@@ -530,7 +529,7 @@ void UVREditorInteractor::Shutdown_Implementation()
 EControllerHand UVREditorInteractor::GetControllerSide() const
 {
 	EControllerHand Hand = EControllerHand::Left;
-	FXRMotionControllerBase::GetHandEnumForSourceName( ControllerMotionSource, Hand );
+	IMotionController::GetHandEnumForSourceName( ControllerMotionSource, Hand );
 	return Hand;
 
 }
@@ -1021,7 +1020,7 @@ void UVREditorInteractor::HandleInputAxis( FEditorViewportClient& ViewportClient
 				bHasTriggerBeenReleasedSinceLastPress = false;
 				// Synthesize an input key for this light press
 				const EInputEvent InputEvent = IE_Pressed;
-				const bool bWasLightPressHandled = UViewportInteractor::HandleInputKey( ViewportClient, ControllerMotionSource == FXRMotionControllerBase::LeftHandSourceId ? MotionController_Left_PressedTriggerAxis : MotionController_Right_PressedTriggerAxis, InputEvent );
+				const bool bWasLightPressHandled = UViewportInteractor::HandleInputKey( ViewportClient, ControllerMotionSource == IMotionController::LeftHandSourceId ? MotionController_Left_PressedTriggerAxis : MotionController_Right_PressedTriggerAxis, InputEvent );
 			}
 			else if (bIsTriggerPressed && Delta < TriggerPressedThreshold)
 			{
@@ -1029,7 +1028,7 @@ void UVREditorInteractor::HandleInputAxis( FEditorViewportClient& ViewportClient
 
 				// Synthesize an input key for this light press
 				const EInputEvent InputEvent = IE_Released;
-				const bool bWasLightReleaseHandled = UViewportInteractor::HandleInputKey( ViewportClient, ControllerMotionSource == FXRMotionControllerBase::LeftHandSourceId ? MotionController_Left_PressedTriggerAxis : MotionController_Right_PressedTriggerAxis, InputEvent );
+				const bool bWasLightReleaseHandled = UViewportInteractor::HandleInputKey( ViewportClient, ControllerMotionSource == IMotionController::LeftHandSourceId ? MotionController_Left_PressedTriggerAxis : MotionController_Right_PressedTriggerAxis, InputEvent );
 			}
 		}
 
@@ -1048,14 +1047,14 @@ void UVREditorInteractor::HandleInputAxis( FEditorViewportClient& ViewportClient
 				bIsTriggerFullyPressed = true;
 
 				const EInputEvent InputEvent = IE_Pressed;
-				UViewportInteractor::HandleInputKey( ViewportClient, ControllerMotionSource == FXRMotionControllerBase::LeftHandSourceId ? MotionController_Left_FullyPressedTriggerAxis : MotionController_Right_FullyPressedTriggerAxis, InputEvent );
+				UViewportInteractor::HandleInputKey( ViewportClient, ControllerMotionSource == IMotionController::LeftHandSourceId ? MotionController_Left_FullyPressedTriggerAxis : MotionController_Right_FullyPressedTriggerAxis, InputEvent );
 			}
 			else if (bIsTriggerFullyPressed && Delta < TriggerPressedThreshold)
 			{
 				bIsTriggerFullyPressed = false;
 
 				const EInputEvent InputEvent = IE_Released;
-				UViewportInteractor::HandleInputKey( ViewportClient, ControllerMotionSource == FXRMotionControllerBase::LeftHandSourceId ? MotionController_Left_FullyPressedTriggerAxis : MotionController_Right_FullyPressedTriggerAxis, InputEvent );
+				UViewportInteractor::HandleInputKey( ViewportClient, ControllerMotionSource == IMotionController::LeftHandSourceId ? MotionController_Left_FullyPressedTriggerAxis : MotionController_Right_FullyPressedTriggerAxis, InputEvent );
 			}
 		}
 	}
@@ -1219,8 +1218,8 @@ void UVREditorInteractor::PlayHapticEffect( const float Strength )
 
 		//@todo viewportinteration
 		FForceFeedbackValues ForceFeedbackValues;
-		ForceFeedbackValues.LeftLarge = ControllerMotionSource == FXRMotionControllerBase::LeftHandSourceId ? Strength : 0;
-		ForceFeedbackValues.RightLarge = ControllerMotionSource == FXRMotionControllerBase::RightHandSourceId ? Strength : 0;
+		ForceFeedbackValues.LeftLarge = ControllerMotionSource == IMotionController::LeftHandSourceId ? Strength : 0;
+		ForceFeedbackValues.RightLarge = ControllerMotionSource == IMotionController::RightHandSourceId ? Strength : 0;
 
 		// @todo vreditor: If an Xbox controller is plugged in, this causes both the motion controllers and the Xbox controller to vibrate!
 		InputInterface->SetForceFeedbackChannelValues( WorldInteraction->GetMotionControllerID(), ForceFeedbackValues );
@@ -1538,7 +1537,7 @@ void UVREditorInteractor::UpdateHelpLabels()
 			FTransform SocketRelativeTransform( Socket->RelativeRotation, Socket->RelativeLocation, Socket->RelativeScale );
 
 			// Oculus has asymmetrical controllers, so we the sock transform horizontally
-			if (ControllerMotionSource == FXRMotionControllerBase::RightHandSourceId && GetVRMode().GetHMDDeviceType() == OculusDeviceType)
+			if (ControllerMotionSource == IMotionController::RightHandSourceId && GetVRMode().GetHMDDeviceType() == OculusDeviceType)
 			{
 				const FVector Scale3D = SocketRelativeTransform.GetLocation();
 				SocketRelativeTransform.SetLocation( FVector( Scale3D.X, -Scale3D.Y, Scale3D.Z ) );
diff --git a/Engine/Source/Editor/VREditor/Private/VREditorMode.cpp b/Engine/Source/Editor/VREditor/Private/VREditorMode.cpp
--- a/Engine/Source/Editor/VREditor/Private/VREditorMode.cpp
+++ b/Engine/Source/Editor/VREditor/Private/VREditorMode.cpp
@@ -51,7 +51,7 @@
 #include "Engine/StaticMeshActor.h"
 #include "EngineUtils.h"
 #include "Kismet/GameplayStatics.h"
-#include "XRMotionControllerBase.h" // for FXRMotionControllerBase::Left/RightHandSourceId
+#include "IMotionController.h"
 #include "UI/VREditorFloatingUI.h"
 #include "AssetEditorViewportLayout.h"
 #include "LevelViewportActions.h"
@@ -221,8 +221,8 @@ void UVREditorMode::AllocateInteractors()
 	Interactors.Add( LeftHandInteractor );
 	Interactors.Add( RightHandInteractor );
 
-	LeftHandInteractor->SetControllerHandSide( FXRMotionControllerBase::LeftHandSourceId );
-	RightHandInteractor->SetControllerHandSide( FXRMotionControllerBase::RightHandSourceId );
+	LeftHandInteractor->SetControllerHandSide(IMotionController::LeftHandSourceId );
+	RightHandInteractor->SetControllerHandSide( IMotionController::RightHandSourceId );
 
 	for (UVREditorInteractor* Interactor : Interactors)
 	{
diff --git a/Engine/Source/Programs/UnrealBuildTool/Configuration/UEBuildModuleCPP.cs b/Engine/Source/Programs/UnrealBuildTool/Configuration/UEBuildModuleCPP.cs
--- a/Engine/Source/Programs/UnrealBuildTool/Configuration/UEBuildModuleCPP.cs
+++ b/Engine/Source/Programs/UnrealBuildTool/Configuration/UEBuildModuleCPP.cs
@@ -161,6 +161,7 @@ namespace UnrealBuildTool
 			new KeyValuePair<string, string>("NavigationSystem", "UnrealEd"),
 			new KeyValuePair<string, string>("PacketHandler", "ReliabilityHandlerComponent"),
 			new KeyValuePair<string, string>("PIEPreviewDeviceProfileSelector", "UnrealEd"),
+			new KeyValuePair<string, string>("PIEPreviewDeviceProfileSelector", "MaterialShaderQualitySettings"),
 			new KeyValuePair<string, string>("PixelInspectorModule", "UnrealEd"),
 			new KeyValuePair<string, string>("Sequencer", "MovieSceneTools"),
 			new KeyValuePair<string, string>("Sequencer", "ViewportInteraction"),
@@ -178,6 +179,7 @@ namespace UnrealBuildTool
 			new KeyValuePair<string, string>("UnrealEd", "ToolMenusEditor"),
 			new KeyValuePair<string, string>("UnrealEd", "ViewportInteraction"),
 			new KeyValuePair<string, string>("UnrealEd", "VREditor"),
+			new KeyValuePair<string, string>("UnrealEd", "MaterialShaderQualitySettings"),
 			new KeyValuePair<string, string>("WebBrowser", "WebBrowserTexture"),
 		};
 
diff --git a/Engine/Source/Runtime/AugmentedReality/Private/ARBlueprintLibrary.cpp b/Engine/Source/Runtime/AugmentedReality/Private/ARBlueprintLibrary.cpp
--- a/Engine/Source/Runtime/AugmentedReality/Private/ARBlueprintLibrary.cpp
+++ b/Engine/Source/Runtime/AugmentedReality/Private/ARBlueprintLibrary.cpp
@@ -6,6 +6,7 @@
 #include "Engine/Engine.h"
 #include "ARPin.h"
 #include "ARGeoTrackingSupport.h"
+#include "IXRTrackingSystem.h"
 
 #include UE_INLINE_GENERATED_CPP_BY_NAME(ARBlueprintLibrary)
 
diff --git a/Engine/Source/Runtime/AugmentedReality/Private/ARPin.cpp b/Engine/Source/Runtime/AugmentedReality/Private/ARPin.cpp
--- a/Engine/Source/Runtime/AugmentedReality/Private/ARPin.cpp
+++ b/Engine/Source/Runtime/AugmentedReality/Private/ARPin.cpp
@@ -1,10 +1,12 @@
 // Copyright Epic Games, Inc. All Rights Reserved.
 
 #include "ARPin.h"
+#include "ARSupportInterface.h"
 #include "ARSystem.h"
 #include "ARDebugDrawHelpers.h"
 #include "DrawDebugHelpers.h"
 #include "Components/SceneComponent.h"
+#include "IXRTrackingSystem.h"
 
 #include "Engine/Engine.h"
 
diff --git a/Engine/Source/Runtime/AugmentedReality/Private/ARTraceResult.cpp b/Engine/Source/Runtime/AugmentedReality/Private/ARTraceResult.cpp
--- a/Engine/Source/Runtime/AugmentedReality/Private/ARTraceResult.cpp
+++ b/Engine/Source/Runtime/AugmentedReality/Private/ARTraceResult.cpp
@@ -2,6 +2,8 @@
 
 #include "ARTraceResult.h"
 #include "ARSystem.h"
+#include "ARSupportInterface.h"
+#include "IXRTrackingSystem.h"
 
 #include UE_INLINE_GENERATED_CPP_BY_NAME(ARTraceResult)
 
diff --git a/Engine/Source/Runtime/AugmentedReality/Public/ARSystem.h b/Engine/Source/Runtime/AugmentedReality/Public/ARSystem.h
--- a/Engine/Source/Runtime/AugmentedReality/Public/ARSystem.h
+++ b/Engine/Source/Runtime/AugmentedReality/Public/ARSystem.h
@@ -3,17 +3,18 @@
 #pragma once
 
 #include "CoreMinimal.h"
-#include "XRTrackingSystemBase.h"
 #include "ARTypes.h"
 #include "ARSessionConfig.h"
 #include "Engine/Engine.h" // for FWorldContext
 
 class UARPin;
+class UARTexture;
 class UARTrackedGeometry;
 class UARSessionConfig;
 class UARTextureCameraImage;
 class UARTextureCameraDepth;
 struct FARTraceResult;
+enum class EARTextureType : uint8;
 
 #define DEFINE_AR_DELEGATE_BASE(DelegateName) \
 private: \
diff --git a/Engine/Source/Runtime/AugmentedReality/Public/ARTypes.h b/Engine/Source/Runtime/AugmentedReality/Public/ARTypes.h
--- a/Engine/Source/Runtime/AugmentedReality/Public/ARTypes.h
+++ b/Engine/Source/Runtime/AugmentedReality/Public/ARTypes.h
@@ -7,7 +7,6 @@
 #include "ARTypes.generated.h"
 
 class USceneComponent;
-class IXRTrackingSystem;
 class UARPin;
 class UARTrackedGeometry;
 class UARLightEstimate;
diff --git a/Engine/Source/Runtime/Engine/Public/IStereoLayers.h b/Engine/Source/Runtime/Engine/Public/IStereoLayers.h
--- a/Engine/Source/Runtime/Engine/Public/IStereoLayers.h
+++ b/Engine/Source/Runtime/Engine/Public/IStereoLayers.h
@@ -239,64 +239,6 @@ public:
 	virtual bool ShouldCopyDebugLayersToSpectatorScreen() const = 0;
 
 public:
-	/**
-	* Set the splash screen attributes
-	*
-	* @param Texture			(in) A texture to be used for the splash. B8R8G8A8 format.
-	* @param Scale				(in) Scale of the texture.
-	* @param Offset				(in) Position from which to start rendering the texture.
-	* @param ShowLoadingMovie	(in) Whether the splash screen presents loading movies.
-	*/
-	UE_DEPRECATED(4.24, "Use the IXRLoadingScreen interface instead of IStereoLayers::*SplashScreen")
-	void SetSplashScreen(FTextureRHIRef Texture, FVector2D Scale, FVector Offset, bool bShowLoadingMovie)
-	{
-		bSplashShowMovie = bShowLoadingMovie;
-		SplashTexture = nullptr;
-		if (Texture)
-		{
-			SplashTexture = Texture->GetTexture2D();
-			SplashOffset = Offset;
-			SplashScale = Scale;
-		}
-	}
-
-	/**
-	* Show the splash screen and override the normal VR display
-	*/
-	UE_DEPRECATED(4.24, "Use the IXRLoadingScreen interface instead of IStereoLayers::*SplashScreen")
-	void ShowSplashScreen()
-	{
-		bSplashIsShown = true;
-		UpdateSplashScreen();
-	}
-
-	/**
-	* Hide the splash screen and return to normal display.
-	*/
-	UE_DEPRECATED(4.24, "Use the IXRLoadingScreen interface instead of IStereoLayers::*SplashScreen")
-	void HideSplashScreen()
-	{
-		bSplashIsShown = false;
-		UpdateSplashScreen();
-	}
-
-	/**
-	* Set the splash screen's movie texture.
-	*
-	* @param InMovieTexture		(in) A movie texture to be used for the splash. B8R8G8A8 format.
-	*/
-	UE_DEPRECATED(4.24, "Use the IXRLoadingScreen interface instead of IStereoLayers::*SplashScreen")
-	void SetSplashScreenMovie(FTextureRHIRef Texture)
-	{
-		SplashMovie = nullptr;
-		if (Texture)
-		{
-			SplashMovie = Texture->GetTexture2D();
-			bSplashShowMovie = true;
-		}
-		UpdateSplashScreen();
-	}
-
 	virtual FLayerDesc GetDebugCanvasLayerDesc(FTextureRHIRef Texture)
 	{
 		// Default debug layer desc
diff --git a/Engine/Source/Runtime/HeadMountedDisplay/HeadMountedDisplay.Build.cs b/Engine/Source/Runtime/HeadMountedDisplay/HeadMountedDisplay.Build.cs
--- a/Engine/Source/Runtime/HeadMountedDisplay/HeadMountedDisplay.Build.cs
+++ b/Engine/Source/Runtime/HeadMountedDisplay/HeadMountedDisplay.Build.cs
@@ -13,30 +13,11 @@ public class HeadMountedDisplay : ModuleRules
 				"CoreUObject",
 				"Engine",
                 "InputCore",
-				"Slate",
-				"SlateCore",
                 "RHI",
                 "Renderer",
                 "RenderCore",
-                "Analytics",
                 "EngineSettings",
             }
         );
-
-        PublicDependencyModuleNames.AddRange(
-            new string[] {
-                "AugmentedReality",
-            }
-        );
-
-        if (Target.bBuildEditor == true)
-		{
-			PrivateDependencyModuleNames.AddRange(
-				new string[]
-				{
-					"EditorFramework",
-					"UnrealEd"
-				});
-		}
 	}
 }
diff --git a/Engine/Source/Runtime/HeadMountedDisplay/Private/IXRTrackingSystem.cpp b/Engine/Source/Runtime/HeadMountedDisplay/Private/IXRTrackingSystem.cpp
--- a/Engine/Source/Runtime/HeadMountedDisplay/Private/IXRTrackingSystem.cpp
+++ b/Engine/Source/Runtime/HeadMountedDisplay/Private/IXRTrackingSystem.cpp
@@ -7,6 +7,8 @@
 #include "Kismet/GameplayStatics.h"
 #include "Misc/App.h"
 
+const int32 IXRTrackingSystem::HMDDeviceId;
+
 void IXRTrackingSystem::GetHMDData(UObject* WorldContext, FXRHMDData& HMDData)
 {
 	HMDData.bValid = true;
diff --git a/Engine/Source/Runtime/HeadMountedDisplay/Private/MotionControllerComponent.cpp b/Engine/Source/Runtime/HeadMountedDisplay/Private/MotionControllerComponent.cpp
--- a/Engine/Source/Runtime/HeadMountedDisplay/Private/MotionControllerComponent.cpp
+++ b/Engine/Source/Runtime/HeadMountedDisplay/Private/MotionControllerComponent.cpp
@@ -10,14 +10,15 @@
 #include "Components/StaticMeshComponent.h"
 #include "MotionDelayBuffer.h"
 #include "UObject/VRObjectVersion.h"
-#include "XRDeviceVisualizationComponent.h"
-#include "XRMotionControllerBase.h"
 #include "IXRTrackingSystem.h"
 
 #include UE_INLINE_GENERATED_CPP_BY_NAME(MotionControllerComponent)
 
 DEFINE_LOG_CATEGORY_STATIC(LogMotionControllerComponent, Log, All);
 
+UMotionControllerComponent::FActivateVisualizationComponent UMotionControllerComponent::OnActivateVisualizationComponent;
+
+
 namespace {
 	/** This is to prevent destruction of motion controller components while they are
 		in the middle of being accessed by the render thread */
@@ -66,7 +67,7 @@ UMotionControllerComponent::UMotionControllerComponent(const FObjectInitializer&
 	PrimaryComponentTick.bTickEvenWhenPaused = true;
 
 	PlayerIndex = 0;
-	MotionSource = FXRMotionControllerBase::LeftHandSourceId;
+	MotionSource = IMotionController::LeftHandSourceId;
 	bDisableLowLatencyUpdate = false;
 	bHasAuthority = false;
 	bAutoActivate = true;
@@ -102,24 +103,6 @@ void UMotionControllerComponent::PostLoad()
 	}
 }
 
-//=============================================================================
-void UMotionControllerComponent::BeginPlay() 
-{
-	Super::BeginPlay();
-
-	TArray<USceneComponent*> ChildrenComponents;
-	GetChildrenComponents(/*bIncludeAllDescendants*/ false, ChildrenComponents);
-
-	for (USceneComponent* ChildComponent : ChildrenComponents) 
-	{
-		VisualizationComponent = Cast<UXRDeviceVisualizationComponent>(ChildComponent);
-		if (VisualizationComponent != nullptr) 
-		{
-			return;
-		}
-	}
-}
-
 //=============================================================================
 FName UMotionControllerComponent::GetTrackingMotionSource() 
 {
@@ -144,13 +127,9 @@ void UMotionControllerComponent::TickComponent(float DeltaTime, enum ELevelTick
 
 		// if controller tracking just kicked in or we haven't started rendering in the (possibly present) 
 		// visualization component.
-		if (!bTracked && bNewTrackedState && !bHasStartedRendering)
+		if (!bTracked && bNewTrackedState)
 		{
-			if (VisualizationComponent)
-			{
-				VisualizationComponent->SetIsRenderingActive(true);
-				bHasStartedRendering = true;
-			}
+			OnActivateVisualizationComponent.Broadcast(true);
 		}
 
 		// This part is deprecated and will be removed in later versions.
@@ -217,7 +196,7 @@ void UMotionControllerComponent::SetTrackingSource(const EControllerHand NewSour
 EControllerHand UMotionControllerComponent::GetTrackingSource() const
 {
 	EControllerHand Hand = EControllerHand::Left;
-	FXRMotionControllerBase::GetHandEnumForSourceName(MotionSource, Hand);
+	IMotionController::GetHandEnumForSourceName(MotionSource, Hand);
 	return Hand;
 }
 
@@ -250,11 +229,6 @@ void UMotionControllerComponent::Serialize(FArchive& Ar)
 	Ar.UsingCustomVersion(FVRObjectVersion::GUID);
 
 	Super::Serialize(Ar);
-
-	if (Ar.CustomVer(FVRObjectVersion::GUID) < FVRObjectVersion::UseFNameInsteadOfEControllerHandForMotionSource)
-	{
-		LegacyMotionSources::GetSourceNameForHand(Hand_DEPRECATED, MotionSource);
-	}
 }
 
 //=============================================================================
@@ -384,7 +358,7 @@ bool UMotionControllerComponent::PollControllerState(FVector& Position, FRotator
 			}
 		}
 
-		if (MotionSource == FXRMotionControllerBase::HMDSourceId)
+		if (MotionSource == IMotionController::HMDSourceId)
 		{
 			IXRTrackingSystem* TrackingSys = GEngine->XRSystem.Get();
 			if (TrackingSys)
@@ -672,14 +646,14 @@ void UMotionControllerComponent::RefreshDisplayComponent(const bool bForceDestro
 					}
 
 					int32 DeviceId = INDEX_NONE;
-					if (MotionSource == FXRMotionControllerBase::HMDSourceId)
+					if (MotionSource == IMotionController::HMDSourceId)
 					{
 						DeviceId = IXRTrackingSystem::HMDDeviceId;
 					}
 					else
 					{
 						EControllerHand ControllerHandIndex;
-						if (!FXRMotionControllerBase::GetHandEnumForSourceName(MotionSource, ControllerHandIndex))
+						if (!IMotionController::GetHandEnumForSourceName(MotionSource, ControllerHandIndex))
 						{
 							break;
 						}
diff --git a/Engine/Source/Runtime/HeadMountedDisplay/Private/MotionDelayBuffer.cpp b/Engine/Source/Runtime/HeadMountedDisplay/Private/MotionDelayBuffer.cpp
--- a/Engine/Source/Runtime/HeadMountedDisplay/Private/MotionDelayBuffer.cpp
+++ b/Engine/Source/Runtime/HeadMountedDisplay/Private/MotionDelayBuffer.cpp
@@ -11,7 +11,6 @@
 #include "Templates/TypeHash.h"
 #include "Features/IModularFeatures.h"
 #include "GameFramework/WorldSettings.h"
-#include "XRMotionControllerBase.h"
 #include "Engine/Engine.h" // for GEngine->XRSystem
 
 DEFINE_LOG_CATEGORY_STATIC(LogMotionDelayBuffer, Log, All);
@@ -279,7 +278,7 @@ static bool MotionDelayService_Impl::SampleDevicePose(const FMotionDelayTarget::
 		}
 	}
 
-	if (TargetSource.SourceId == FXRMotionControllerBase::HMDSourceId)
+	if (TargetSource.SourceId == IMotionController::HMDSourceId)
 	{
 		IXRTrackingSystem* TrackingSys = GEngine->XRSystem.Get();
 		if (TrackingSys)
diff --git a/Engine/Source/Runtime/HeadMountedDisplay/Public/HeadMountedDisplayTypes.h b/Engine/Source/Runtime/HeadMountedDisplay/Public/HeadMountedDisplayTypes.h
--- a/Engine/Source/Runtime/HeadMountedDisplay/Public/HeadMountedDisplayTypes.h
+++ b/Engine/Source/Runtime/HeadMountedDisplay/Public/HeadMountedDisplayTypes.h
@@ -229,8 +229,6 @@ struct FSpectatorScreenModeTexturePlusEyeLayout
 	bool bClearBlack;
 };
 
-DECLARE_DELEGATE_FiveParams(FSpectatorScreenRenderDelegate, FRHICommandListImmediate& /* RHICmdList */, FTexture2DRHIRef /* TargetTexture */, FTexture2DRHIRef /* EyeTexture */, FTexture2DRHIRef /* OtherTexture */, FVector2D /* WindowSize */);
-
 UENUM(BlueprintType)
 enum class EXRTrackedDeviceType : uint8
 {
diff --git a/Engine/Source/Runtime/HeadMountedDisplay/Public/IMotionController.h b/Engine/Source/Runtime/HeadMountedDisplay/Public/IMotionController.h
--- a/Engine/Source/Runtime/HeadMountedDisplay/Public/IMotionController.h
+++ b/Engine/Source/Runtime/HeadMountedDisplay/Public/IMotionController.h
@@ -122,7 +122,7 @@ public:
 	 * @return			The value of the parameter
 	 */
 	virtual float GetCustomParameterValue(const FName MotionSource, FName ParameterName, bool& bOutValueFound) const = 0;
-
+	
 	virtual bool GetHandJointPosition(const FName MotionSource, int jointIndex, FVector& OutPosition) const = 0;
 
 	/**
@@ -134,4 +134,11 @@ public:
 	 * @return			False if the input config can't be attached to the session, true otherwise
 	 */
 	virtual bool SetPlayerMappableInputConfig(TObjectPtr<class UPlayerMappableInputConfig> InputConfig = nullptr) { return true; };
+
+	// explicit source names
+	static FName LeftHandSourceId;
+	static FName RightHandSourceId;
+	static FName HMDSourceId;
+
+	static bool GetHandEnumForSourceName(const FName Source, EControllerHand& OutHand);
 };
diff --git a/Engine/Source/Runtime/HeadMountedDisplay/Public/ISpectatorScreenController.h b/Engine/Source/Runtime/HeadMountedDisplay/Public/ISpectatorScreenController.h
--- a/Engine/Source/Runtime/HeadMountedDisplay/Public/ISpectatorScreenController.h
+++ b/Engine/Source/Runtime/HeadMountedDisplay/Public/ISpectatorScreenController.h
@@ -19,6 +19,8 @@ struct FSpectatorScreenModeTexturePlusEyeLayout;
 class HEADMOUNTEDDISPLAY_API ISpectatorScreenController : public IModularFeature
 {
 public:
+	virtual ~ISpectatorScreenController() {}
+
 	static FName GetModularFeatureName()
 	{
 		static FName FeatureName = FName(TEXT("SpectatorScreenController"));
diff --git a/Engine/Source/Runtime/HeadMountedDisplay/Public/IXRTrackingSystem.h b/Engine/Source/Runtime/HeadMountedDisplay/Public/IXRTrackingSystem.h
--- a/Engine/Source/Runtime/HeadMountedDisplay/Public/IXRTrackingSystem.h
+++ b/Engine/Source/Runtime/HeadMountedDisplay/Public/IXRTrackingSystem.h
@@ -8,13 +8,13 @@
 #include "UObject/ObjectMacros.h"
 #include "Features/IModularFeature.h"
 #include "IXRInput.h"
-#include "XRGestureConfig.h"
 #include "StereoRendering.h"
 
 class IXRCamera;
 class UARPin;
 class FSceneViewFamily;
 struct FWorldContext;
+class FARSupportInterface;
 
 /**
  * Struct representing the properties of an external tracking sensor.
@@ -340,6 +340,13 @@ public:
 	virtual IXRInput* GetXRInput() { return nullptr; }
 
 
+	/**
+	* Access optionsal ARCompositionComponent
+	**/
+	virtual TSharedPtr<FARSupportInterface, ESPMode::ThreadSafe> GetARCompositionComponent() { return nullptr; }
+	virtual const TSharedPtr<const FARSupportInterface, ESPMode::ThreadSafe> GetARCompositionComponent() const { return nullptr; }
+
+
 	/**
 	 * Access the loading screen interface associated with this tracking system, if any.
 	 *
@@ -424,7 +431,8 @@ public:
 	virtual void GetMotionControllerData(UObject* WorldContext, const EControllerHand Hand, FXRMotionControllerData& MotionControllerData) = 0;
 	virtual bool GetCurrentInteractionProfile(const EControllerHand Hand, FString& InteractionProfile) = 0;
 
-	virtual bool ConfigureGestures(const FXRGestureConfig& GestureConfig) = 0;
+	UE_DEPRECATED(5.3, "The only implementation for this function was removed many UE releases ago.")
+	virtual bool ConfigureGestures(const struct FXRGestureConfig& GestureConfig) { return false; };
 
 	virtual EXRDeviceConnectionResult::Type ConnectRemoteXRDevice(const FString& IpAddress, const int32 BitRate)
 	{ 
diff --git a/Engine/Source/Runtime/HeadMountedDisplay/Public/MotionControllerComponent.h b/Engine/Source/Runtime/HeadMountedDisplay/Public/MotionControllerComponent.h
--- a/Engine/Source/Runtime/HeadMountedDisplay/Public/MotionControllerComponent.h
+++ b/Engine/Source/Runtime/HeadMountedDisplay/Public/MotionControllerComponent.h
@@ -24,10 +24,6 @@ class HEADMOUNTEDDISPLAY_API UMotionControllerComponent : public UPrimitiveCompo
 	UPROPERTY(EditAnywhere, BlueprintReadWrite, BlueprintSetter = SetAssociatedPlayerIndex, Category = "MotionController")
 	int32 PlayerIndex;
 
-	/** DEPRECATED (use MotionSource instead) Which hand this component should automatically follow */
-	UPROPERTY(BlueprintSetter = SetTrackingSource, BlueprintGetter = GetTrackingSource, Category = "MotionController")
-	EControllerHand Hand_DEPRECATED;
-
 	/** Defines which pose this component should receive from the OpenXR Runtime. Left/Right MotionSource is the same as LeftGrip/RightGrip. See OpenXR specification for details on poses. */
 	UPROPERTY(EditAnywhere, BlueprintReadWrite, BlueprintSetter = SetTrackingMotionSource, Category = "MotionController")
 	FName MotionSource;
@@ -40,9 +36,6 @@ class HEADMOUNTEDDISPLAY_API UMotionControllerComponent : public UPrimitiveCompo
 	UPROPERTY(BlueprintReadOnly, Category = "MotionController")
 	ETrackingStatus CurrentTrackingStatus;
 
-	/** Used to visualize this component's device */
-	UXRDeviceVisualizationComponent* VisualizationComponent;
-
 	/** Used to automatically render a model associated with the set hand. */
 	UE_DEPRECATED(5.2, "bDisplayDeviceModel is deprecated. Please use the XRDeviceVisualizationComponent for rendering instead.")
 	UPROPERTY(EditAnywhere, BlueprintReadWrite, BlueprintSetter = SetShowDeviceModel, Category = "Visualization", meta = (DeprecatedProperty, DeprecationMessage = "bDisplayDeviceModel is deprecated. Please use the XRDeviceVisualizationComponent for rendering instead."))
@@ -98,7 +91,6 @@ class HEADMOUNTEDDISPLAY_API UMotionControllerComponent : public UPrimitiveCompo
 	UFUNCTION(BlueprintSetter)
 	void SetAssociatedPlayerIndex(const int32 NewPlayer);
 
-	void BeginPlay() override;
 	void TickComponent(float DeltaTime, enum ELevelTick TickType, FActorComponentTickFunction* ThisTickFunction) override;
 	void BeginDestroy() override;
 
@@ -141,6 +133,10 @@ public:
 	virtual void InitializeComponent() override;
 	virtual void OnComponentDestroyed(bool bDestroyingHierarchy) override;
 
+	// Delegate for activation of XRDeviceVisualizationComponent
+	DECLARE_MULTICAST_DELEGATE_OneParam(FActivateVisualizationComponent, bool);
+	static FActivateVisualizationComponent OnActivateVisualizationComponent;
+
 protected:
 	//~ Begin UActorComponent Interface.
 	virtual void CreateRenderState_Concurrent(FRegisterComponentContext* Context) override;
@@ -169,9 +165,6 @@ private:
 	/** Whether or not this component has authority within the frame*/
 	bool bHasAuthority;
 
-	/** Whether or not this component has informed the visualization component (if present) to start rendering */
-	bool bHasStartedRendering;
-
 	/** If true, the Position and Orientation args will contain the most recent controller state */
 	bool PollControllerState(FVector& Position, FRotator& Orientation, float WorldToMetersScale);
 
diff --git a/Engine/Source/Runtime/HeadMountedDisplay/Public/XRGestureConfig.h b/Engine/Source/Runtime/HeadMountedDisplay/Public/XRGestureConfig.h
--- a/Engine/Source/Runtime/HeadMountedDisplay/Public/XRGestureConfig.h
+++ b/Engine/Source/Runtime/HeadMountedDisplay/Public/XRGestureConfig.h
@@ -14,7 +14,7 @@ enum class ESpatialInputGestureAxis : uint8
 };
 
 
-USTRUCT(BlueprintType)
+USTRUCT(BlueprintType, meta = (Deprecated = "5.3"))
 struct HEADMOUNTEDDISPLAY_API FXRGestureConfig
 {
 	GENERATED_BODY()
diff --git a/Engine/Source/Runtime/Renderer/Private/VariableRateShading/FoveatedImageGenerator.cpp b/Engine/Source/Runtime/Renderer/Private/VariableRateShading/FoveatedImageGenerator.cpp
--- a/Engine/Source/Runtime/Renderer/Private/VariableRateShading/FoveatedImageGenerator.cpp
+++ b/Engine/Source/Runtime/Renderer/Private/VariableRateShading/FoveatedImageGenerator.cpp
@@ -11,7 +11,6 @@
 #include "SceneTexturesConfig.h"
 #include "IEyeTracker.h"
 #include "IHeadMountedDisplay.h"
-#include "HeadMountedDisplayFunctionLibrary.h"
 #include "SceneRendering.h"
 
 /* CVar values used to control generator behavior */
diff --git a/Templates/TP_FirstPerson/Source/TP_FirstPerson/TP_FirstPerson.Build.cs b/Templates/TP_FirstPerson/Source/TP_FirstPerson/TP_FirstPerson.Build.cs
--- a/Templates/TP_FirstPerson/Source/TP_FirstPerson/TP_FirstPerson.Build.cs
+++ b/Templates/TP_FirstPerson/Source/TP_FirstPerson/TP_FirstPerson.Build.cs
@@ -8,6 +8,6 @@ public class TP_FirstPerson : ModuleRules
 	{
 		PCHUsage = PCHUsageMode.UseExplicitOrSharedPCHs;
 
-		PublicDependencyModuleNames.AddRange(new string[] { "Core", "CoreUObject", "Engine", "InputCore", "HeadMountedDisplay", "EnhancedInput" });
+		PublicDependencyModuleNames.AddRange(new string[] { "Core", "CoreUObject", "Engine", "InputCore", "EnhancedInput" });
 	}
 }
diff --git a/Templates/TP_Puzzle/Source/TP_Puzzle/TP_Puzzle.Build.cs b/Templates/TP_Puzzle/Source/TP_Puzzle/TP_Puzzle.Build.cs
--- a/Templates/TP_Puzzle/Source/TP_Puzzle/TP_Puzzle.Build.cs
+++ b/Templates/TP_Puzzle/Source/TP_Puzzle/TP_Puzzle.Build.cs
@@ -8,6 +8,6 @@ public class TP_Puzzle : ModuleRules
 	{
 		PCHUsage = PCHUsageMode.UseExplicitOrSharedPCHs;
 
-		PublicDependencyModuleNames.AddRange(new string[] { "Core", "CoreUObject", "Engine", "InputCore", "HeadMountedDisplay" });
+		PublicDependencyModuleNames.AddRange(new string[] { "Core", "CoreUObject", "Engine", "InputCore" });
 	}
 }
diff --git a/Templates/TP_Puzzle/Source/TP_Puzzle/TP_PuzzlePawn.cpp b/Templates/TP_Puzzle/Source/TP_Puzzle/TP_PuzzlePawn.cpp
--- a/Templates/TP_Puzzle/Source/TP_Puzzle/TP_PuzzlePawn.cpp
+++ b/Templates/TP_Puzzle/Source/TP_Puzzle/TP_PuzzlePawn.cpp
@@ -2,7 +2,6 @@
 
 #include "TP_PuzzlePawn.h"
 #include "TP_PuzzleBlock.h"
-#include "HeadMountedDisplayFunctionLibrary.h"
 #include "Camera/CameraComponent.h"
 #include "GameFramework/PlayerController.h"
 #include "Engine/World.h"
@@ -20,22 +19,10 @@ void ATP_PuzzlePawn::Tick(float DeltaSeconds)
 
 	if (APlayerController* PC = Cast<APlayerController>(GetController()))
 	{
-		if (UHeadMountedDisplayFunctionLibrary::IsHeadMountedDisplayEnabled())
-		{
-			if (UCameraComponent* OurCamera = PC->GetViewTarget()->FindComponentByClass<UCameraComponent>())
-			{
-				FVector Start = OurCamera->GetComponentLocation();
-				FVector End = Start + (OurCamera->GetComponentRotation().Vector() * 8000.0f);
-				TraceForBlock(Start, End, true);
-			}
-		}
-		else
-		{
-			FVector Start, Dir, End;
-			PC->DeprojectMousePositionToWorld(Start, Dir);
-			End = Start + (Dir * 8000.0f);
-			TraceForBlock(Start, End, false);
-		}
+		FVector Start, Dir, End;
+		PC->DeprojectMousePositionToWorld(Start, Dir);
+		End = Start + (Dir * 8000.0f);
+		TraceForBlock(Start, End, false);
 	}
 }
 
diff --git a/Templates/TP_ThirdPerson/Source/TP_ThirdPerson/TP_ThirdPerson.Build.cs b/Templates/TP_ThirdPerson/Source/TP_ThirdPerson/TP_ThirdPerson.Build.cs
--- a/Templates/TP_ThirdPerson/Source/TP_ThirdPerson/TP_ThirdPerson.Build.cs
+++ b/Templates/TP_ThirdPerson/Source/TP_ThirdPerson/TP_ThirdPerson.Build.cs
@@ -8,6 +8,6 @@ public class TP_ThirdPerson : ModuleRules
 	{
 		PCHUsage = PCHUsageMode.UseExplicitOrSharedPCHs;
 
-		PublicDependencyModuleNames.AddRange(new string[] { "Core", "CoreUObject", "Engine", "InputCore", "HeadMountedDisplay", "EnhancedInput" });
+		PublicDependencyModuleNames.AddRange(new string[] { "Core", "CoreUObject", "Engine", "InputCore", "EnhancedInput" });
 	}
 }
diff --git a/Templates/TP_TopDown/Source/TP_TopDown/TP_TopDown.Build.cs b/Templates/TP_TopDown/Source/TP_TopDown/TP_TopDown.Build.cs
--- a/Templates/TP_TopDown/Source/TP_TopDown/TP_TopDown.Build.cs
+++ b/Templates/TP_TopDown/Source/TP_TopDown/TP_TopDown.Build.cs
@@ -8,6 +8,6 @@ public class TP_TopDown : ModuleRules
 	{
 		PCHUsage = PCHUsageMode.UseExplicitOrSharedPCHs;
 
-        PublicDependencyModuleNames.AddRange(new string[] { "Core", "CoreUObject", "Engine", "InputCore", "HeadMountedDisplay", "NavigationSystem", "AIModule", "Niagara", "EnhancedInput" });
+        PublicDependencyModuleNames.AddRange(new string[] { "Core", "CoreUObject", "Engine", "InputCore", "NavigationSystem", "AIModule", "Niagara", "EnhancedInput" });
     }
 }
diff --git a/Templates/TP_VehicleAdv/Source/TP_VehicleAdv/TP_VehicleAdv.Build.cs b/Templates/TP_VehicleAdv/Source/TP_VehicleAdv/TP_VehicleAdv.Build.cs
--- a/Templates/TP_VehicleAdv/Source/TP_VehicleAdv/TP_VehicleAdv.Build.cs
+++ b/Templates/TP_VehicleAdv/Source/TP_VehicleAdv/TP_VehicleAdv.Build.cs
@@ -8,8 +8,6 @@ public class TP_VehicleAdv : ModuleRules
 	{
 		PCHUsage = PCHUsageMode.UseExplicitOrSharedPCHs;
 
-		PublicDependencyModuleNames.AddRange(new string[] { "Core", "CoreUObject", "Engine", "InputCore", "ChaosVehicles", "HeadMountedDisplay", "PhysicsCore" });
-
-		PublicDefinitions.Add("HMD_MODULE_INCLUDED=1");
+		PublicDependencyModuleNames.AddRange(new string[] { "Core", "CoreUObject", "Engine", "InputCore", "ChaosVehicles", "PhysicsCore" });
 	}
 }
