Fix LWC float truncation warnings for VREditor, AugmentedReality, and MRMesh modules.
Primarily float casts or FVector2f conversions for calculations that are taking place in local space or screen space anyway.

diff --git a/Engine/Source/Editor/VREditor/Private/Teleporter/VREditorTeleporter.cpp b/Engine/Source/Editor/VREditor/Private/Teleporter/VREditorTeleporter.cpp
--- a/Engine/Source/Editor/VREditor/Private/Teleporter/VREditorTeleporter.cpp
+++ b/Engine/Source/Editor/VREditor/Private/Teleporter/VREditorTeleporter.cpp
@@ -373,7 +373,7 @@ void AVREditorTeleporter::UpdateTeleportAim(const float DeltaTime)
 			{
 				// Calculate an offset with the impact normal, so the teleporter will show up on-top, underneath or next to where laser is aiming at.
 				OffsetDistance = HitResult.ImpactNormal * (50 * TeleportGoalScale);
-				float OffsetFromImpactNormalZ = (HitResult.ImpactNormal.Z - 1);
+				double OffsetFromImpactNormalZ = (HitResult.ImpactNormal.Z - 1.0f);
 				OffsetFromImpactNormalZ -= (OffsetFromImpactNormalZ * VREd::TeleportOffsetMultiplier->GetFloat());
 				OffsetDistance.Z = OffsetFromImpactNormalZ * (((VRMode->GetHeadTransform().GetLocation().Z - VRMode->GetRoomTransform().GetLocation().Z) / VRMode->GetWorldScaleFactor())  * TeleportGoalScale);
 
@@ -551,7 +551,7 @@ float AVREditorTeleporter::GetSlideDelta_Implementation(UVREditorInteractor* Int
 	float Result = 0.0f;
 	if (FMath::Abs(SlideDelta[Axis]) > FMath::Abs(SlideDelta[!Axis]))
 	{
-		Result = SlideDelta[Axis];
+		Result = (float) SlideDelta[Axis];
 	}
 
 	return Result;
diff --git a/Engine/Source/Editor/VREditor/Private/UI/VREditorDockableWindow.cpp b/Engine/Source/Editor/VREditor/Private/UI/VREditorDockableWindow.cpp
--- a/Engine/Source/Editor/VREditor/Private/UI/VREditorDockableWindow.cpp
+++ b/Engine/Source/Editor/VREditor/Private/UI/VREditorDockableWindow.cpp
@@ -185,8 +185,8 @@ void AVREditorDockableWindow::TickManually( float DeltaTime )
 
 			const FVector UICapsuleStart = FVector( 0.0f, 0.0f, -Size.Y * 0.4f ) * CurrentScaleFactor * AnimatedScale;
 			const FVector UICapsuleEnd = FVector( 0.0f, 0.0f, Size.Y * 0.4f ) * CurrentScaleFactor * AnimatedScale;
-			const float UICapsuleLocalRadius = Size.X * 0.5f * CurrentScaleFactor * AnimatedScale.Y;
-			const float MinDistanceToUICapsule = 10.0f * CurrentScaleFactor * AnimatedScale.Y;	// @todo vreditor tweak
+			const float UICapsuleLocalRadius = (float)Size.X * 0.5f * CurrentScaleFactor * (float)AnimatedScale.Y;
+			const float MinDistanceToUICapsule = 10.0f * CurrentScaleFactor * (float)AnimatedScale.Y;	// @todo vreditor tweak
 			const FVector UIForwardVector = FVector::ForwardVector;
 			const float MinDotForAimingAtOtherHand = -1.1f;	// @todo vreditor tweak
 
@@ -391,7 +391,7 @@ void AVREditorDockableWindow::OnPressed( UViewportInteractor* Interactor, const
 		else if(InHitResult.Component == GetSelectionBarMeshComponent() && !GetOwner().IsDraggingPanelFromOpen())
 		{
 			bOutResultedInDrag = true;
-			SetDockSelectDistance((InHitResult.TraceStart - InHitResult.Location ).Size());
+			SetDockSelectDistance((float)(InHitResult.TraceStart - InHitResult.Location ).Size());
 			GetOwner().StartDraggingDockUI(this, VREditorInteractor, DockSelectDistance);
 		}
 	}
diff --git a/Engine/Source/Editor/VREditor/Private/UI/VREditorRadialFloatingUI.cpp b/Engine/Source/Editor/VREditor/Private/UI/VREditorRadialFloatingUI.cpp
--- a/Engine/Source/Editor/VREditor/Private/UI/VREditorRadialFloatingUI.cpp
+++ b/Engine/Source/Editor/VREditor/Private/UI/VREditorRadialFloatingUI.cpp
@@ -459,7 +459,7 @@ const void AVREditorRadialFloatingUI::HighlightSlot(const FVector2D& TrackpadPos
 	}
 
 	const float AnglePerItem = 360.0f / NumberOfEntries;
-	float Angle = FRotator::NormalizeAxis(FMath::RadiansToDegrees(FMath::Atan2(TrackpadPosition.X, TrackpadPosition.Y)));
+	float Angle = (float) FRotator::NormalizeAxis(FMath::RadiansToDegrees(FMath::Atan2(TrackpadPosition.X, TrackpadPosition.Y)));
 
 	// first element of the menu is at 90
 	float ArrowAngle = Angle - 90.0f;
@@ -483,7 +483,7 @@ const void AVREditorRadialFloatingUI::HighlightSlot(const FVector2D& TrackpadPos
 		Angle = Angle - 360.0f;
 	}
 	ArrowMeshComponent->SetRelativeRotation(FRotator(0.0f, ArrowAngle, 0.0f).Quaternion());
-	float NewArrowScaleFactor = TrackpadPosition.Size();
+	float NewArrowScaleFactor = (float) TrackpadPosition.Size();
 	if (TrackpadPosition.GetAbsMax() > VREd::MinJoystickOffsetBeforeRadialMenu->GetFloat())
 	{
 		ArrowMeshComponent->SetVisibility(true);
@@ -502,7 +502,7 @@ const void AVREditorRadialFloatingUI::HighlightSlot(const FVector2D& TrackpadPos
 	const FPointerEvent& SimulatedPointer = FPointerEvent();
 	const FGeometry& ChildGeometry = FGeometry();
 	TSharedRef<SWidget> TestWidget = SNullWidget::NullWidget;
-	const int32 Index = (Angle / AnglePerItem);
+	const int32 Index = (int32) (Angle / AnglePerItem);
 	// Make sure we are checking for a valid radial menu widget
 	if (WidgetComponents.IsValidIndex(Index) && WidgetComponents[Index]->GetSlateWidget())
 	{
diff --git a/Engine/Source/Editor/VREditor/Private/UI/VREditorUISystem.cpp b/Engine/Source/Editor/VREditor/Private/UI/VREditorUISystem.cpp
--- a/Engine/Source/Editor/VREditor/Private/UI/VREditorUISystem.cpp
+++ b/Engine/Source/Editor/VREditor/Private/UI/VREditorUISystem.cpp
@@ -235,7 +235,7 @@ void UVREditorUISystem::Shutdown()
 		AVREditorFloatingUI* SequencerPanel = GetPanel(SequencerPanelID);
 		if (SequencerPanel != nullptr)
 		{
-			SequencerPanel->SetSlateWidget(*this, SequencerPanelID, SNullWidget::NullWidget, FIntPoint(VREd::SequencerUIResolutionX->GetFloat(), VREd::SequencerUIResolutionY->GetFloat()), VREd::EditorUISize->GetFloat(), AVREditorFloatingUI::EDockedTo::Nothing);
+			SequencerPanel->SetSlateWidget(*this, SequencerPanelID, SNullWidget::NullWidget, FIntPoint((int32)VREd::SequencerUIResolutionX->GetFloat(), (int32)VREd::SequencerUIResolutionY->GetFloat()), VREd::EditorUISize->GetFloat(), AVREditorFloatingUI::EDockedTo::Nothing);
 		}
 		FVREditorActionCallbacks::CloseSequencer(GetOwner().GetCurrentSequencer()->GetRootMovieSceneSequence());
 	}
@@ -671,12 +671,12 @@ void UVREditorUISystem::OnVRHoverUpdate(UViewportInteractor* Interactor, FVector
 									if( bIsAbsolute )
 									{
 										const float ScrollSpeed = VREd::UIAbsoluteScrollSpeed->GetFloat();
-										ScrollDelta = ( VREditorInteractor->GetTrackpadPosition().Y - VREditorInteractor->GetLastTrackpadPosition().Y ) * ScrollSpeed;
+										ScrollDelta = (float) ( VREditorInteractor->GetTrackpadPosition().Y - VREditorInteractor->GetLastTrackpadPosition().Y ) * ScrollSpeed;
 									}
 									else
 									{
 										const float ScrollSpeed = VREd::UIRelativeScrollSpeed->GetFloat();
-										ScrollDelta = VREditorInteractor->GetTrackpadPosition().Y * ScrollSpeed;
+										ScrollDelta = (float) VREditorInteractor->GetTrackpadPosition().Y * ScrollSpeed;
 									}
 								}
 
@@ -737,7 +737,7 @@ void UVREditorUISystem::OnVRHoverUpdate(UViewportInteractor* Interactor, FVector
 									FVector ScrollVelocityVector( VREditorInteractor->GetUIScrollVelocity(), 0.0f, 0.0f );
 									const bool bVelocitySensitive = false;
 									GetOwner().GetWorldInteraction().ApplyVelocityDamping( ScrollVelocityVector, bVelocitySensitive );
-									VREditorInteractor->SetUIScrollVelocity( ScrollVelocityVector.X );
+									VREditorInteractor->SetUIScrollVelocity( (float)ScrollVelocityVector.X );
 								}
 								else
 								{
@@ -924,7 +924,7 @@ void UVREditorUISystem::CreateUIs()
 
 			const bool bWithSceneComponent = false;
 			AVREditorFloatingUI* SequencerUI = GetOwner().SpawnTransientSceneActor<AVREditorDockableWindow>(TEXT("SequencerUI"), bWithSceneComponent);
-			SequencerUI->SetSlateWidget(*this, SequencerPanelID, WidgetToDraw, FIntPoint(VREd::SequencerUIResolutionX->GetFloat(), VREd::SequencerUIResolutionY->GetFloat()), VREd::EditorUISize->GetFloat(), AVREditorFloatingUI::EDockedTo::Nothing);
+			SequencerUI->SetSlateWidget(*this, SequencerPanelID, WidgetToDraw, FIntPoint((int32)VREd::SequencerUIResolutionX->GetFloat(), (int32)VREd::SequencerUIResolutionY->GetFloat()), VREd::EditorUISize->GetFloat(), AVREditorFloatingUI::EDockedTo::Nothing);
 			SequencerUI->ShowUI(false);
 			FloatingUIs.Add(SequencerPanelID, SequencerUI);
 		}
@@ -1096,7 +1096,7 @@ void UVREditorUISystem::ShowEditorUIPanel(AVREditorFloatingUI* Panel, UVREditorI
 			{
 				if (GetOwner().GetCurrentSequencer() != nullptr)
 				{
-					Panel->SetSlateWidget(*this, SequencerPanelID, SNullWidget::NullWidget, FIntPoint(VREd::SequencerUIResolutionX->GetFloat(), VREd::SequencerUIResolutionY->GetFloat()), VREd::EditorUISize->GetFloat(), AVREditorFloatingUI::EDockedTo::Nothing);
+					Panel->SetSlateWidget(*this, SequencerPanelID, SNullWidget::NullWidget, FIntPoint((int32)VREd::SequencerUIResolutionX->GetFloat(), (int32)VREd::SequencerUIResolutionY->GetFloat()), VREd::EditorUISize->GetFloat(), AVREditorFloatingUI::EDockedTo::Nothing);
 					FVREditorActionCallbacks::CloseSequencer(GetOwner().GetCurrentSequencer()->GetRootMovieSceneSequence());
 				}
 			}
@@ -1440,7 +1440,7 @@ void UVREditorUISystem::MakeUniformGridMenu(const TSharedRef<FMultiBox>& MultiBo
 		{
 			TSharedRef<SWidget> BlockWidget = Block->MakeWidget(MultiBoxWidget, EMultiBlockLocation::Middle, true, nullptr)->AsWidget();
 		
-			int32 Column = FMath::Fmod(static_cast<float>(Index), static_cast<float>(Columns));
+			int32 Column = (int32) FMath::Fmod(static_cast<float>(Index), static_cast<float>(Columns));
 			Row = FMath::DivideAndRoundDown(Index, Columns);
 			TSharedRef<SOverlay> TestOverlay = SNew(SOverlay);
 			GridPanel->AddSlot(Column, Row)
@@ -1561,7 +1561,7 @@ void UVREditorUISystem::MakeRadialBoxMenu(const TSharedRef<FMultiBox>& MultiBox,
 	}
 	MultiBox.Get().SetStyle(&FVREditorStyle::Get(), StyleName);
 
-	QuickRadialMenu->SetNumberOfEntries(NumItems);
+	QuickRadialMenu->SetNumberOfEntries((float)NumItems);
 
 	for (const TSharedRef<const FMultiBlock>& MultiBlock : Blocks)
 	{
@@ -1949,7 +1949,7 @@ void UVREditorUISystem::UpdateSequencerUI()
 				];
 
 			const bool bWithSceneComponent = false;
-			SequencerPanel->SetSlateWidget(*this, SequencerPanelID, WidgetToDraw, FIntPoint(VREd::SequencerUIResolutionX->GetFloat(), VREd::SequencerUIResolutionY->GetFloat()), VREd::EditorUISize->GetFloat(), AVREditorFloatingUI::EDockedTo::Nothing);
+			SequencerPanel->SetSlateWidget(*this, SequencerPanelID, WidgetToDraw, FIntPoint((int32)VREd::SequencerUIResolutionX->GetFloat(), (int32)VREd::SequencerUIResolutionY->GetFloat()), VREd::EditorUISize->GetFloat(), AVREditorFloatingUI::EDockedTo::Nothing);
 	
 			if (bSequencerOpenedFromRadialMenu == true)
 			{
@@ -2016,7 +2016,7 @@ void UVREditorUISystem::UpdateActorPreviewUI( TSharedRef<SWidget> InWidget, int3
 		return;
 	}
 
-	const FIntPoint ViewfinderSize = FIntPoint(VREd::DefaultCameraUIResolutionX->GetFloat(), VREd::DefaultCameraUIResolutionY->GetFloat());
+	const FIntPoint ViewfinderSize = FIntPoint((int32)VREd::DefaultCameraUIResolutionX->GetFloat(), (int32)VREd::DefaultCameraUIResolutionY->GetFloat());
 
 	// If a widget is sent but no panel exists, create one.
 	if (Actor && PreviewPanel == nullptr)
@@ -2092,7 +2092,7 @@ void UVREditorUISystem::UpdateDetachedActorPreviewUI(TSharedRef<SWidget> InWidge
 		VRMode->DestroyTransientActor(PreviewPanel);
 		return;
 	}
-	FIntPoint ViewfinderSize = FIntPoint(VREd::DefaultCameraUIResolutionX->GetFloat(), VREd::DefaultCameraUIResolutionY->GetFloat());
+	FIntPoint ViewfinderSize = FIntPoint((int32)VREd::DefaultCameraUIResolutionX->GetFloat(), (int32)VREd::DefaultCameraUIResolutionY->GetFloat());
 
 	// If a widget is sent but no panel exists, create one.
 	if (PreviewPanel == nullptr)
@@ -2178,7 +2178,7 @@ void UVREditorUISystem::UpdateExternalUMGUI(const FVREditorFloatingUICreationCon
 		}
 				
 		FloatingUIs.Add(PreviewPanelID, ExternalPanel);
-		ExternalPanel->SetUMGWidget(*this, PreviewPanelID, nullptr, FIntPoint(VREd::DefaultEditorUIResolutionX->GetFloat(), VREd::DefaultEditorUIResolutionY->GetFloat()), VREd::EditorUISize->GetFloat(), AVREditorFloatingUI::EDockedTo::Nothing);
+		ExternalPanel->SetUMGWidget(*this, PreviewPanelID, nullptr, FIntPoint((int32)VREd::DefaultEditorUIResolutionX->GetFloat(), (int32)VREd::DefaultEditorUIResolutionY->GetFloat()), VREd::EditorUISize->GetFloat(), AVREditorFloatingUI::EDockedTo::Nothing);
 		ExternalPanel->ShowUI(true);
 	}
 
@@ -2188,10 +2188,10 @@ void UVREditorUISystem::UpdateExternalUMGUI(const FVREditorFloatingUICreationCon
 		UWorld* World = GetOwner().GetWorld();
 		if (World != NULL)
 		{
-			FIntPoint UISize = FIntPoint(VREd::DefaultEditorUIResolutionX->GetFloat(), VREd::DefaultEditorUIResolutionY->GetFloat());
+			FIntPoint UISize = FIntPoint((int32)VREd::DefaultEditorUIResolutionX->GetFloat(), (int32)VREd::DefaultEditorUIResolutionY->GetFloat());
 			if (CreationContext.PanelSize != FVector2D::ZeroVector) 
 			{
-				UISize = FIntPoint(CreationContext.PanelSize.X, CreationContext.PanelSize.Y);
+				UISize = FIntPoint((int32)CreationContext.PanelSize.X, (int32)CreationContext.PanelSize.Y);
 			}
 			
 			// @todo: Currently passes CreationContext params into SetUMGWidget(). AVREditorFloatingUI has the context and should use it instead (deprecate the old individual vars)
@@ -2229,7 +2229,7 @@ void UVREditorUISystem::UpdateExternalSlateUI(TSharedRef<SWidget> InWidget, FNam
 		const bool bWithSceneComponent = false;
 		ExternalPanel = GetOwner().SpawnTransientSceneActor<AVREditorDockableWindow>(TEXT("ActorPreviewUI"), bWithSceneComponent);
 		FloatingUIs.Add(PreviewPanelID, ExternalPanel);
-		ExternalPanel->SetSlateWidget(*this, PreviewPanelID, SNullWidget::NullWidget, FIntPoint(VREd::DefaultEditorUIResolutionX->GetFloat(), VREd::DefaultEditorUIResolutionY->GetFloat()), VREd::EditorUISize->GetFloat(), AVREditorFloatingUI::EDockedTo::Nothing);
+		ExternalPanel->SetSlateWidget(*this, PreviewPanelID, SNullWidget::NullWidget, FIntPoint((int32)VREd::DefaultEditorUIResolutionX->GetFloat(), (int32)VREd::DefaultEditorUIResolutionY->GetFloat()), VREd::EditorUISize->GetFloat(), AVREditorFloatingUI::EDockedTo::Nothing);
 		ExternalPanel->ShowUI(true);
 	}
 
@@ -2239,10 +2239,10 @@ void UVREditorUISystem::UpdateExternalSlateUI(TSharedRef<SWidget> InWidget, FNam
 		UWorld* World = GetOwner().GetWorld();
 		if (World != NULL)
 		{
-			FIntPoint UISize = FIntPoint(VREd::DefaultEditorUIResolutionX->GetFloat(), VREd::DefaultEditorUIResolutionY->GetFloat());
+			FIntPoint UISize = FIntPoint((int32)VREd::DefaultEditorUIResolutionX->GetFloat(), (int32)VREd::DefaultEditorUIResolutionY->GetFloat());
 			if (InSize != FVector2D::ZeroVector)
 			{
-				UISize = FIntPoint(InSize.X, InSize.Y);
+				UISize = FIntPoint((int32)InSize.X, (int32)InSize.Y);
 			}
 			ExternalPanel->SetSlateWidget(*this, PreviewPanelID, InWidget, UISize, VREd::EditorUISize->GetFloat(), AVREditorFloatingUI::EDockedTo::Nothing);
 			const bool bDragFromOpen = ShouldPreviewPanel();
diff --git a/Engine/Source/Editor/VREditor/Private/VREditorActions.cpp b/Engine/Source/Editor/VREditor/Private/VREditorActions.cpp
--- a/Engine/Source/Editor/VREditor/Private/VREditorActions.cpp
+++ b/Engine/Source/Editor/VREditor/Private/VREditorActions.cpp
@@ -273,7 +273,7 @@ void FVREditorActionCallbacks::SimulateKeyDown(const FKey Key, const bool bRepea
 
 	if (CharCodePtr)
 	{
-		FCharacterEvent CharacterEvent( CharCode, FModifierKeysState(), 0, bRepeat );
+		FCharacterEvent CharacterEvent( (TCHAR)CharCode, FModifierKeysState(), 0, bRepeat );
 		FSlateApplication::Get().ProcessKeyCharEvent( CharacterEvent );
 	}
 }
diff --git a/Engine/Source/Editor/VREditor/Private/VREditorAvatarActor.cpp b/Engine/Source/Editor/VREditor/Private/VREditorAvatarActor.cpp
--- a/Engine/Source/Editor/VREditor/Private/VREditorAvatarActor.cpp
+++ b/Engine/Source/Editor/VREditor/Private/VREditorAvatarActor.cpp
@@ -313,7 +313,7 @@ void AVREditorAvatarActor::TickManually( const float DeltaTime )
 			const FVector ScaleIndicatorEndPosition = RightHandInteractor->GetTransform( ).GetLocation( );
 			const FVector ScaleIndicatorDirection = ( ScaleIndicatorEndPosition - ScaleIndicatorStartPosition ).GetSafeNormal( );
 
-			const float ScaleIndicatorLength = FMath::Max( 0.000001f, ( ScaleIndicatorEndPosition - ScaleIndicatorStartPosition ).Size( ) );
+			const float ScaleIndicatorLength = FMath::Max( 0.000001f, (float)( ScaleIndicatorEndPosition - ScaleIndicatorStartPosition ).Size() );
 			const float Radius = ( VREd::ScaleProgressBarRadius->GetFloat( ) * WorldScaleFactor );
 			const float ProgressbarLength = VREd::ScaleProgressBarLength->GetFloat( );
 			const float Scale = WorldScaleFactor;
@@ -385,11 +385,11 @@ void AVREditorAvatarActor::TickManually( const float DeltaTime )
 				}
 				else if ( ( LeftHandInteractor->GetLastDraggingMode() == EViewportInteractionDraggingMode::World && !LeftHandInteractor->GetDragTranslationVelocity().IsNearlyZero( VREd::GridMovementTolerance->GetFloat() ) ) )
 				{
-					GoalOpacity = FMath::Clamp( LeftHandInteractor->GetDragTranslationVelocity().Size() / VREd::GridFadeStartVelocity->GetFloat(), 0.0f, 1.0f );
+					GoalOpacity = FMath::Clamp( (float)LeftHandInteractor->GetDragTranslationVelocity().Size() / VREd::GridFadeStartVelocity->GetFloat(), 0.0f, 1.0f );
 				}
 				else if ( ( RightHandInteractor->GetLastDraggingMode() == EViewportInteractionDraggingMode::World && !RightHandInteractor->GetDragTranslationVelocity().IsNearlyZero( VREd::GridMovementTolerance->GetFloat() ) ) )
 				{
-					GoalOpacity = FMath::Clamp( RightHandInteractor->GetDragTranslationVelocity().Size() / VREd::GridFadeStartVelocity->GetFloat(), 0.0f, 1.0f );
+					GoalOpacity = FMath::Clamp( (float)RightHandInteractor->GetDragTranslationVelocity().Size() / VREd::GridFadeStartVelocity->GetFloat(), 0.0f, 1.0f );
 				}
 
 				// Check the current opacity and add or subtract to reach the goal
diff --git a/Engine/Source/Editor/VREditor/Private/VREditorFloatingText.cpp b/Engine/Source/Editor/VREditor/Private/VREditorFloatingText.cpp
--- a/Engine/Source/Editor/VREditor/Private/VREditorFloatingText.cpp
+++ b/Engine/Source/Editor/VREditor/Private/VREditorFloatingText.cpp
@@ -193,7 +193,7 @@ void AFloatingText::Update( const FVector OrientateToward )
 	// @todo vreditor tweak
 	const float LineRadius = 0.1f;
 	const float FirstLineLength = 4.0f;	   // Default line length (note that socket scale can affect this!)
-	const float SecondLineLength = TextComponent->GetTextLocalSize().Y;	// The second line "underlines" the text
+	const float SecondLineLength = (float) TextComponent->GetTextLocalSize().Y;	// The second line "underlines" the text
 
 
 	// NOTE: The origin of the actor will be the designated target of the text
diff --git a/Engine/Source/Editor/VREditor/Private/VREditorInteractor.cpp b/Engine/Source/Editor/VREditor/Private/VREditorInteractor.cpp
--- a/Engine/Source/Editor/VREditor/Private/VREditorInteractor.cpp
+++ b/Engine/Source/Editor/VREditor/Private/VREditorInteractor.cpp
@@ -919,8 +919,8 @@ void UVREditorInteractor::PreviewInputKey( class FEditorViewportClient& Viewport
 
 			// Detect swipe on trackpad.
 			const FVector2D SwipeDelta = LastTrackpadPosition - InitialTouchPosition;
-			const float AbsSwipeDeltaX = FMath::Abs( SwipeDelta.X );
-			const float AbsSwipeDeltaY = FMath::Abs( SwipeDelta.Y );
+			const float AbsSwipeDeltaX = FMath::Abs( (float)SwipeDelta.X );
+			const float AbsSwipeDeltaY = FMath::Abs( (float)SwipeDelta.Y );
 			if (!FMath::IsNearlyZero( SwipeDelta.X, 1.0f ) && AbsSwipeDeltaX > AbsSwipeDeltaY)
 			{
 				if (SwipeDelta.X > 0)
@@ -1107,8 +1107,8 @@ void UVREditorInteractor::HandleInputAxis( FEditorViewportClient& ViewportClient
 				Right,
 			};
 
-			const float Magnitude = TrackpadPosition.Size();
-			const float Angle = FMath::Atan2(TrackpadPosition.Y, TrackpadPosition.X);
+			const float Magnitude = (float) TrackpadPosition.Size();
+			const float Angle = FMath::Atan2((float)TrackpadPosition.Y, (float)TrackpadPosition.X);
 
 			auto IsDpadDirectionPressed = [&](EDpadDirection Dir) -> bool
 			{
@@ -1251,11 +1251,11 @@ float UVREditorInteractor::GetTrackpadSlideDelta( const bool Axis /*= 1*/ ) cons
 	{
 		if (bIsAbsolute)
 		{
-			SlideDelta = ((TrackpadPosition[Axis] - LastTrackpadPosition[Axis]) * VREd::TrackpadAbsoluteDragSpeed->GetFloat());
+			SlideDelta = (float) ((TrackpadPosition[Axis] - LastTrackpadPosition[Axis]) * VREd::TrackpadAbsoluteDragSpeed->GetFloat());
 		}
 		else
 		{
-			SlideDelta = (TrackpadPosition[Axis] * VREd::TrackpadRelativeDragSpeed->GetFloat());
+			SlideDelta = (float) (TrackpadPosition[Axis] * VREd::TrackpadRelativeDragSpeed->GetFloat());
 		}
 	}
 
@@ -1490,7 +1490,7 @@ void UVREditorInteractor::UpdateHelpLabels()
 	const FTransform HeadTransform = GetVRMode().GetHeadTransform();
 
 	// Only show help labels if the hand is pretty close to the face
-	const float DistanceToHead = (GetTransform().GetLocation() - HeadTransform.GetLocation()).Size();
+	const float DistanceToHead = (float) (GetTransform().GetLocation() - HeadTransform.GetLocation()).Size();
 	const float MinDistanceToHeadForHelp = VREd::HelpLabelFadeDistance->GetFloat() * GetVRMode().GetWorldScaleFactor();	// (in cm)
 	bool bShowHelp = VREd::ShowControllerHelpLabels->GetInt() != 0 && DistanceToHead <= MinDistanceToHeadForHelp;
 
@@ -1665,7 +1665,7 @@ void UVREditorInteractor::UpdateSplineLaser( const FVector& InStartLocation, con
 		LaserSplineComponent->ClearSplinePoints( true );
 
 		const FVector SmoothLaserDirection = InEndLocation - InStartLocation;
-		float Distance = SmoothLaserDirection.Size();
+		float Distance = (float) SmoothLaserDirection.Size();
 		const FVector StraightLaserEndLocation = InStartLocation + (InForward * Distance);
 		const int32 NumLaserSplinePoints = LaserSplineMeshComponents.Num();
 
@@ -1756,7 +1756,7 @@ void UVREditorInteractor::UpdateRadialMenuInput( const float DeltaTime )
 				const FVector2D ReturnToCenter = FVector2D::ZeroVector;
 				UISystem.GetRadialMenuFloatingUI()->HighlightSlot( ReturnToCenter );
 
-				const float NewPlayRate = FMath::GetMappedRangeValueClamped( FVector2D( -1.0f, 1.0f ), FVector2D( -1.0f*VREd::SequencerScrubMax->GetFloat(), VREd::SequencerScrubMax->GetFloat() ), TrackpadPosition.X );
+				const float NewPlayRate = (float) FMath::GetMappedRangeValueClamped( FVector2D( -1.0f, 1.0f ), FVector2D( -1.0f*VREd::SequencerScrubMax->GetFloat(), VREd::SequencerScrubMax->GetFloat() ), TrackpadPosition.X );
 				FVREditorActionCallbacks::PlaySequenceAtRate( VRMode, NewPlayRate );
 			}
 			else
diff --git a/Engine/Source/Editor/VREditor/Private/VREditorMode.cpp b/Engine/Source/Editor/VREditor/Private/VREditorMode.cpp
--- a/Engine/Source/Editor/VREditor/Private/VREditorMode.cpp
+++ b/Engine/Source/Editor/VREditor/Private/VREditorMode.cpp
@@ -628,11 +628,11 @@ void UVREditorMode::PostTick( float DeltaTime )
 		const float MaxLocationVelocity = VREd::HeadLocationMaxVelocity->GetFloat();	// cm/s
 		const float MaxRotationVelocity = VREd::HeadRotationMaxVelocity->GetFloat();	// degrees/s
 
-		const float LocationVelocity = ( LastRoomSpaceHeadToWorld.GetLocation() / LastWorldScaleFactor - RoomSpaceHeadToWorld.GetLocation() / WorldScaleFactor ).Size() / DeltaTime;
+		const float LocationVelocity = (float) ( LastRoomSpaceHeadToWorld.GetLocation() / LastWorldScaleFactor - RoomSpaceHeadToWorld.GetLocation() / WorldScaleFactor ).Size() / DeltaTime;
 
-		const float YawVelocity = FMath::Abs( FMath::FindDeltaAngleDegrees( LastRoomSpaceHeadToWorld.GetRotation().Rotator().Yaw, RoomSpaceHeadToWorld.GetRotation().Rotator().Yaw ) ) / DeltaTime;
-		const float PitchVelocity = FMath::Abs( FMath::FindDeltaAngleDegrees( LastRoomSpaceHeadToWorld.GetRotation().Rotator().Pitch, RoomSpaceHeadToWorld.GetRotation().Rotator().Pitch ) ) / DeltaTime;
-		const float RollVelocity = FMath::Abs( FMath::FindDeltaAngleDegrees( LastRoomSpaceHeadToWorld.GetRotation().Rotator().Roll, RoomSpaceHeadToWorld.GetRotation().Rotator().Roll ) ) / DeltaTime;
+		const float YawVelocity = (float) FMath::Abs( FMath::FindDeltaAngleDegrees( LastRoomSpaceHeadToWorld.GetRotation().Rotator().Yaw, RoomSpaceHeadToWorld.GetRotation().Rotator().Yaw ) ) / DeltaTime;
+		const float PitchVelocity = (float) FMath::Abs( FMath::FindDeltaAngleDegrees( LastRoomSpaceHeadToWorld.GetRotation().Rotator().Pitch, RoomSpaceHeadToWorld.GetRotation().Rotator().Pitch ) ) / DeltaTime;
+		const float RollVelocity = (float) FMath::Abs( FMath::FindDeltaAngleDegrees( LastRoomSpaceHeadToWorld.GetRotation().Rotator().Roll, RoomSpaceHeadToWorld.GetRotation().Rotator().Roll ) ) / DeltaTime;
 		const float RotationVelocity = YawVelocity + PitchVelocity + RollVelocity;
 
 		static float LastLocationVelocity = LocationVelocity;
@@ -904,7 +904,7 @@ bool UVREditorMode::IsHandAimingTowardsCapsule(UViewportInteractor* Interactor,
 
 		if( false )	// @todo vreditor debug
 		{
-			const float RenderCapsuleLength = ( CapsuleEnd - CapsuleStart ).Size() + CapsuleRadius * 2.0f;
+			const float RenderCapsuleLength = (float) ( CapsuleEnd - CapsuleStart ).Size() + CapsuleRadius * 2.0f;
 			// @todo vreditor:  This capsule draws with the wrong orientation
 			if( false )
 			{
@@ -927,7 +927,7 @@ bool UVREditorMode::IsHandAimingTowardsCapsule(UViewportInteractor* Interactor,
 				DrawDebugLine( GetWorld(), CapsuleTransform.TransformPosition( FVector::ZeroVector ), CapsuleTransform.TransformPosition( -LaserPointerDirectionInCapsuleSpace * 5.0f ), FColor::Purple, false, 0.0f );
 			}
 
-			const float Dot = FVector::DotProduct( CapsuleFrontDirection, -LaserPointerDirectionInCapsuleSpace );
+			const float Dot = (float) FVector::DotProduct( CapsuleFrontDirection, -LaserPointerDirectionInCapsuleSpace );
 			if( Dot >= MinDotForAimingAtCapsule )
 			{
 				bIsAimingTowards = true;
diff --git a/Engine/Source/Runtime/AugmentedReality/Private/ARBlueprintLibrary.cpp b/Engine/Source/Runtime/AugmentedReality/Private/ARBlueprintLibrary.cpp
--- a/Engine/Source/Runtime/AugmentedReality/Private/ARBlueprintLibrary.cpp
+++ b/Engine/Source/Runtime/AugmentedReality/Private/ARBlueprintLibrary.cpp
@@ -708,9 +708,9 @@ namespace UARBlueprintLibraryNamespace
 
 		FORCEINLINE void SetColumn(int32 Index, FVector Axis)
 		{
-			M[Index][0] = Axis.X;
-			M[Index][1] = Axis.Y;
-			M[Index][2] = Axis.Z;
+			M[Index][0] = (float) Axis.X;
+			M[Index][1] = (float) Axis.Y;
+			M[Index][2] = (float) Axis.Z;
 		};
 
 		FORCEINLINE FVector GetColumn(int32 Index) const
@@ -792,12 +792,12 @@ void UARBlueprintLibrary::CalculateClosestIntersection(const TArray<FVector>& St
 	{
 		FVector NormView = (EndPoints[PointIdx] - StartPoints[PointIdx]).GetSafeNormal();
 
-		float XX = (NormView.X * NormView.X) - 1.0;
-		float YY = (NormView.Y * NormView.Y) - 1.0;
-		float ZZ = (NormView.Z * NormView.Z) - 1.0;
-		float XY = (NormView.X * NormView.Y);
-		float XZ = (NormView.X * NormView.Z);
-		float YZ = (NormView.Y * NormView.Z);
+		float XX = ((float)NormView.X * (float)NormView.X) - 1.0f;
+		float YY = ((float)NormView.Y * (float)NormView.Y) - 1.0f;
+		float ZZ = ((float)NormView.Z * (float)NormView.Z) - 1.0f;
+		float XY = ((float)NormView.X * (float)NormView.Y);
+		float XZ = ((float)NormView.X * (float)NormView.Z);
+		float YZ = ((float)NormView.Y * (float)NormView.Z);
 
 		SXX += XX;
 		SYY += YY;
@@ -806,17 +806,17 @@ void UARBlueprintLibrary::CalculateClosestIntersection(const TArray<FVector>& St
 		SXZ += XZ;
 		SYZ += YZ;
 
-		CX += (StartPoints[PointIdx].X * XX) +
-			(StartPoints[PointIdx].Y * XY) +
-			(StartPoints[PointIdx].Z * XZ);
+		CX += ((float)StartPoints[PointIdx].X * XX) +
+			((float)StartPoints[PointIdx].Y * XY) +
+			((float)StartPoints[PointIdx].Z * XZ);
 
-		CY += (StartPoints[PointIdx].X * XY) +
-			(StartPoints[PointIdx].Y * YY) +
-			(StartPoints[PointIdx].Z * YZ);
+		CY += ((float)StartPoints[PointIdx].X * XY) +
+			((float)StartPoints[PointIdx].Y * YY) +
+			((float)StartPoints[PointIdx].Z * YZ);
 
-		CZ += (StartPoints[PointIdx].X * XZ) +
-			(StartPoints[PointIdx].Y * YZ) +
-			(StartPoints[PointIdx].Z * ZZ);
+		CZ += ((float)StartPoints[PointIdx].X * XZ) +
+			((float)StartPoints[PointIdx].Y * YZ) +
+			((float)StartPoints[PointIdx].Z * ZZ);
 	}
 	UARBlueprintLibraryNamespace::FSquareMatrix3 S;
 	S.SetColumn(0, { SXX, SXY, SXZ });
@@ -914,7 +914,7 @@ float UARBlueprintLibrary::GetARWorldScale()
 {
 	const auto AlignmentTransform = GetAlignmentTransform();
 	// Assuming the scale is uniform
-	return 1.f / AlignmentTransform.GetScale3D().X;
+	return 1.f / (float)AlignmentTransform.GetScale3D().X;
 }
 
 FTransform UARBlueprintLibrary::GetAlignmentTransform()
diff --git a/Engine/Source/Runtime/AugmentedReality/Private/ARTrackable.cpp b/Engine/Source/Runtime/AugmentedReality/Private/ARTrackable.cpp
--- a/Engine/Source/Runtime/AugmentedReality/Private/ARTrackable.cpp
+++ b/Engine/Source/Runtime/AugmentedReality/Private/ARTrackable.cpp
@@ -31,7 +31,7 @@ void UARTrackedGeometry::DebugDraw( UWorld* World, const FLinearColor& OutlineCo
 	const FVector Location = WorldTrans.GetLocation();
 	const FRotator Rotation = FRotator(WorldTrans.GetRotation());
 	const FVector Scale3D = WorldTrans.GetScale3D();
-	DrawDebugCoordinateSystem(World, Location, Rotation, Scale3D.X, true, PersistForSeconds, 0, OutlineThickness);
+	DrawDebugCoordinateSystem(World, Location, Rotation, (float)Scale3D.X, true, PersistForSeconds, 0, OutlineThickness);
 }
 
 void UARTrackedGeometry::GetNetworkPayload(FARMeshUpdatePayload& Payload)
@@ -114,7 +114,7 @@ void UARTrackedGeometry::SetName(const FString& InName)
 
 float UARTrackedGeometry::GetLastUpdateTimestamp() const
 {
-	return LastUpdateTimestamp;
+	return (float)LastUpdateTimestamp;
 }
 void UARTrackedGeometry::UpdateTrackedGeometryNoMove(const TSharedRef<FARSupportInterface, ESPMode::ThreadSafe>& InTrackingSystem, uint32 FrameNumber, double Timestamp)
 {
@@ -219,10 +219,10 @@ void UARPlaneGeometry::DebugDraw( UWorld* World, const FLinearColor& OutlineColo
 		for (int32 i=1; i<BoundaryPolygon.Num(); ++i)
 		{
 			const FVector NewVert = LocalToWorldTransform.TransformPosition(BoundaryPolygon[i]);
-			DrawDebugLine(World, LastVert, NewVert, OutlineRGB, PersistForSeconds > 0 ? true : false, PersistForSeconds, OutlineThickness);
+			DrawDebugLine(World, LastVert, NewVert, OutlineRGB, PersistForSeconds > 0 ? true : false, PersistForSeconds, 0, OutlineThickness);
 			LastVert = NewVert;
 		}
-		DrawDebugLine(World, LastVert, LocalToWorldTransform.TransformPosition(BoundaryPolygon[0]), OutlineRGB, PersistForSeconds > 0 ? true : false, PersistForSeconds, OutlineThickness);
+		DrawDebugLine(World, LastVert, LocalToWorldTransform.TransformPosition(BoundaryPolygon[0]), OutlineRGB, PersistForSeconds > 0 ? true : false, PersistForSeconds, 0, OutlineThickness);
 	}
 
 	const FVector WorldSpaceCenter = LocalToWorldTransform.TransformPosition(Center);
diff --git a/Engine/Source/Runtime/AugmentedReality/Public/ARTextures.h b/Engine/Source/Runtime/AugmentedReality/Public/ARTextures.h
--- a/Engine/Source/Runtime/AugmentedReality/Public/ARTextures.h
+++ b/Engine/Source/Runtime/AugmentedReality/Public/ARTextures.h
@@ -67,7 +67,7 @@ public:
 
 	/** The width and height of the texture */
 	UPROPERTY(BlueprintReadOnly, Category="AR AugmentedReality", meta=(Experimental))
-	FVector2D Size;
+	FVector2f Size;
 };
 
 /**
@@ -152,5 +152,5 @@ public:
 	
 	/** The width and height of the texture */
 	UPROPERTY(BlueprintReadOnly, Category="AR AugmentedReality", meta=(Experimental))
-	FVector2D Size;
+	FVector2f Size;
 };
diff --git a/Engine/Source/Runtime/MRMesh/Private/MRMeshComponent.cpp b/Engine/Source/Runtime/MRMesh/Private/MRMeshComponent.cpp
--- a/Engine/Source/Runtime/MRMesh/Private/MRMeshComponent.cpp
+++ b/Engine/Source/Runtime/MRMesh/Private/MRMeshComponent.cpp
@@ -474,7 +474,7 @@ private:
 						const auto& BrickBounds = BrickVisibility.Value.Key;
 						FColor BoundsColor(BrickVisibility.Value.Value ? FColor::Green : ColorGray);
 						FPrimitiveDrawInterface *PDI = Collector.GetPDI(ViewIndex);
-						DrawWireBox(PDI, BrickBounds, BoundsColor, GetDepthPriorityGroup(Views[ViewIndex]));
+						DrawWireBox(PDI, BrickBounds, BoundsColor, (uint8)GetDepthPriorityGroup(Views[ViewIndex]));
 					}
 				}
 			}
@@ -505,7 +505,7 @@ private:
 
 	uint32 GetAllocatedSize(void) const
 	{
-		return(FPrimitiveSceneProxy::GetAllocatedSize());
+		return((uint32)FPrimitiveSceneProxy::GetAllocatedSize());
 	}
 
 private:
diff --git a/Engine/Source/Runtime/MRMesh/Private/MockDataMeshTrackerComponent.cpp b/Engine/Source/Runtime/MRMesh/Private/MockDataMeshTrackerComponent.cpp
--- a/Engine/Source/Runtime/MRMesh/Private/MockDataMeshTrackerComponent.cpp
+++ b/Engine/Source/Runtime/MRMesh/Private/MockDataMeshTrackerComponent.cpp
@@ -618,7 +618,7 @@ void UMockDataMeshTrackerComponent::UpdateBlock(int32 BlockIndex)
 			// Hack because blueprints don't support uint32.
 			TArray<int32> Triangles(reinterpret_cast<const int32*>(CurrentMeshDataCache->
 				Triangles.GetData()), CurrentMeshDataCache->Triangles.Num());
-			OnMeshTrackerUpdated.Broadcast(CurrentMeshDataCache->BrickId, CurrentMeshDataCache->OffsetVertices,
+			OnMeshTrackerUpdated.Broadcast((int32)CurrentMeshDataCache->BrickId, CurrentMeshDataCache->OffsetVertices,
 				Triangles, CurrentMeshDataCache->Normals, CurrentMeshDataCache->Confidence);
 		}
 	}
