Virtual Scouting: Fix for crash dereferencing null RootComponent when opening menu.

The class in question doesn't actually load the asset container this check was ostensibly added to avoid, and it's not clear how this code ever worked.

diff --git a/Engine/Source/Editor/VREditor/Private/UI/VREditorFloatingUI.cpp b/Engine/Source/Editor/VREditor/Private/UI/VREditorFloatingUI.cpp
--- a/Engine/Source/Editor/VREditor/Private/UI/VREditorFloatingUI.cpp
+++ b/Engine/Source/Editor/VREditor/Private/UI/VREditorFloatingUI.cpp
@@ -29,11 +29,6 @@ AVREditorFloatingUI::AVREditorFloatingUI(const FObjectInitializer& ObjectInitial
 	UISystemID(NAME_None),
 	bClearWidgetOnHide(false)
 {
-	if (HasAnyFlags(RF_ClassDefaultObject))
-	{
-		return;
-	}
-
 	const bool bTransient = true;
 	USceneComponent* SceneComponent = CreateDefaultSubobject<USceneComponent>(TEXT("SceneComponent"), bTransient);
 	check(SceneComponent != nullptr);
