VREditor: Remove unused references to LevelEditor/Private.

#rnx

diff --git a/Engine/Source/Editor/VREditor/Private/UI/VRRadialMenuHandler.cpp b/Engine/Source/Editor/VREditor/Private/UI/VRRadialMenuHandler.cpp
--- a/Engine/Source/Editor/VREditor/Private/UI/VRRadialMenuHandler.cpp
+++ b/Engine/Source/Editor/VREditor/Private/UI/VRRadialMenuHandler.cpp
@@ -80,7 +80,6 @@
 #include "VREditorActions.h"
 #include "Framework/Commands/GenericCommands.h"
 #include "ISequencer.h"
-#include "SLevelEditor.h"
 #include "Kismet2/DebuggerCommands.h"
 #include "UMGStyle.h"
 
diff --git a/Engine/Source/Editor/VREditor/VREditor.Build.cs b/Engine/Source/Editor/VREditor/VREditor.Build.cs
--- a/Engine/Source/Editor/VREditor/VREditor.Build.cs
+++ b/Engine/Source/Editor/VREditor/VREditor.Build.cs
@@ -6,12 +6,6 @@ namespace UnrealBuildTool.Rules
 	{
         public VREditor(ReadOnlyTargetRules Target) : base(Target)
 		{
-			PrivateIncludePaths.AddRange(
-				new string[] {
-					System.IO.Path.Combine(GetModuleDirectory("LevelEditor"), "Private"),
-				}
-			);
-
             PrivateDependencyModuleNames.AddRange(
                 new string[] {
                     "AppFramework",
