VR Editor: Fixes a regression from CL 36864748 that led to FSlateRHIRenderer not correctly interoperating with the XR swapchain for editor viewports.

#rnx

diff --git a/Engine/Plugins/Experimental/XRCreativeFramework/Source/XRCreativeEditor/Private/XRCreativeVREditorMode.cpp b/Engine/Plugins/Experimental/XRCreativeFramework/Source/XRCreativeEditor/Private/XRCreativeVREditorMode.cpp
--- a/Engine/Plugins/Experimental/XRCreativeFramework/Source/XRCreativeEditor/Private/XRCreativeVREditorMode.cpp
+++ b/Engine/Plugins/Experimental/XRCreativeFramework/Source/XRCreativeEditor/Private/XRCreativeVREditorMode.cpp
@@ -211,12 +211,13 @@ void UXRCreativeVREditorMode::SetHeadTransform(const FTransform& HeadToWorld)
 }
 
 
+// Does not call the base class implementation, because we call UOpenXRInputFunctionLibrary::BeginXRSession
+// instead of GEngine->StereoRenderingDevice->EnableStereo(true).
 void UXRCreativeVREditorMode::EnableStereo()
 {
-	if (TSharedPtr<SLevelViewport> Viewport = GetVrLevelViewport())
+	if (TSharedPtr<SLevelViewport> Viewport = GetVrLevelViewport(); ensure(Viewport))
 	{
-		Viewport->EnableStereoRendering(true);
-		Viewport->SetRenderDirectlyToWindow(true);
+		StereoViewportSetup(Viewport.ToSharedRef());
 	}
 	
 	TSet<UInputMappingContext*> Contexts;
@@ -255,14 +256,15 @@ void UXRCreativeVREditorMode::EnableStereo()
 }
 
 
+// Does not call the base class implementation, because we call UOpenXRInputFunctionLibrary::EndXRSession
+// instead of GEngine->StereoRenderingDevice->EnableStereo(false).
 void UXRCreativeVREditorMode::DisableStereo()
 {
 	UOpenXRInputFunctionLibrary::EndXRSession();
 
-	if (TSharedPtr<SLevelViewport> VREditorLevelViewport = GetVrLevelViewport())
+	if (TSharedPtr<SLevelViewport> Viewport = GetVrLevelViewport(); ensure(Viewport))
 	{
-		VREditorLevelViewport->EnableStereoRendering(false);
-		VREditorLevelViewport->SetRenderDirectlyToWindow(false);
+		StereoViewportShutdown(Viewport.ToSharedRef());
 	}
 }
 
diff --git a/Engine/Source/Editor/VREditor/Private/VREditorModeBase.cpp b/Engine/Source/Editor/VREditor/Private/VREditorModeBase.cpp
--- a/Engine/Source/Editor/VREditor/Private/VREditorModeBase.cpp
+++ b/Engine/Source/Editor/VREditor/Private/VREditorModeBase.cpp
@@ -3,6 +3,7 @@
 #include "VREditorModeBase.h"
 #include "CameraController.h"
 #include "Engine/Engine.h"
+#include "Framework/Application/SlateApplication.h"
 #include "GameFramework/WorldSettings.h"
 #include "IHeadMountedDisplay.h"
 #include "ILevelEditor.h"
@@ -245,12 +246,53 @@ void UVREditorModeBase::CloseViewport( const bool bShouldDisableStereo )
 }
 
 
+void UVREditorModeBase::StereoViewportSetup(TSharedRef<SLevelViewport> Viewport)
+{
+	Viewport->EnableStereoRendering(true);
+	Viewport->SetRenderDirectlyToWindow(true);
+
+	if (TSharedPtr<SViewport> ViewportWidget = Viewport->GetViewportWidget().Pin(); ensure(ViewportWidget))
+	{
+		TSharedRef<SViewport> ViewportWidgetRef = ViewportWidget.ToSharedRef();
+		TSharedPtr<SWindow> ParentWindow = FSlateApplication::Get().FindWidgetWindow(ViewportWidgetRef);
+		TSharedPtr<ISlateViewport> SlateViewport = ViewportWidgetRef->GetViewportInterface().Pin();
+		if (ensure(ParentWindow) && ensure(SlateViewport))
+		{
+			// It is important to ensure that our containing SWindow is associated with our
+			// ISlateViewport, so FSlateRHIRenderer will invoke the UseSeparateRenderTarget path
+			// to correctly interoperate with the XR swapchain.
+			ParentWindow->SetViewport(SlateViewport.ToSharedRef());
+		}
+	}
+}
+
+
+void UVREditorModeBase::StereoViewportShutdown(TSharedRef<SLevelViewport> Viewport)
+{
+	Viewport->EnableStereoRendering(false);
+	Viewport->SetRenderDirectlyToWindow(false);
+
+	if (TSharedPtr<SViewport> ViewportWidget = Viewport->GetViewportWidget().Pin(); ensure(ViewportWidget))
+	{
+		TSharedRef<SViewport> ViewportWidgetRef = ViewportWidget.ToSharedRef();
+		TSharedPtr<SWindow> ParentWindow = FSlateApplication::Get().FindWidgetWindow(ViewportWidgetRef);
+		TSharedPtr<ISlateViewport> SlateViewport = ViewportWidgetRef->GetViewportInterface().Pin();
+		if (ensure(ParentWindow) && ensure(SlateViewport))
+		{
+			// It's also important to ensure that we remove the ISlateViewport association when
+			// we're done with stereo rendering, or else multiple SWindows might end up associated
+			// with the same viewport (e.g. undocked/floating/orphan viewport), causing artifacts.
+			ParentWindow->UnsetViewport(SlateViewport.ToSharedRef());
+		}
+	}
+}
+
+
 void UVREditorModeBase::EnableStereo()
 {
-	if (TSharedPtr<SLevelViewport> Viewport = GetVrLevelViewport())
+	if (TSharedPtr<SLevelViewport> Viewport = GetVrLevelViewport(); ensure(Viewport))
 	{
-		Viewport->EnableStereoRendering(true);
-		Viewport->SetRenderDirectlyToWindow(true);
+		StereoViewportSetup(Viewport.ToSharedRef());
 	}
 
 	if (GEngine->XRSystem)
@@ -267,9 +309,8 @@ void UVREditorModeBase::DisableStereo()
 		GEngine->StereoRenderingDevice->EnableStereo(false);
 	}
 
-	if (TSharedPtr<SLevelViewport> Viewport = GetVrLevelViewport())
+	if (TSharedPtr<SLevelViewport> Viewport = GetVrLevelViewport(); ensure(Viewport))
 	{
-		Viewport->EnableStereoRendering(false);
-		Viewport->SetRenderDirectlyToWindow(false);
+		StereoViewportShutdown(Viewport.ToSharedRef());
 	}
 }
diff --git a/Engine/Source/Editor/VREditor/Public/VREditorModeBase.h b/Engine/Source/Editor/VREditor/Public/VREditorModeBase.h
--- a/Engine/Source/Editor/VREditor/Public/VREditorModeBase.h
+++ b/Engine/Source/Editor/VREditor/Public/VREditorModeBase.h
@@ -74,6 +74,9 @@ protected:
 	/** Close the current viewport */
 	virtual void CloseViewport(bool bShouldDisableStereo);
 
+	virtual void StereoViewportSetup(TSharedRef<SLevelViewport> Viewport);
+	virtual void StereoViewportShutdown(TSharedRef<SLevelViewport> Viewport);
+
 	virtual void EnableStereo();
 	virtual void DisableStereo();
 
diff --git a/Engine/Source/Runtime/SlateCore/Public/Widgets/SWindow.h b/Engine/Source/Runtime/SlateCore/Public/Widgets/SWindow.h
--- a/Engine/Source/Runtime/SlateCore/Public/Widgets/SWindow.h
+++ b/Engine/Source/Runtime/SlateCore/Public/Widgets/SWindow.h
@@ -1057,6 +1057,14 @@ public:
 		Viewport = ViewportRef;
 	}
 
+	void UnsetViewport(TSharedRef<ISlateViewport> ViewportRef)
+	{
+		if (ensure(Viewport == ViewportRef))
+		{
+			Viewport.Reset();
+		}
+	}
+
 	TSharedPtr<ISlateViewport> GetViewport()
 	{
 		return Viewport.Pin();
