Update VREditorModeManager to use array based API of assetregistry broadcast
#rnx

diff --git a/Engine/Source/Editor/VREditor/Private/VREditorModeManager.cpp b/Engine/Source/Editor/VREditor/Private/VREditorModeManager.cpp
--- a/Engine/Source/Editor/VREditor/Private/VREditorModeManager.cpp
+++ b/Engine/Source/Editor/VREditor/Private/VREditorModeManager.cpp
@@ -358,12 +358,12 @@ void FVREditorModeManager::HandleAssetFilesLoaded()
 {
 	IAssetRegistry& AssetRegistry = FModuleManager::LoadModuleChecked<FAssetRegistryModule>(AssetRegistryConstants::ModuleName).Get();
 	AssetRegistry.OnFilesLoaded().RemoveAll(this);
-	AssetRegistry.OnAssetAdded().AddRaw(this, &FVREditorModeManager::HandleAssetAddedOrRemoved);
-	AssetRegistry.OnAssetRemoved().AddRaw(this, &FVREditorModeManager::HandleAssetAddedOrRemoved);
+	AssetRegistry.OnAssetsAdded().AddRaw(this, &FVREditorModeManager::HandleAssetAddedOrRemoved);
+	AssetRegistry.OnAssetsRemoved().AddRaw(this, &FVREditorModeManager::HandleAssetAddedOrRemoved);
 	UpdateCachedModeClassPaths();
 }
 
-void FVREditorModeManager::HandleAssetAddedOrRemoved(const FAssetData& NewAssetData)
+void FVREditorModeManager::HandleAssetAddedOrRemoved(TConstArrayView<FAssetData> NewAssetData)
 {
 	UpdateCachedModeClassPaths();
 }
diff --git a/Engine/Source/Editor/VREditor/Private/VREditorModeManager.h b/Engine/Source/Editor/VREditor/Private/VREditorModeManager.h
--- a/Engine/Source/Editor/VREditor/Private/VREditorModeManager.h
+++ b/Engine/Source/Editor/VREditor/Private/VREditorModeManager.h
@@ -87,7 +87,7 @@ private:
 	TSet<FTopLevelAssetPath> CachedModeClassPaths;
 
 	void HandleAssetFilesLoaded();
-	void HandleAssetAddedOrRemoved(const FAssetData& NewAssetData);
+	void HandleAssetAddedOrRemoved(TConstArrayView<FAssetData> NewAssetData);
 
 	/** Broadcasts when VR mode is started */
 	FOnVREditingModeEnter OnVREditingModeEnterHandle;
