VR Editor: Deprecate VR Editor mode and most associated classes.

Adds deprecation annotations to most types from public headers in the VREditor module.

Note that `UVREditorMode*Base*` is still used by XR Creative Framework, along with much of the `IVREditorModule` functionality, and a few fields in `UVRModeSettings`.

diff --git a/Engine/Plugins/Developer/Concert/ConcertSync/ConcertSyncClient/Source/ConcertSyncClient/Private/ConcertClientPresenceManager.cpp b/Engine/Plugins/Developer/Concert/ConcertSync/ConcertSyncClient/Source/ConcertSyncClient/Private/ConcertClientPresenceManager.cpp
--- a/Engine/Plugins/Developer/Concert/ConcertSync/ConcertSyncClient/Source/ConcertSyncClient/Private/ConcertClientPresenceManager.cpp
+++ b/Engine/Plugins/Developer/Concert/ConcertSync/ConcertSyncClient/Source/ConcertSyncClient/Private/ConcertClientPresenceManager.cpp
@@ -748,7 +748,7 @@ bool FConcertClientPresenceManager::IsJumpToPresenceEnabled(FGuid InEndpointId)
 	}
 
 	// Only enable the button if we have a valid perspective viewport to move and we're not in VR
-	if (!GetPerspectiveViewport() || IVREditorModule::Get().IsVREditorModeActive())
+	if (!GetPerspectiveViewport() || IVREditorModule::Get().GetVRModeBase())
 	{
 		return false;
 	}
@@ -967,10 +967,12 @@ void FConcertClientDefaultPresenceModeFactory::QueryVRDeviceType()
 {
 	if (UVREditorModeBase* VRModeBase = IVREditorModule::Get().GetVRModeBase())
 	{
+PRAGMA_DISABLE_DEPRECATION_WARNINGS
 		if (UVREditorMode* VRMode = Cast<UVREditorMode>(VRModeBase))
 		{
 			VRDeviceType = VRMode->GetHMDDeviceType();
 		}
+PRAGMA_ENABLE_DEPRECATION_WARNINGS
 		else
 		{
 			static const FName Generic("Generic");
diff --git a/Engine/Plugins/Experimental/VirtualProductionUtilities/Source/VPBookmarkEditor/Private/VPBookmarkTypeActions.cpp b/Engine/Plugins/Experimental/VirtualProductionUtilities/Source/VPBookmarkEditor/Private/VPBookmarkTypeActions.cpp
--- a/Engine/Plugins/Experimental/VirtualProductionUtilities/Source/VPBookmarkEditor/Private/VPBookmarkTypeActions.cpp
+++ b/Engine/Plugins/Experimental/VirtualProductionUtilities/Source/VPBookmarkEditor/Private/VPBookmarkTypeActions.cpp
@@ -64,6 +64,7 @@ void FVPBookmarkTypeActions::ActivateBookmark(UVPBookmark* InBookmark, FEditorVi
 		// Set Rotation
 		bool bIsRotationSet = false;
 		FRotator UseRotation = ActorRotation;
+PRAGMA_DISABLE_DEPRECATION_WARNINGS
 		if (IVREditorModule::IsAvailable())
 		{
 			IVREditorModule& VREditorModule = IVREditorModule::Get();
@@ -82,6 +83,7 @@ void FVPBookmarkTypeActions::ActivateBookmark(UVPBookmark* InBookmark, FEditorVi
 				}
 			}
 		}
+PRAGMA_ENABLE_DEPRECATION_WARNINGS
 
 		Transform.SetLocation(ActorLocation - Offset);
 
@@ -131,6 +133,8 @@ void FVPBookmarkTypeActions::InitFromViewport(UBookmarkBase* InBookmark, FEditor
 			BookmarkActorPtr->Modify();
 
 			FTransform HeadTransform = VPBookmark::GetEditorViewportTransform(InViewportClient);
+PRAGMA_DISABLE_DEPRECATION_WARNINGS
+			// TODO?: Do we want this logic transitioned to XR Creative?
 			if (IVREditorModule::IsAvailable())
 			{
 				IVREditorModule& VREditorModule = IVREditorModule::Get();
@@ -148,6 +152,7 @@ void FVPBookmarkTypeActions::InitFromViewport(UBookmarkBase* InBookmark, FEditor
 					}
 				}
 			}
+PRAGMA_ENABLE_DEPRECATION_WARNINGS
 
 			FRotator SpawnRotation = HeadTransform.Rotator();
 			VPBookmark->CachedViewportData.LookRotation = SpawnRotation;
diff --git a/Engine/Plugins/Experimental/VirtualProductionUtilities/Source/VPUtilities/Private/Libraries/VPBlueprintLibrary.cpp b/Engine/Plugins/Experimental/VirtualProductionUtilities/Source/VPUtilities/Private/Libraries/VPBlueprintLibrary.cpp
--- a/Engine/Plugins/Experimental/VirtualProductionUtilities/Source/VPUtilities/Private/Libraries/VPBlueprintLibrary.cpp
+++ b/Engine/Plugins/Experimental/VirtualProductionUtilities/Source/VPUtilities/Private/Libraries/VPBlueprintLibrary.cpp
@@ -182,7 +182,9 @@ bool UVPBlueprintLibrary::IsVREditorModeActive()
 #if WITH_EDITOR
 	if (IVREditorModule::IsAvailable())
 	{
+PRAGMA_DISABLE_DEPRECATION_WARNINGS
 		return IVREditorModule::Get().IsVREditorModeActive();
+PRAGMA_ENABLE_DEPRECATION_WARNINGS
 	}
 #endif
 
@@ -200,6 +202,7 @@ FVector UVPBlueprintLibrary::GetVREditorLaserHoverLocation()
 
 		for (UViewportInteractor* Interactor : Interactors)
 		{
+PRAGMA_DISABLE_DEPRECATION_WARNINGS
 			if (UVREditorInteractor* EdInteractor = Cast<UVREditorInteractor>(Interactor))
 			{
 				if (EdInteractor->GetControllerType() == EControllerType::Laser)
@@ -210,6 +213,7 @@ FVector UVPBlueprintLibrary::GetVREditorLaserHoverLocation()
 					return EdInteractor->GetInteractorData().LastHoverLocationOverUI;
 				}
 			}
+PRAGMA_ENABLE_DEPRECATION_WARNINGS
 		}
 	}
 #endif
@@ -227,6 +231,7 @@ FVector UVPBlueprintLibrary::GetLeftInteractorLocation()
 
 		for (UViewportInteractor* Interactor : Interactors)
 		{
+PRAGMA_DISABLE_DEPRECATION_WARNINGS
 			if (UVREditorInteractor* EdInteractor = Cast<UVREditorInteractor>(Interactor))
 			{
 				if (EdInteractor->GetControllerType() == EControllerType::Laser)
@@ -237,6 +242,7 @@ FVector UVPBlueprintLibrary::GetLeftInteractorLocation()
 					return InteractorLocation;
 				}
 			}
+PRAGMA_ENABLE_DEPRECATION_WARNINGS
 		}
 	}
 #endif
@@ -322,7 +328,9 @@ void UVPBlueprintLibrary::VPBookmarkSplineMeshIndicatorSetStartAndEnd(USplineMes
 		
 	// @todo: Fix - GetVREditorLaserHoverLocation() does not return the correct hover location
 	// USplineMeshComponent::SetEndPosition expects local space
+PRAGMA_DISABLE_DEPRECATION_WARNINGS // TODO: XR Creative Framework relevant?
 	SplineMesh->SetEndPosition(SplineTransform.InverseTransformPosition(IsVREditorModeActive() ? GetLeftInteractorLocation(): GetEditorViewportTransform().TransformPosition(FVector(80, 0, -20))));
+PRAGMA_ENABLE_DEPRECATION_WARNINGS
 }
 
 void UVPBlueprintLibrary::VPBookmarkSplineMeshIndicatorDisable(USplineMeshComponent* SplineMesh)
diff --git a/Engine/Plugins/Experimental/VirtualProductionUtilities/Source/VPUtilities/Public/Libraries/VPBlueprintLibrary.h b/Engine/Plugins/Experimental/VirtualProductionUtilities/Source/VPUtilities/Public/Libraries/VPBlueprintLibrary.h
--- a/Engine/Plugins/Experimental/VirtualProductionUtilities/Source/VPUtilities/Public/Libraries/VPBlueprintLibrary.h
+++ b/Engine/Plugins/Experimental/VirtualProductionUtilities/Source/VPUtilities/Public/Libraries/VPBlueprintLibrary.h
@@ -61,15 +61,18 @@ public:
 	static void SetGrabSpeed(const float Speed);
 
 	/** Get whether the user is in editor VR mode */
-	UFUNCTION(BlueprintCallable, Category = "Virtual Production")
+	UE_DEPRECATED(5.6, "VR Editor mode is deprecated; use the XR Creative Framework plugin instead.")
+	UFUNCTION(BlueprintCallable, Category = "Virtual Production", meta = (DeprecatedFunction, DeprecationMessage = "VR Editor mode is deprecated; use the XR Creative Framework plugin instead."))
 	static bool IsVREditorModeActive();
 
 	/** Get the hitlocation of the interaction controller's laser pointer, in world space */
-	UFUNCTION(BlueprintCallable, Category = "Virtual Production")
+	UE_DEPRECATED(5.6, "VR Editor mode is deprecated; use the XR Creative Framework plugin instead.")
+	UFUNCTION(BlueprintCallable, Category = "Virtual Production", meta = (DeprecatedFunction, DeprecationMessage = "VR Editor mode is deprecated; use the XR Creative Framework plugin instead."))
 	static FVector GetVREditorLaserHoverLocation();
 
 	/** Get the location of the left interactor, in world space */
-	UFUNCTION(BlueprintCallable, Category = "Virtual Production")
+	UE_DEPRECATED(5.6, "VR Editor mode is deprecated; use the XR Creative Framework plugin instead.")
+	UFUNCTION(BlueprintCallable, Category = "Virtual Production", meta = (DeprecatedFunction, DeprecationMessage = "VR Editor mode is deprecated; use the XR Creative Framework plugin instead."))
 	static FVector GetLeftInteractorLocation();
 
 	/** Trigger an UnrealEd Undo */
diff --git a/Engine/Plugins/MovieScene/LevelSequenceEditor/Source/LevelSequenceEditor/Private/LevelSequenceEditorToolkit.cpp b/Engine/Plugins/MovieScene/LevelSequenceEditor/Source/LevelSequenceEditor/Private/LevelSequenceEditorToolkit.cpp
--- a/Engine/Plugins/MovieScene/LevelSequenceEditor/Source/LevelSequenceEditor/Private/LevelSequenceEditorToolkit.cpp
+++ b/Engine/Plugins/MovieScene/LevelSequenceEditor/Source/LevelSequenceEditor/Private/LevelSequenceEditorToolkit.cpp
@@ -280,6 +280,7 @@ void FLevelSequenceEditorToolkit::Initialize(const EToolkitMode::Type Mode, cons
 
 	FLevelSequenceEditorToolkit::OnOpened().Broadcast(*this);
 
+PRAGMA_DISABLE_DEPRECATION_WARNINGS
 	{
 		UWorld* World = PlaybackContext->GetPlaybackContext()->GetWorld();
 		UVREditorMode* VRMode = Cast<UVREditorMode>( GEditor->GetEditorWorldExtensionsManager()->GetEditorWorldExtensions( World )->FindExtension( UVREditorMode::StaticClass() ) );
@@ -298,6 +299,7 @@ void FLevelSequenceEditorToolkit::Initialize(const EToolkitMode::Type Mode, cons
 			VRMode->RefreshVREditorSequencer(Sequencer.Get());
 		}
 	}
+PRAGMA_ENABLE_DEPRECATION_WARNINGS
 
 	// Extend the Level Editor viewport toolbar.
 	{
@@ -759,6 +761,7 @@ void FLevelSequenceEditorToolkit::HandleActorAddedToSequencer(AActor* Actor, con
 
 void FLevelSequenceEditorToolkit::HandleVREditorModeExit()
 {
+PRAGMA_DISABLE_DEPRECATION_WARNINGS
 	UWorld* World = PlaybackContext->GetPlaybackContext()->GetWorld();
 	UVREditorMode* VRMode = CastChecked<UVREditorMode>( GEditor->GetEditorWorldExtensionsManager()->GetEditorWorldExtensions( World )->FindExtension( UVREditorMode::StaticClass() ) );
 
@@ -766,6 +769,7 @@ void FLevelSequenceEditorToolkit::HandleVREditorModeExit()
 	Sequencer->SetAutoChangeMode(VRMode->GetSavedEditorState().AutoChangeMode);
 	Sequencer->SetKeyGroupMode(VRMode->GetSavedEditorState().bKeyAllEnabled ? EKeyGroupMode::KeyAll : EKeyGroupMode::KeyChanged);
 	VRMode->OnVREditingModeExit_Handler.Unbind();
+PRAGMA_ENABLE_DEPRECATION_WARNINGS
 }
 
 void FLevelSequenceEditorToolkit::HandleMapChanged(class UWorld* NewWorld, EMapChangeType MapChangeType)
@@ -993,12 +997,14 @@ void FLevelSequenceEditorToolkit::OnClose()
 	UToolMenus::Get()->UnregisterOwnerByName(UE::MovieScene::Private::ViewportToolbarOwnerName);
 
 	UWorld* World = PlaybackContext->GetPlaybackContext()->GetWorld();
+PRAGMA_DISABLE_DEPRECATION_WARNINGS
 	UVREditorMode* VRMode = Cast<UVREditorMode>(GEditor->GetEditorWorldExtensionsManager()->GetEditorWorldExtensions(World)->FindExtension(UVREditorMode::StaticClass()));
 	if (VRMode != nullptr)
 	{
 		// Null out the VR Mode's sequencer pointer
 		VRMode->RefreshVREditorSequencer(nullptr);
 	}
+PRAGMA_ENABLE_DEPRECATION_WARNINGS
 	OpenToolkits.Remove(this);
 
 	OnClosedEvent.Broadcast();
diff --git a/Engine/Plugins/VirtualProduction/VirtualScouting/Source/VirtualScoutingEditor/Private/VirtualScoutingEditorModule.cpp b/Engine/Plugins/VirtualProduction/VirtualScouting/Source/VirtualScoutingEditor/Private/VirtualScoutingEditorModule.cpp
--- a/Engine/Plugins/VirtualProduction/VirtualScouting/Source/VirtualScoutingEditor/Private/VirtualScoutingEditorModule.cpp
+++ b/Engine/Plugins/VirtualProduction/VirtualScouting/Source/VirtualScoutingEditor/Private/VirtualScoutingEditorModule.cpp
@@ -5,6 +5,9 @@
 #include "VRModeSettings.h"
 
 
+PRAGMA_DISABLE_DEPRECATION_WARNINGS
+
+
 DEFINE_LOG_CATEGORY(LogVirtualScoutingEditor);
 
 
@@ -29,3 +32,6 @@ private:
 
 
 IMPLEMENT_MODULE(FVirtualScoutingEditorModule, VirtualScoutingEditor);
+
+
+PRAGMA_ENABLE_DEPRECATION_WARNINGS
diff --git a/Engine/Source/Editor/LevelEditor/Private/SLevelViewport.cpp b/Engine/Source/Editor/LevelEditor/Private/SLevelViewport.cpp
--- a/Engine/Source/Editor/LevelEditor/Private/SLevelViewport.cpp
+++ b/Engine/Source/Editor/LevelEditor/Private/SLevelViewport.cpp
@@ -243,7 +243,9 @@ SLevelViewport::~SLevelViewport()
 	{
 		ActorPreview.bIsPinned = false;
 	}
+PRAGMA_DISABLE_DEPRECATION_WARNINGS
 	const bool bPreviewInDesktopViewport = !IVREditorModule::Get().IsVREditorModeActive();
+PRAGMA_ENABLE_DEPRECATION_WARNINGS
 	PreviewActors( TArray< AActor* >(), bPreviewInDesktopViewport);
 
 	FLevelViewportCommands::NewStatCommandDelegate.RemoveAll(this);
@@ -1319,7 +1321,9 @@ void SLevelViewport::Tick( const FGeometry& AllottedGeometry, const double InCur
 
 	if (bNeedToUpdatePreviews)
 	{
+PRAGMA_DISABLE_DEPRECATION_WARNINGS
 		const bool bPreviewInDesktopViewport = !IVREditorModule::Get().IsVREditorModeActive();
+PRAGMA_ENABLE_DEPRECATION_WARNINGS
 		if (GetDefault<ULevelEditorViewportSettings>()->bPreviewSelectedCameras && GCurrentLevelEditingViewportClient == LevelViewportClient.Get())
 		{
 			PreviewSelectedCameraActors(bPreviewInDesktopViewport);
@@ -3830,8 +3834,10 @@ void SActorPreview::Construct( const FArguments& InArgs )
 	EVisibility BorderVisibility = (InArgs._IsInteractive ? EVisibility::SelfHitTestInvisible : EVisibility::HitTestInvisible);
 	
 	//We draw certain buttons depending on whether we're in editor or VR mode
+PRAGMA_DISABLE_DEPRECATION_WARNINGS
 	EVisibility VRVisibility = IVREditorModule::Get().IsVREditorModeActive() ? EVisibility::Visible : EVisibility::Hidden;
 	EVisibility EditorVisibility = IVREditorModule::Get().IsVREditorModeActive() ? EVisibility::Hidden : EVisibility::Visible;
+PRAGMA_ENABLE_DEPRECATION_WARNINGS
 
 	this->ChildSlot
 	[
@@ -4507,7 +4513,9 @@ void SLevelViewport::PreviewActors( const TArray< AActor* >& InActorsToPreview,
 			IVREditorModule& VREditorModule = IVREditorModule::Get();
 			if (!bPreviewInDesktopViewport)
 			{
+PRAGMA_DISABLE_DEPRECATION_WARNINGS
 				VREditorModule.UpdateActorPreview( NewActorPreview.PreviewWidget.ToSharedRef(), ActorPreviews.Num()-1, CurActor);
+PRAGMA_ENABLE_DEPRECATION_WARNINGS
 			}
 			else
 			{
@@ -4578,10 +4586,12 @@ void SLevelViewport::ToggleActorPreviewIsPanelDetached(TWeakObjectPtr<AActor> Pr
 				{
 					ActorPreview.ToggleIsPanelDetached();
 					IVREditorModule& VREditorModule = IVREditorModule::Get();
+PRAGMA_DISABLE_DEPRECATION_WARNINGS
 					//Disable current actor preview
 					VREditorModule.UpdateActorPreview(SNullWidget::NullWidget, ActorPreviews.Num() - 1, PreviewActor.Get(), !ActorPreview.bIsPanelDetached);
 					//Enable new current preview (the old one was detached or attached, the new one is the opposite
 					VREditorModule.UpdateActorPreview(ActorPreview.PreviewWidget.ToSharedRef(), ActorPreviews.Num() - 1, PreviewActor.Get(), ActorPreview.bIsPanelDetached);
+PRAGMA_ENABLE_DEPRECATION_WARNINGS
 				}
 			}
 		}
@@ -4660,7 +4670,10 @@ void SLevelViewport::UpdateActorPreviewViewports()
 
 void SLevelViewport::OnPreviewSelectedCamerasChange()
 {
+PRAGMA_DISABLE_DEPRECATION_WARNINGS
 	const bool bPreviewInDesktopViewport = !IVREditorModule::Get().IsVREditorModeActive();
+PRAGMA_ENABLE_DEPRECATION_WARNINGS
+
 	// Check to see if previewing selected cameras is enabled and if we're the active level viewport client.
 	if (GetDefault<ULevelEditorViewportSettings>()->bPreviewSelectedCameras && GCurrentLevelEditingViewportClient == LevelViewportClient.Get())
 	{
@@ -5439,7 +5452,9 @@ void SLevelViewport::RemoveActorPreview( int32 PreviewIndex, AActor* Actor, cons
 	{
 		if (ActorPreviews.IsValidIndex(PreviewIndex))
 		{
+PRAGMA_DISABLE_DEPRECATION_WARNINGS
 			VREditorModule.UpdateActorPreview(SNullWidget::NullWidget, PreviewIndex, Actor, ActorPreviews[PreviewIndex].bIsPanelDetached);
+PRAGMA_ENABLE_DEPRECATION_WARNINGS
 		}
 	}
 	else
diff --git a/Engine/Source/Editor/MovieSceneTools/Private/Sections/ThumbnailSection.cpp b/Engine/Source/Editor/MovieSceneTools/Private/Sections/ThumbnailSection.cpp
--- a/Engine/Source/Editor/MovieSceneTools/Private/Sections/ThumbnailSection.cpp
+++ b/Engine/Source/Editor/MovieSceneTools/Private/Sections/ThumbnailSection.cpp
@@ -303,7 +303,9 @@ int32 FThumbnailSection::OnPaintSection( FSequencerSectionPainter& InPainter ) c
 			FSlateLayoutTransform(ThumbnailScale, FVector2D(PositionX-HorizontalCropOffset, PositionY))
 		);
 
+PRAGMA_DISABLE_DEPRECATION_WARNINGS // TODO: XR Creative Framework relevant?
 		if (IVREditorModule::Get().IsVREditorModeActive())
+PRAGMA_ENABLE_DEPRECATION_WARNINGS
 		{
 			// In VR editor every widget is in the world and gamma corrected by the scene renderer.  Thumbnails will have already been gamma
 			// corrected and so they need to be reversed
diff --git a/Engine/Source/Editor/Sequencer/Private/Menus/SequencerToolbarUtils.cpp b/Engine/Source/Editor/Sequencer/Private/Menus/SequencerToolbarUtils.cpp
--- a/Engine/Source/Editor/Sequencer/Private/Menus/SequencerToolbarUtils.cpp
+++ b/Engine/Source/Editor/Sequencer/Private/Menus/SequencerToolbarUtils.cpp
@@ -177,7 +177,9 @@ static TSharedRef<SWidget> MakeAutoChangeMenu(const TWeakPtr<FSequencer> InWeakS
 		MenuBuilder.AddMenuEntry(FSequencerCommands::Get().SetAutoTrack);
 	}
 
+PRAGMA_DISABLE_DEPRECATION_WARNINGS
 	if (IVREditorModule::Get().IsVREditorModeActive()
+PRAGMA_ENABLE_DEPRECATION_WARNINGS
 		|| (InWeakSequencer.Pin()->IsLevelEditorSequencer() && !ExactCast<ULevelSequence>(InWeakSequencer.Pin()->GetFocusedMovieSceneSequence())))
 	{
 		MenuBuilder.AddMenuEntry(FSequencerCommands::Get().SetAutoChangeAll);
@@ -192,7 +194,9 @@ template<typename TCallback, typename TCallback2>
 static void GetAutoKeyMenuEntryArgs(const TSharedPtr<FSequencer>& InSequencer, TCallback&& InConsumer, TCallback2&& InConsumer2)
 {
 	const TWeakPtr<FSequencer> WeakSequencer = InSequencer.ToWeakPtr();
+PRAGMA_DISABLE_DEPRECATION_WARNINGS
 	if (IVREditorModule::Get().IsVREditorModeActive()
+PRAGMA_ENABLE_DEPRECATION_WARNINGS
 		|| (InSequencer->IsLevelEditorSequencer() && ExactCast<ULevelSequence>(InSequencer->GetFocusedMovieSceneSequence()) == nullptr))
 	{
 		TAttribute<FSlateIcon> AutoChangeModeIcon;
diff --git a/Engine/Source/Editor/UnrealEd/Private/AssetThumbnail.cpp b/Engine/Source/Editor/UnrealEd/Private/AssetThumbnail.cpp
--- a/Engine/Source/Editor/UnrealEd/Private/AssetThumbnail.cpp
+++ b/Engine/Source/Editor/UnrealEd/Private/AssetThumbnail.cpp
@@ -241,9 +241,11 @@ public:
 			TSharedPtr<SViewport> Viewport = 
 				SNew( SViewport )
 				.EnableGammaCorrection(false)
+PRAGMA_DISABLE_DEPRECATION_WARNINGS // TODO: XR Creative Framework relevant?
 				// In VR editor every widget is in the world and gamma corrected by the scene renderer.  Thumbnails will have already been gamma
 				// corrected and so they need to be reversed
 				.ReverseGammaCorrection(IVREditorModule::Get().IsVREditorModeActive())
+PRAGMA_ENABLE_DEPRECATION_WARNINGS
 				.EnableBlending(true)
 				.ViewportSize(AssetThumbnail->GetSize());
 
diff --git a/Engine/Source/Editor/UnrealEd/Private/LevelEditorViewport.cpp b/Engine/Source/Editor/UnrealEd/Private/LevelEditorViewport.cpp
--- a/Engine/Source/Editor/UnrealEd/Private/LevelEditorViewport.cpp
+++ b/Engine/Source/Editor/UnrealEd/Private/LevelEditorViewport.cpp
@@ -4120,6 +4120,7 @@ void FLevelEditorViewportClient::RedrawAllViewportsIntoThisScene()
 
 void FLevelEditorViewportClient::SetVREditView(bool bGameViewEnable)
 {
+PRAGMA_DISABLE_DEPRECATION_WARNINGS
 	UEditorWorldExtensionCollection* ExtensionCollection = GEditor->GetEditorWorldExtensionsManager()->GetEditorWorldExtensions(GetWorld());
 	check(ExtensionCollection != nullptr);
 	UVREditorMode* VREditorMode = Cast<UVREditorMode>(ExtensionCollection->FindExtension(UVREditorMode::StaticClass()));
@@ -4134,6 +4135,7 @@ void FLevelEditorViewportClient::SetVREditView(bool bGameViewEnable)
 			VREditorMode->OnPlacePreviewActor().RemoveAll(this);
 		}
 	}
+PRAGMA_ENABLE_DEPRECATION_WARNINGS
 
 	FEditorViewportClient::SetVREditView(bGameViewEnable);
 }
diff --git a/Engine/Source/Editor/UnrealEd/Private/UnrealEdMisc.cpp b/Engine/Source/Editor/UnrealEd/Private/UnrealEdMisc.cpp
--- a/Engine/Source/Editor/UnrealEd/Private/UnrealEdMisc.cpp
+++ b/Engine/Source/Editor/UnrealEd/Private/UnrealEdMisc.cpp
@@ -832,7 +832,7 @@ void FUnrealEdMisc::EditorAnalyticsHeartbeat()
 	{
 		bWasDebuggerPresent = bIsDebuggerPresent;
 	}
-	const bool bInVRMode = IVREditorModule::Get().IsVREditorModeActive();
+	const bool bInVRMode = IVREditorModule::Get().GetVRModeBase() != nullptr;
 	const bool bIsEnterprise = IProjectManager::Get().IsEnterpriseProject();
 
 	const double LastInteractionTime = FSlateApplication::Get().GetLastUserInteractionTime();
diff --git a/Engine/Source/Editor/VREditor/Private/Teleporter/VREditorAutoScaler.cpp b/Engine/Source/Editor/VREditor/Private/Teleporter/VREditorAutoScaler.cpp
--- a/Engine/Source/Editor/VREditor/Private/Teleporter/VREditorAutoScaler.cpp
+++ b/Engine/Source/Editor/VREditor/Private/Teleporter/VREditorAutoScaler.cpp
@@ -10,6 +10,8 @@
 
 #include UE_INLINE_GENERATED_CPP_BY_NAME(VREditorAutoScaler)
 
+PRAGMA_DISABLE_DEPRECATION_WARNINGS
+
 namespace VREd
 {
 	static FAutoConsoleVariable AllowResetScale(TEXT("VREd.AllowResetScale"), 1, TEXT("Allowed to reset world to meters to default world to meters"));
@@ -79,3 +81,5 @@ void UVREditorAutoScaler::OnInteractorAction( class FEditorViewportClient& Viewp
 	}
 }
 
+
+PRAGMA_ENABLE_DEPRECATION_WARNINGS
diff --git a/Engine/Source/Editor/VREditor/Private/Teleporter/VREditorAutoScaler.h b/Engine/Source/Editor/VREditor/Private/Teleporter/VREditorAutoScaler.h
--- a/Engine/Source/Editor/VREditor/Private/Teleporter/VREditorAutoScaler.h
+++ b/Engine/Source/Editor/VREditor/Private/Teleporter/VREditorAutoScaler.h
@@ -7,8 +7,13 @@
 #include "UObject/Object.h"
 #include "VREditorAutoScaler.generated.h"
 
+
+PRAGMA_DISABLE_DEPRECATION_WARNINGS
+
+
 class UVREditorMode;
 
+
 /**
  * Automatically scales the user when dragging the world and pressing the touchpad
  */
@@ -41,3 +46,6 @@ private:
 	UPROPERTY()
 	TObjectPtr<UVREditorMode> VRMode;
 };
+
+
+PRAGMA_ENABLE_DEPRECATION_WARNINGS
diff --git a/Engine/Source/Editor/VREditor/Private/Teleporter/VREditorTeleporter.cpp b/Engine/Source/Editor/VREditor/Private/Teleporter/VREditorTeleporter.cpp
--- a/Engine/Source/Editor/VREditor/Private/Teleporter/VREditorTeleporter.cpp
+++ b/Engine/Source/Editor/VREditor/Private/Teleporter/VREditorTeleporter.cpp
@@ -19,6 +19,8 @@
 
 #include UE_INLINE_GENERATED_CPP_BY_NAME(VREditorTeleporter)
 
+PRAGMA_DISABLE_DEPRECATION_WARNINGS
+
 namespace VREd
 {
 	static FAutoConsoleVariable TeleportLerpTime(TEXT("VREd.TeleportLerpTime"), 0.1f, TEXT("The lerp time to teleport"));
@@ -558,3 +560,6 @@ float AVREditorTeleporter::GetSlideDelta_Implementation(UVREditorInteractor* Int
 
 	return Result;
 }
+
+
+PRAGMA_ENABLE_DEPRECATION_WARNINGS
diff --git a/Engine/Source/Editor/VREditor/Private/UI/VREditorBaseUserWidget.cpp b/Engine/Source/Editor/VREditor/Private/UI/VREditorBaseUserWidget.cpp
--- a/Engine/Source/Editor/VREditor/Private/UI/VREditorBaseUserWidget.cpp
+++ b/Engine/Source/Editor/VREditor/Private/UI/VREditorBaseUserWidget.cpp
@@ -5,6 +5,9 @@
 
 #include UE_INLINE_GENERATED_CPP_BY_NAME(VREditorBaseUserWidget)
 
+PRAGMA_DISABLE_DEPRECATION_WARNINGS
+
+
 UVREditorBaseUserWidget::UVREditorBaseUserWidget( const FObjectInitializer& ObjectInitializer )
 	: Super( ObjectInitializer ),
 	  Owner( nullptr )
@@ -17,3 +20,5 @@ void UVREditorBaseUserWidget::SetOwner( class AVREditorFloatingUI* NewOwner )
 	Owner = NewOwner;
 }
 
+
+PRAGMA_ENABLE_DEPRECATION_WARNINGS
diff --git a/Engine/Source/Editor/VREditor/Private/UI/VREditorBaseUserWidget.h b/Engine/Source/Editor/VREditor/Private/UI/VREditorBaseUserWidget.h
--- a/Engine/Source/Editor/VREditor/Private/UI/VREditorBaseUserWidget.h
+++ b/Engine/Source/Editor/VREditor/Private/UI/VREditorBaseUserWidget.h
@@ -2,12 +2,15 @@
 
 #pragma once
 
-
 #include "CoreMinimal.h"
 #include "UObject/ObjectMacros.h"
 #include "Blueprint/UserWidget.h"
 #include "VREditorBaseUserWidget.generated.h"
 
+
+PRAGMA_DISABLE_DEPRECATION_WARNINGS
+
+
 class AVREditorFloatingUI;
 
 /**
@@ -46,3 +49,6 @@ protected:
 	TWeakObjectPtr<class AVREditorFloatingUI> Owner;
 
 };
+
+
+PRAGMA_ENABLE_DEPRECATION_WARNINGS
diff --git a/Engine/Source/Editor/VREditor/Private/UI/VREditorDockableWindow.cpp b/Engine/Source/Editor/VREditor/Private/UI/VREditorDockableWindow.cpp
--- a/Engine/Source/Editor/VREditor/Private/UI/VREditorDockableWindow.cpp
+++ b/Engine/Source/Editor/VREditor/Private/UI/VREditorDockableWindow.cpp
@@ -15,6 +15,9 @@
 
 #include UE_INLINE_GENERATED_CPP_BY_NAME(VREditorDockableWindow)
 
+PRAGMA_DISABLE_DEPRECATION_WARNINGS
+
+
 namespace VREd
 {
 	static FAutoConsoleVariable DockWindowThickness( TEXT( "VREd.DockWindowTickness" ), 1.0f, TEXT( "How thick the window is" ) );
@@ -542,3 +545,6 @@ void UDockableWindowDragOperation::ExecuteDrag( UViewportInteractor* Interactor,
 		LastLaserImpactToWorld = SmoothedLaserImpactToWorld;
 	}
 }
+
+
+PRAGMA_ENABLE_DEPRECATION_WARNINGS
diff --git a/Engine/Source/Editor/VREditor/Private/UI/VREditorDockableWindow.h b/Engine/Source/Editor/VREditor/Private/UI/VREditorDockableWindow.h
--- a/Engine/Source/Editor/VREditor/Private/UI/VREditorDockableWindow.h
+++ b/Engine/Source/Editor/VREditor/Private/UI/VREditorDockableWindow.h
@@ -9,6 +9,10 @@
 #include "ViewportInteractableInterface.h"
 #include "VREditorDockableWindow.generated.h"
 
+
+PRAGMA_DISABLE_DEPRECATION_WARNINGS
+
+
 class UStaticMeshComponent;
 class UViewportInteractor;
 
@@ -160,3 +164,6 @@ public:
 	/** last frame's LaserImpactToWorld transform */
 	TOptional<FTransform> LastLaserImpactToWorld;
 };
+
+
+PRAGMA_ENABLE_DEPRECATION_WARNINGS
diff --git a/Engine/Source/Editor/VREditor/Private/UI/VREditorFloatingCameraUI.cpp b/Engine/Source/Editor/VREditor/Private/UI/VREditorFloatingCameraUI.cpp
--- a/Engine/Source/Editor/VREditor/Private/UI/VREditorFloatingCameraUI.cpp
+++ b/Engine/Source/Editor/VREditor/Private/UI/VREditorFloatingCameraUI.cpp
@@ -15,6 +15,9 @@
 
 #include UE_INLINE_GENERATED_CPP_BY_NAME(VREditorFloatingCameraUI)
 
+PRAGMA_DISABLE_DEPRECATION_WARNINGS
+
+
 #define LOCTEXT_NAMESPACE "AVREditorFloatingCameraUI"
 
 AVREditorFloatingCameraUI::AVREditorFloatingCameraUI(const FObjectInitializer& ObjectInitializer)
@@ -56,3 +59,6 @@ FTransform AVREditorFloatingCameraUI::MakeCustomUITransform()
 }
 
 #undef LOCTEXT_NAMESPACE
+
+
+PRAGMA_ENABLE_DEPRECATION_WARNINGS
diff --git a/Engine/Source/Editor/VREditor/Private/UI/VREditorFloatingCameraUI.h b/Engine/Source/Editor/VREditor/Private/UI/VREditorFloatingCameraUI.h
--- a/Engine/Source/Editor/VREditor/Private/UI/VREditorFloatingCameraUI.h
+++ b/Engine/Source/Editor/VREditor/Private/UI/VREditorFloatingCameraUI.h
@@ -8,6 +8,10 @@
 #include "UI/VREditorFloatingUI.h"
 #include "VREditorFloatingCameraUI.generated.h"
 
+
+PRAGMA_DISABLE_DEPRECATION_WARNINGS
+
+
 class UVREditorBaseUserWidget;
 class UVREditorUISystem;
 
@@ -36,3 +40,6 @@ private:
 	UPROPERTY( )
 	TWeakObjectPtr<AActor> LinkedActor;
 };
+
+
+PRAGMA_ENABLE_DEPRECATION_WARNINGS
diff --git a/Engine/Source/Editor/VREditor/Private/UI/VREditorFloatingUI.cpp b/Engine/Source/Editor/VREditor/Private/UI/VREditorFloatingUI.cpp
--- a/Engine/Source/Editor/VREditor/Private/UI/VREditorFloatingUI.cpp
+++ b/Engine/Source/Editor/VREditor/Private/UI/VREditorFloatingUI.cpp
@@ -11,6 +11,9 @@
 
 #include UE_INLINE_GENERATED_CPP_BY_NAME(VREditorFloatingUI)
 
+PRAGMA_DISABLE_DEPRECATION_WARNINGS
+
+
 namespace VREd
 {
 	static FAutoConsoleVariable UIFadeSpeed(TEXT("VREd.UIFadeSpeed"), 6.0f, TEXT("How fast UI should fade in and out"));
@@ -475,3 +478,6 @@ void AVREditorFloatingUI::SetWindowMesh(class UStaticMesh* InWindowMesh)
 	check(InWindowMesh != nullptr);
 	WindowMeshComponent->SetStaticMesh(InWindowMesh);
 }
+
+
+PRAGMA_ENABLE_DEPRECATION_WARNINGS
diff --git a/Engine/Source/Editor/VREditor/Private/UI/VREditorRadialFloatingUI.cpp b/Engine/Source/Editor/VREditor/Private/UI/VREditorRadialFloatingUI.cpp
--- a/Engine/Source/Editor/VREditor/Private/UI/VREditorRadialFloatingUI.cpp
+++ b/Engine/Source/Editor/VREditor/Private/UI/VREditorRadialFloatingUI.cpp
@@ -15,6 +15,9 @@
 
 #include UE_INLINE_GENERATED_CPP_BY_NAME(VREditorRadialFloatingUI)
 
+PRAGMA_DISABLE_DEPRECATION_WARNINGS
+
+
 namespace VREd
 {
 	static FAutoConsoleVariable RadialUIFadeSpeed( TEXT( "VREd.RadialUIFadeSpeed" ), 6.0f, TEXT( "How fast UI should fade in and out" ) );
@@ -582,3 +585,6 @@ void AVREditorRadialFloatingUI::UpdateCentralWidgetComponent(const TSharedPtr<SW
 		}
 	}
 }
+
+
+PRAGMA_ENABLE_DEPRECATION_WARNINGS
diff --git a/Engine/Source/Editor/VREditor/Private/UI/VREditorRadialFloatingUI.h b/Engine/Source/Editor/VREditor/Private/UI/VREditorRadialFloatingUI.h
--- a/Engine/Source/Editor/VREditor/Private/UI/VREditorRadialFloatingUI.h
+++ b/Engine/Source/Editor/VREditor/Private/UI/VREditorRadialFloatingUI.h
@@ -10,6 +10,9 @@
 #include "VREditorRadialFloatingUI.generated.h"
 
 
+PRAGMA_DISABLE_DEPRECATION_WARNINGS
+
+
 class UVREditorBaseUserWidget;
 class UVREditorUISystem;
 
@@ -214,3 +217,5 @@ private:
 
 };
 
+
+PRAGMA_ENABLE_DEPRECATION_WARNINGS
diff --git a/Engine/Source/Editor/VREditor/Private/UI/VREditorUISystem.cpp b/Engine/Source/Editor/VREditor/Private/UI/VREditorUISystem.cpp
--- a/Engine/Source/Editor/VREditor/Private/UI/VREditorUISystem.cpp
+++ b/Engine/Source/Editor/VREditor/Private/UI/VREditorUISystem.cpp
@@ -97,6 +97,9 @@
 
 #include UE_INLINE_GENERATED_CPP_BY_NAME(VREditorUISystem)
 
+PRAGMA_DISABLE_DEPRECATION_WARNINGS
+
+
 #define LOCTEXT_NAMESPACE "VREditor"
 
 namespace VREd
@@ -2318,3 +2321,6 @@ void UVREditorUISystem::OnCloseAssetEditor(UObject* Asset, EAssetEditorCloseReas
 }
 
 #undef LOCTEXT_NAMESPACE
+
+
+PRAGMA_ENABLE_DEPRECATION_WARNINGS
diff --git a/Engine/Source/Editor/VREditor/Private/UI/VRRadialMenuHandler.cpp b/Engine/Source/Editor/VREditor/Private/UI/VRRadialMenuHandler.cpp
--- a/Engine/Source/Editor/VREditor/Private/UI/VRRadialMenuHandler.cpp
+++ b/Engine/Source/Editor/VREditor/Private/UI/VRRadialMenuHandler.cpp
@@ -87,6 +87,9 @@
 
 #include UE_INLINE_GENERATED_CPP_BY_NAME(VRRadialMenuHandler)
 
+PRAGMA_DISABLE_DEPRECATION_WARNINGS
+
+
 #define LOCTEXT_NAMESPACE "VREditor"
 
 FText UVRRadialMenuHandler::ActionMenuLabel = LOCTEXT("DefaultActions", "Actions");
@@ -559,3 +562,6 @@ void UVRRadialMenuHandler::RegisterMenuGenerator(const FOnRadialMenuGenerated Ne
 
 
 #undef LOCTEXT_NAMESPACE
+
+
+PRAGMA_ENABLE_DEPRECATION_WARNINGS
\ No newline at end of file
diff --git a/Engine/Source/Editor/VREditor/Private/VREditorActions.cpp b/Engine/Source/Editor/VREditor/Private/VREditorActions.cpp
--- a/Engine/Source/Editor/VREditor/Private/VREditorActions.cpp
+++ b/Engine/Source/Editor/VREditor/Private/VREditorActions.cpp
@@ -32,6 +32,10 @@
 #include "Factories/Factory.h"
 #include "Subsystems/AssetEditorSubsystem.h"
 
+
+PRAGMA_DISABLE_DEPRECATION_WARNINGS
+
+
 #define LOCTEXT_NAMESPACE "VREditorActions"
 
 namespace VREd
@@ -541,3 +545,6 @@ void FVREditorActionCallbacks::ExitVRMode(UVREditorMode* InVRMode)
 }
 
 #undef LOCTEXT_NAMESPACE
+
+
+PRAGMA_ENABLE_DEPRECATION_WARNINGS
diff --git a/Engine/Source/Editor/VREditor/Private/VREditorActions.h b/Engine/Source/Editor/VREditor/Private/VREditorActions.h
--- a/Engine/Source/Editor/VREditor/Private/VREditorActions.h
+++ b/Engine/Source/Editor/VREditor/Private/VREditorActions.h
@@ -7,6 +7,10 @@
 #include "Editor.h"
 #include "UnrealWidgetFwd.h"
 
+
+PRAGMA_DISABLE_DEPRECATION_WARNINGS
+
+
 // Forward declarations 
 class UVREditorMode;
 enum class EGizmoHandleTypes : uint8;
@@ -202,3 +206,6 @@ public:
 	static FText SequencerLoopText;
 
 };
+
+
+PRAGMA_ENABLE_DEPRECATION_WARNINGS
diff --git a/Engine/Source/Editor/VREditor/Private/VREditorAvatarActor.cpp b/Engine/Source/Editor/VREditor/Private/VREditorAvatarActor.cpp
--- a/Engine/Source/Editor/VREditor/Private/VREditorAvatarActor.cpp
+++ b/Engine/Source/Editor/VREditor/Private/VREditorAvatarActor.cpp
@@ -26,6 +26,9 @@
 
 #include UE_INLINE_GENERATED_CPP_BY_NAME(VREditorAvatarActor)
 
+PRAGMA_DISABLE_DEPRECATION_WARNINGS
+
+
 namespace VREd
 {
 	static FAutoConsoleVariable GridMovementTolerance( TEXT( "VREd.GridMovementTolerance" ), 0.1f, TEXT( "Tolerance for movement when the grid must disappear" ) );
@@ -483,3 +486,6 @@ void AVREditorAvatarActor::TickManually( const float DeltaTime )
 		}
 	}
 }
+
+
+PRAGMA_ENABLE_DEPRECATION_WARNINGS
diff --git a/Engine/Source/Editor/VREditor/Private/VREditorAvatarActor.h b/Engine/Source/Editor/VREditor/Private/VREditorAvatarActor.h
--- a/Engine/Source/Editor/VREditor/Private/VREditorAvatarActor.h
+++ b/Engine/Source/Editor/VREditor/Private/VREditorAvatarActor.h
@@ -7,6 +7,10 @@
 #include "GameFramework/Actor.h"
 #include "VREditorAvatarActor.generated.h"
 
+
+PRAGMA_DISABLE_DEPRECATION_WARNINGS
+
+
 class UMaterialInstanceDynamic;
 
 /**
@@ -115,3 +119,6 @@ private:
 	UPROPERTY()
 	TObjectPtr<class UVREditorMode> VRMode;
 };
+
+
+PRAGMA_ENABLE_DEPRECATION_WARNINGS
diff --git a/Engine/Source/Editor/VREditor/Private/VREditorBaseActor.cpp b/Engine/Source/Editor/VREditor/Private/VREditorBaseActor.cpp
--- a/Engine/Source/Editor/VREditor/Private/VREditorBaseActor.cpp
+++ b/Engine/Source/Editor/VREditor/Private/VREditorBaseActor.cpp
@@ -6,6 +6,9 @@
 
 #include UE_INLINE_GENERATED_CPP_BY_NAME(VREditorBaseActor)
 
+PRAGMA_DISABLE_DEPRECATION_WARNINGS
+
+
 namespace VREd
 {
 	static FAutoConsoleVariable UIOnHandRotationOffset( TEXT( "VREd.UIOnHandRotationOffset" ), 45.0f, TEXT( "Rotation offset for UI that's docked to your hand, to make it more comfortable to hold" ) );
@@ -216,3 +219,6 @@ void AVREditorBaseActor::TickMoveTo( const float DeltaTime )
 
 	SetTransform( NewTransform );
 }
+
+
+PRAGMA_ENABLE_DEPRECATION_WARNINGS
diff --git a/Engine/Source/Editor/VREditor/Private/VREditorCameraWidgetComponent.cpp b/Engine/Source/Editor/VREditor/Private/VREditorCameraWidgetComponent.cpp
--- a/Engine/Source/Editor/VREditor/Private/VREditorCameraWidgetComponent.cpp
+++ b/Engine/Source/Editor/VREditor/Private/VREditorCameraWidgetComponent.cpp
@@ -7,6 +7,9 @@
 
 #include UE_INLINE_GENERATED_CPP_BY_NAME(VREditorCameraWidgetComponent)
 
+PRAGMA_DISABLE_DEPRECATION_WARNINGS
+
+
 UVREditorCameraWidgetComponent::UVREditorCameraWidgetComponent()
 	: Super()	
 {	
@@ -36,3 +39,5 @@ UVREditorCameraWidgetComponent::UVREditorCameraWidgetComponent()
 	}
 }
 
+
+PRAGMA_ENABLE_DEPRECATION_WARNINGS
diff --git a/Engine/Source/Editor/VREditor/Private/VREditorFloatingText.cpp b/Engine/Source/Editor/VREditor/Private/VREditorFloatingText.cpp
--- a/Engine/Source/Editor/VREditor/Private/VREditorFloatingText.cpp
+++ b/Engine/Source/Editor/VREditor/Private/VREditorFloatingText.cpp
@@ -17,6 +17,8 @@
 
 #include UE_INLINE_GENERATED_CPP_BY_NAME(VREditorFloatingText)
 
+PRAGMA_DISABLE_DEPRECATION_WARNINGS
+
 
 AFloatingText::AFloatingText()
 {
@@ -221,3 +223,6 @@ void AFloatingText::Update( const FVector OrientateToward )
 	TextComponent->SetWorldRotation( ( TowardRotation * FVector::ForwardVector ).ToOrientationQuat() );
 
 }
+
+
+PRAGMA_ENABLE_DEPRECATION_WARNINGS
diff --git a/Engine/Source/Editor/VREditor/Private/VREditorInteractor.cpp b/Engine/Source/Editor/VREditor/Private/VREditorInteractor.cpp
--- a/Engine/Source/Editor/VREditor/Private/VREditorInteractor.cpp
+++ b/Engine/Source/Editor/VREditor/Private/VREditorInteractor.cpp
@@ -31,6 +31,9 @@
 
 #include UE_INLINE_GENERATED_CPP_BY_NAME(VREditorInteractor)
 
+PRAGMA_DISABLE_DEPRECATION_WARNINGS
+
+
 namespace VREd
 {
 	static FAutoConsoleVariable TriggerTouchThreshold_Vive( TEXT( "VI.TriggerTouchThreshold_Vive" ), 0.025f, TEXT( "Minimum trigger threshold before we consider the trigger 'touched'" ) );
@@ -876,7 +879,7 @@ FHitResult UVREditorInteractor::GetHitResultFromLaserPointer( TArray<AActor*>* O
 	// Ignore UI widgets too
 	if (GetDraggingMode() == EViewportInteractionDraggingMode::TransformablesAtLaserImpact)
 	{
-		for (TActorIterator<AVREditorFloatingUI> UIActorIt( WorldInteraction->GetWorld() ); UIActorIt; ++UIActorIt)
+		for (TActorIterator<AVREditorFloatingUI> UIActorIt( WorldInteraction->GetWorld(), AVREditorFloatingUI::StaticClass() ); UIActorIt; ++UIActorIt)
 		{
 			OptionalListOfIgnoredActors->Add( *UIActorIt );
 		}
@@ -2012,3 +2015,6 @@ float UVREditorInteractor::GetSelectAndMoveTriggerValue() const
 }
 
 #undef LOCTEXT_NAMESPACE
+
+
+PRAGMA_ENABLE_DEPRECATION_WARNINGS
diff --git a/Engine/Source/Editor/VREditor/Private/VREditorModule.cpp b/Engine/Source/Editor/VREditor/Private/VREditorModule.cpp
--- a/Engine/Source/Editor/VREditor/Private/VREditorModule.cpp
+++ b/Engine/Source/Editor/VREditor/Private/VREditorModule.cpp
@@ -52,9 +52,11 @@ public:
 	virtual void EnableVREditor( const bool bEnable, const bool bForceWithoutHMD ) override;
 	virtual bool IsVREditorModeActive() override;
 	virtual UVREditorModeBase* GetVRModeBase() override;
+PRAGMA_DISABLE_DEPRECATION_WARNINGS
 	virtual UVREditorMode* GetVRMode() override;
-	virtual void UpdateActorPreview(TSharedRef<SWidget> InWidget, int32 Index, AActor *Actor, bool bIsPanelDetached = false) override;
 	virtual void UpdateExternalUMGUI(const FVREditorFloatingUICreationContext& CreationContext) override;
+PRAGMA_ENABLE_DEPRECATION_WARNINGS
+	virtual void UpdateActorPreview(TSharedRef<SWidget> InWidget, int32 Index, AActor *Actor, bool bIsPanelDetached = false) override;
 	virtual void UpdateExternalSlateUI(TSharedRef<SWidget> InSlateWidget, FName Name, FVector2D InSize) override;
 	virtual TSharedPtr<FExtender> GetRadialMenuExtender() override
 	{
@@ -128,7 +130,9 @@ void FVREditorModule::ShutdownModule()
 
 	if (GIsEditor)
 	{
+PRAGMA_DISABLE_DEPRECATION_WARNINGS
 		FVREditorStyle::Shutdown();
+PRAGMA_ENABLE_DEPRECATION_WARNINGS
 	}
 
 	ModeManager.Reset();
@@ -169,6 +173,8 @@ UVREditorModeBase* FVREditorModule::GetVRModeBase()
 	return ModeManager->GetCurrentVREditorMode();
 }
 
+
+PRAGMA_DISABLE_DEPRECATION_WARNINGS
 UVREditorMode* FVREditorModule::GetVRMode()
 {
 	UVREditorModeBase* ModeBase = GetVRModeBase();
@@ -189,6 +195,7 @@ void FVREditorModule::UpdateExternalSlateUI(TSharedRef<SWidget> InSlateWidget, F
 {
 	GetVRMode()->UpdateExternalSlateUI(InSlateWidget, Name, InSize);
 }
+PRAGMA_ENABLE_DEPRECATION_WARNINGS
 
 void FVREditorModule::ToggleForceVRMode()
 {
diff --git a/Engine/Source/Editor/VREditor/Private/VREditorPlacement.cpp b/Engine/Source/Editor/VREditor/Private/VREditorPlacement.cpp
--- a/Engine/Source/Editor/VREditor/Private/VREditorPlacement.cpp
+++ b/Engine/Source/Editor/VREditor/Private/VREditorPlacement.cpp
@@ -34,6 +34,9 @@
 
 #include UE_INLINE_GENERATED_CPP_BY_NAME(VREditorPlacement)
 
+PRAGMA_DISABLE_DEPRECATION_WARNINGS
+
+
 #define LOCTEXT_NAMESPACE "VREditor"
 
 
@@ -109,3 +112,6 @@ void UVREditorPlacement::UpdateNearClipPlaneOnScaleChange(const float NewWorldTo
 
 
 #undef LOCTEXT_NAMESPACE
+
+
+PRAGMA_ENABLE_DEPRECATION_WARNINGS
diff --git a/Engine/Source/Editor/VREditor/Private/VREditorPlacement.h b/Engine/Source/Editor/VREditor/Private/VREditorPlacement.h
--- a/Engine/Source/Editor/VREditor/Private/VREditorPlacement.h
+++ b/Engine/Source/Editor/VREditor/Private/VREditorPlacement.h
@@ -8,6 +8,10 @@
 #include "AssetRegistry/AssetData.h"
 #include "VREditorPlacement.generated.h"
 
+
+PRAGMA_DISABLE_DEPRECATION_WARNINGS
+
+
 class UActorComponent;
 class UViewportInteractor;
 class UViewportWorldInteraction;
@@ -64,3 +68,5 @@ protected:
 	TObjectPtr<UObject> PlacingMaterialOrTextureAsset;
 };
 
+
+PRAGMA_ENABLE_DEPRECATION_WARNINGS
diff --git a/Engine/Source/Editor/VREditor/Private/VREditorStyle.cpp b/Engine/Source/Editor/VREditor/Private/VREditorStyle.cpp
--- a/Engine/Source/Editor/VREditor/Private/VREditorStyle.cpp
+++ b/Engine/Source/Editor/VREditor/Private/VREditorStyle.cpp
@@ -20,6 +20,10 @@
 #include "Styling/SlateTypes.h"
 #include "Templates/TypeHash.h"
 
+
+PRAGMA_DISABLE_DEPRECATION_WARNINGS
+
+
 #define RootToContentDir Style->RootToContentDir
 
 class ISlateStyle;
@@ -250,3 +254,6 @@ const ISlateStyle& FVREditorStyle::Get()
 }
 
 #undef RootToContentDir
+
+
+PRAGMA_ENABLE_DEPRECATION_WARNINGS
diff --git a/Engine/Source/Editor/VREditor/Private/VREditorWidgetComponent.cpp b/Engine/Source/Editor/VREditor/Private/VREditorWidgetComponent.cpp
--- a/Engine/Source/Editor/VREditor/Private/VREditorWidgetComponent.cpp
+++ b/Engine/Source/Editor/VREditor/Private/VREditorWidgetComponent.cpp
@@ -7,6 +7,9 @@
 
 #include UE_INLINE_GENERATED_CPP_BY_NAME(VREditorWidgetComponent)
 
+PRAGMA_DISABLE_DEPRECATION_WARNINGS
+
+
 UVREditorWidgetComponent::UVREditorWidgetComponent()
 	: Super(),
 	DrawingPolicy(EVREditorWidgetDrawingPolicy::Always),
@@ -59,3 +62,6 @@ void UVREditorWidgetComponent::DrawWidgetToRenderTarget(float DeltaTime)
 
 	Super::DrawWidgetToRenderTarget(DeltaTime);
 }
+
+
+PRAGMA_ENABLE_DEPRECATION_WARNINGS
diff --git a/Engine/Source/Editor/VREditor/Private/VRScoutingInteractor.cpp b/Engine/Source/Editor/VREditor/Private/VRScoutingInteractor.cpp
--- a/Engine/Source/Editor/VREditor/Private/VRScoutingInteractor.cpp
+++ b/Engine/Source/Editor/VREditor/Private/VRScoutingInteractor.cpp
@@ -16,6 +16,9 @@
 
 #include UE_INLINE_GENERATED_CPP_BY_NAME(VRScoutingInteractor)
 
+PRAGMA_DISABLE_DEPRECATION_WARNINGS
+
+
 #define LOCTEXT_NAMESPACE "UVRScoutingInteractor"
 
 UVRScoutingInteractor::UVRScoutingInteractor() :
@@ -122,3 +125,6 @@ void UVRScoutingInteractor::SetReceivesEditorInput(bool bInValue)
 }
 
 #undef LOCTEXT_NAMESPACE
+
+
+PRAGMA_ENABLE_DEPRECATION_WARNINGS
diff --git a/Engine/Source/Editor/VREditor/Public/IVREditorModule.h b/Engine/Source/Editor/VREditor/Public/IVREditorModule.h
--- a/Engine/Source/Editor/VREditor/Public/IVREditorModule.h
+++ b/Engine/Source/Editor/VREditor/Public/IVREditorModule.h
@@ -81,8 +81,7 @@ public:
 	 *
 	 * @return True if the VREditor is currently running 
 	 */
-	// This is slated for deprecation in the next changelist (with accompanying cleanups in external code).
-	//UE_DEPRECATED(5.1, "Test whether GetVRModeBase returns nullptr instead. This only returns true in the case of UVREditorMode.")
+	UE_DEPRECATED(5.7, "Test whether GetVRModeBase returns nullptr instead; this only returns true for the deprecated UVREditorMode.")
 	virtual bool IsVREditorModeActive() = 0;
 
 	/** 
@@ -90,8 +89,7 @@ public:
 	*
 	* @return The current VREditor running
 	*/
-	// This is slated for deprecation in the next changelist (with accompanying cleanups in external code).
-	//UE_DEPRECATED(5.1, "Use GetVRModeBase instead.")
+	UE_DEPRECATED(5.7, "Use GetVRModeBase instead; this only supports the deprecated UVREditorMode.")
 	virtual class UVREditorMode* GetVRMode() = 0;
 
 	/** @return The active UVREditorModeBase */
@@ -102,6 +100,7 @@ public:
 	*
 	* @param The new actor preview widget
 	*/
+	UE_DEPRECATED(5.7, "VR Editor mode is deprecated; use the XR Creative Framework plugin instead.")
 	virtual void UpdateActorPreview(TSharedRef<class SWidget> InWidget, int32 Index, AActor *Actor = nullptr, bool bIsDetached = false) = 0;
 
 
@@ -111,6 +110,7 @@ public:
 	* @param The new widget
 	* @param The label to use for the UI
 	*/
+	UE_DEPRECATED(5.7, "VR Editor mode is deprecated; use the XR Creative Framework plugin instead.")
 	virtual void UpdateExternalUMGUI(const struct FVREditorFloatingUICreationContext& CreationContext) = 0;
 
 	/**
@@ -119,9 +119,11 @@ public:
 	* @param The new widget
 	* @param The label to use for the UI
 	*/
+	UE_DEPRECATED(5.7, "VR Editor mode is deprecated; use the XR Creative Framework plugin instead.")
 	virtual void UpdateExternalSlateUI(TSharedRef<SWidget> InSlateWidget, FName Name, FVector2D InSize = FVector2D::ZeroVector) = 0;
 
 	/** Gets the radial menu extender.  This can be used to add your own menu items to the VR radial menu */
+	UE_DEPRECATED(5.7, "VR Editor mode is deprecated; use the XR Creative Framework plugin instead.")
 	virtual TSharedPtr<class FExtender> GetRadialMenuExtender() = 0;
 
 };
diff --git a/Engine/Source/Editor/VREditor/Public/Teleporter/VREditorTeleporter.h b/Engine/Source/Editor/VREditor/Public/Teleporter/VREditorTeleporter.h
--- a/Engine/Source/Editor/VREditor/Public/Teleporter/VREditorTeleporter.h
+++ b/Engine/Source/Editor/VREditor/Public/Teleporter/VREditorTeleporter.h
@@ -11,6 +11,8 @@
 
 #define UE_API VREDITOR_API
 
+PRAGMA_DISABLE_DEPRECATION_WARNINGS
+
 
 class UMaterialInstanceDynamic;
 class UVREditorInteractor;
@@ -20,8 +22,8 @@ class UVREditorMode;
 /**
  * VR Editor teleport manager and the visual representation of the teleport
  */
-UCLASS(MinimalAPI,  Blueprintable, BlueprintType )
-class AVREditorTeleporter: public AActor
+UCLASS(MinimalAPI, Blueprintable, BlueprintType)
+class UE_DEPRECATED(5.7, "VR Editor mode is deprecated; use the XR Creative Framework plugin instead.") AVREditorTeleporter : public AActor
 {
 	GENERATED_BODY()
 
@@ -195,4 +197,7 @@ private:
 	uint32 TeleportTickDelay;
 };
 
+PRAGMA_ENABLE_DEPRECATION_WARNINGS
+
 #undef UE_API
+
diff --git a/Engine/Source/Editor/VREditor/Public/UI/VREditorFloatingUI.h b/Engine/Source/Editor/VREditor/Public/UI/VREditorFloatingUI.h
--- a/Engine/Source/Editor/VREditor/Public/UI/VREditorFloatingUI.h
+++ b/Engine/Source/Editor/VREditor/Public/UI/VREditorFloatingUI.h
@@ -9,6 +9,10 @@
 #include "Blueprint/UserWidget.h"
 #include "VREditorFloatingUI.generated.h"
 
+
+PRAGMA_DISABLE_DEPRECATION_WARNINGS
+
+
 class UVREditorBaseUserWidget;
 class UVREditorUISystem;
 
@@ -18,7 +22,7 @@ typedef FName VREditorPanelID;
  * Creation parameters for AVREditorFloatingUI
  */
 USTRUCT(BlueprintType)
-struct FVREditorFloatingUICreationContext
+struct UE_DEPRECATED(5.7, "VR Editor mode is deprecated; use the XR Creative Framework plugin instead.") FVREditorFloatingUICreationContext
 {
 	GENERATED_BODY()
 
@@ -70,7 +74,7 @@ public:
  * Represents an interactive floating UI panel in the VR Editor
  */
 UCLASS()
-class AVREditorFloatingUI : public AVREditorBaseActor
+class UE_DEPRECATED(5.7, "VR Editor mode is deprecated; use the XR Creative Framework plugin instead.") AVREditorFloatingUI : public AVREditorBaseActor
 {
 	GENERATED_BODY()
 
@@ -244,4 +248,7 @@ private:
 	/** Null out the widget when hidden. */
 	bool bClearWidgetOnHide;
 
-};
\ No newline at end of file
+};
+
+
+PRAGMA_ENABLE_DEPRECATION_WARNINGS
diff --git a/Engine/Source/Editor/VREditor/Public/UI/VREditorUISystem.h b/Engine/Source/Editor/VREditor/Public/UI/VREditorUISystem.h
--- a/Engine/Source/Editor/VREditor/Public/UI/VREditorUISystem.h
+++ b/Engine/Source/Editor/VREditor/Public/UI/VREditorUISystem.h
@@ -11,6 +11,10 @@
 #include "EditorWorldExtension.h"
 #include "VREditorUISystem.generated.h"
 
+
+PRAGMA_DISABLE_DEPRECATION_WARNINGS
+
+
 class AVREditorDockableWindow;
 class AVREditorFloatingUI;
 class FProxyTabmanager;
@@ -35,7 +39,7 @@ typedef FName VREditorPanelID;
 enum class EAssetEditorCloseReason : uint8;
 
 /** Stores the animation playback state of a VR UI element */
-enum class EVREditorAnimationState : uint8
+enum class UE_DEPRECATED(5.7, "VR Editor mode is deprecated; use the XR Creative Framework plugin instead.") EVREditorAnimationState : uint8
 {
 	None,
 	Forward,
@@ -44,7 +48,7 @@ enum class EVREditorAnimationState : uint8
 
 /** Structure to keep track of all relevant interaction and animation elements of a VR Button */
 USTRUCT()
-struct FVRButton
+struct UE_DEPRECATED(5.7, "VR Editor mode is deprecated; use the XR Creative Framework plugin instead.") FVRButton
 {
 	GENERATED_USTRUCT_BODY()
 
@@ -98,7 +102,7 @@ struct FVRButton
  * VR Editor user interface manager
  */
 UCLASS()
-class UVREditorUISystem : public UObject
+class UE_DEPRECATED(5.7, "VR Editor mode is deprecated; use the XR Creative Framework plugin instead.") UVREditorUISystem : public UObject
 {
 	GENERATED_BODY()
 
@@ -441,3 +445,6 @@ protected:
 
 	FDelegateHandle AssetEditorCloseDelegate;
 };
+
+
+PRAGMA_ENABLE_DEPRECATION_WARNINGS
diff --git a/Engine/Source/Editor/VREditor/Public/UI/VRRadialMenuHandler.h b/Engine/Source/Editor/VREditor/Public/UI/VRRadialMenuHandler.h
--- a/Engine/Source/Editor/VREditor/Public/UI/VRRadialMenuHandler.h
+++ b/Engine/Source/Editor/VREditor/Public/UI/VRRadialMenuHandler.h
@@ -8,16 +8,21 @@
 #include "Framework/MultiBox/MultiBoxBuilder.h"
 #include "VRRadialMenuHandler.generated.h"
 
+
+PRAGMA_DISABLE_DEPRECATION_WARNINGS
+
+
 class FUICommandList;
 class UVREditorMode;
 
+UE_DEPRECATED(5.7, "VR Editor mode is deprecated; use the XR Creative Framework plugin instead.")
 DECLARE_DELEGATE_FourParams(FOnRadialMenuGenerated, FMenuBuilder&, TSharedPtr<FUICommandList>, UVREditorMode*, float&);
 
 /**
 * VR Editor user interface manager
 */
 UCLASS()
-class UVRRadialMenuHandler: public UObject
+class UE_DEPRECATED(5.7, "VR Editor mode is deprecated; use the XR Creative Framework plugin instead.") UVRRadialMenuHandler : public UObject
 {
 	GENERATED_BODY()
 
@@ -82,4 +87,7 @@ protected:
 
 	TArray<FOnRadialMenuGenerated> MenuStack;
 
-};
\ No newline at end of file
+};
+
+
+PRAGMA_ENABLE_DEPRECATION_WARNINGS
diff --git a/Engine/Source/Editor/VREditor/Public/VREditorAssetContainer.h b/Engine/Source/Editor/VREditor/Public/VREditorAssetContainer.h
--- a/Engine/Source/Editor/VREditor/Public/VREditorAssetContainer.h
+++ b/Engine/Source/Editor/VREditor/Public/VREditorAssetContainer.h
@@ -19,7 +19,7 @@ class UFont;
  * Asset container for VREditor.
  */
 UCLASS(MinimalAPI)
-class UVREditorAssetContainer : public UDataAsset
+class UE_DEPRECATED(5.7, "VR Editor mode is deprecated; use the XR Creative Framework plugin instead.") UVREditorAssetContainer : public UDataAsset
 {
 	GENERATED_BODY()
 
diff --git a/Engine/Source/Editor/VREditor/Public/VREditorBaseActor.h b/Engine/Source/Editor/VREditor/Public/VREditorBaseActor.h
--- a/Engine/Source/Editor/VREditor/Public/VREditorBaseActor.h
+++ b/Engine/Source/Editor/VREditor/Public/VREditorBaseActor.h
@@ -11,7 +11,7 @@
 * Represents an actor specifically for VR Editor that has roomspace transformation
 */
 UCLASS()
-class AVREditorBaseActor : public AActor
+class UE_DEPRECATED(5.7, "VR Editor mode is deprecated; use the XR Creative Framework plugin instead.") AVREditorBaseActor : public AActor
 {
 	GENERATED_BODY()
 
@@ -99,9 +99,11 @@ protected:
 
 protected:
 
+PRAGMA_DISABLE_DEPRECATION_WARNINGS
 	/** The VR mode that owns this actor */
 	UPROPERTY()
 	TObjectPtr<class UVREditorMode> VRMode;
+PRAGMA_ENABLE_DEPRECATION_WARNINGS
 
 private:
 	/** Local rotation of the UI */
diff --git a/Engine/Source/Editor/VREditor/Public/VREditorInteractor.h b/Engine/Source/Editor/VREditor/Public/VREditorInteractor.h
--- a/Engine/Source/Editor/VREditor/Public/VREditorInteractor.h
+++ b/Engine/Source/Editor/VREditor/Public/VREditorInteractor.h
@@ -11,13 +11,16 @@
 
 #define UE_API VREDITOR_API
 
+PRAGMA_DISABLE_DEPRECATION_WARNINGS
+
+
 class AActor;
 class UWidgetComponent;
 class UStaticMeshSocket;
 class AVREditorTeleporter;
 
 UENUM( BlueprintType )
-enum class EControllerType : uint8
+enum class UE_DEPRECATED(5.7, "VR Editor mode is deprecated; use the XR Creative Framework plugin instead.") EControllerType : uint8
 {
 	Laser,
 	AssistingLaser,
@@ -28,7 +31,7 @@ enum class EControllerType : uint8
 
 /** Directions the trackpad can be swiped to */
 UENUM( BlueprintType )
-enum class ETouchSwipeDirection : uint8
+enum class UE_DEPRECATED(5.7, "VR Editor mode is deprecated; use the XR Creative Framework plugin instead.") ETouchSwipeDirection : uint8
 {
 	None = 0,
 	Left = 1,
@@ -41,7 +44,7 @@ enum class ETouchSwipeDirection : uint8
  * VREditor default interactor
  */
 UCLASS(MinimalAPI,  Blueprintable, BlueprintType )
-class UVREditorInteractor : public UViewportInteractor
+class UE_DEPRECATED(5.7, "VR Editor mode is deprecated; use the XR Creative Framework plugin instead.") UVREditorInteractor : public UViewportInteractor
 {
 	GENERATED_BODY()
 
@@ -511,4 +514,7 @@ protected:
 	TObjectPtr<class UVREditorMode> VRMode;
 };
 
+PRAGMA_ENABLE_DEPRECATION_WARNINGS
+
 #undef UE_API
+
diff --git a/Engine/Source/Editor/VREditor/Public/VREditorMode.h b/Engine/Source/Editor/VREditor/Public/VREditorMode.h
--- a/Engine/Source/Editor/VREditor/Public/VREditorMode.h
+++ b/Engine/Source/Editor/VREditor/Public/VREditorMode.h
@@ -19,6 +19,9 @@
 
 #define UE_API VREDITOR_API
 
+PRAGMA_DISABLE_DEPRECATION_WARNINGS
+
+
 class AActor;
 class FEditorViewportClient;
 class SLevelViewport;
@@ -34,7 +37,7 @@ enum class EGizmoHandleTypes : uint8;
 /**
  * Types of actions that can be performed with VR controller devices
  */
-namespace VRActionTypes
+namespace UE_DEPRECATED(5.7, "VR Editor mode is deprecated; use the XR Creative Framework plugin instead.") VRActionTypes
 {
 	static const FName Touch( "Touch" );
 	static const FName Modifier( "Modifier" );
@@ -52,8 +55,8 @@ namespace VRActionTypes
 /**
  * VR Editor Mode. Extends editor viewports with functionality for VR controls and object manipulation
  */
-UCLASS(MinimalAPI,  Abstract, Transient )
-class UVREditorMode : public UVREditorModeBase
+UCLASS(MinimalAPI, Abstract, Transient)
+class UE_DEPRECATED(5.7, "VR Editor mode is deprecated; use the XR Creative Framework plugin instead.") UVREditorMode : public UVREditorModeBase
 {
 	GENERATED_BODY()
 
@@ -479,4 +482,7 @@ private:
 	bool bAddedViewportWorldInteractionExtension = false;
 };
 
+PRAGMA_ENABLE_DEPRECATION_WARNINGS
+
 #undef UE_API
+
diff --git a/Engine/Source/Editor/VREditor/Public/VREditorStyle.h b/Engine/Source/Editor/VREditor/Public/VREditorStyle.h
--- a/Engine/Source/Editor/VREditor/Public/VREditorStyle.h
+++ b/Engine/Source/Editor/VREditor/Public/VREditorStyle.h
@@ -11,7 +11,7 @@ class ISlateStyle;
 struct FSlateBrush;
 
 /**  */
-class FVREditorStyle
+class UE_DEPRECATED(5.7, "VR Editor mode is deprecated; use the XR Creative Framework plugin instead.") FVREditorStyle
 {
 public:
 
diff --git a/Engine/Source/Editor/VREditor/Public/VRModeSettings.h b/Engine/Source/Editor/VREditor/Public/VRModeSettings.h
--- a/Engine/Source/Editor/VREditor/Public/VRModeSettings.h
+++ b/Engine/Source/Editor/VREditor/Public/VRModeSettings.h
@@ -17,7 +17,7 @@ class UVREditorModeBase;
 
 
 UENUM()
-enum class EInteractorHand : uint8
+enum class UE_DEPRECATED(5.7, "VR Editor mode is deprecated; use the XR Creative Framework plugin instead.") EInteractorHand : uint8
 {
 	/** Right hand */
 	Right,
@@ -44,53 +44,57 @@ public:
 	uint32 bEnableAutoVREditMode : 1;
 
 	// Whether or not sequences should be autokeyed
-	UPROPERTY(EditAnywhere, config, Category = "Cinematics")
+	UE_DEPRECATED(5.7, "VR Editor mode is deprecated; use the XR Creative Framework plugin instead.")
+	UPROPERTY(config)
 	uint32 bAutokeySequences : 1;
 
+PRAGMA_DISABLE_DEPRECATION_WARNINGS
 	// Which hand should have the primary interactor laser on it
-	UPROPERTY(EditAnywhere, config, Category = "General")
+	UE_DEPRECATED(5.7, "VR Editor mode is deprecated; use the XR Creative Framework plugin instead.")
+	UPROPERTY(config)
 	EInteractorHand InteractorHand;
+PRAGMA_ENABLE_DEPRECATION_WARNINGS
 
 	/** Show the movement grid for orientation while moving through the world */
-	UPROPERTY(EditAnywhere, config, Category = "World Movement")
+	UE_DEPRECATED(5.7, "VR Editor mode is deprecated; use the XR Creative Framework plugin instead.")
+	UPROPERTY(config)
 	uint32 bShowWorldMovementGrid : 1;
 
 	/** Dim the surroundings while moving through the world */
-	UPROPERTY(EditAnywhere, config, Category = "World Movement")
+	UE_DEPRECATED(5.7, "VR Editor mode is deprecated; use the XR Creative Framework plugin instead.")
+	UPROPERTY(config)
 	uint32 bShowWorldMovementPostProcess : 1;
 
 	/** Display a progress bar while scaling that shows your current scale */
-	UPROPERTY(EditAnywhere, config, Category = "UI Customization")
+	UE_DEPRECATED(5.7, "VR Editor mode is deprecated; use the XR Creative Framework plugin instead.")
+	UPROPERTY(config)
 	uint32 bShowWorldScaleProgressBar : 1;
 
 	/** Adjusts the brightness of the UI panels */
-	UPROPERTY(EditAnywhere, config, Category = "UI Customization", meta = (DisplayName = "UI Panel Brightness", ClampMin = 0.01, UIMax = 10.0))
+	UE_DEPRECATED(5.7, "VR Editor mode is deprecated; use the XR Creative Framework plugin instead.")
+	UPROPERTY(config)
 	float UIBrightness;
 
 	/** The size of the transform gizmo */
-	UPROPERTY(EditAnywhere, config, Category = "UI Customization", meta = (ClampMin = 0.1, ClampMax = 2.0))
+	UE_DEPRECATED(5.7, "VR Editor mode is deprecated; use the XR Creative Framework plugin instead.")
+	UPROPERTY(config)
 	float GizmoScale;
 
 	/** The maximum time in seconds between two clicks for a double-click to register */
-	UPROPERTY(EditAnywhere, config, Category = "Motion Controllers", meta = (ClampMin = 0.01, ClampMax = 1.0))
+	UE_DEPRECATED(5.7, "VR Editor mode is deprecated; use the XR Creative Framework plugin instead.")
+	UPROPERTY(config)
 	float DoubleClickTime;
 
 	/** The amount (between 0-1) you have to depress the Vive controller trigger to register a press */
-	UPROPERTY(EditAnywhere, config, Category = "Motion Controllers", meta = (DisplayName = "Trigger Pressed Threshold (Vive)", ClampMin = 0.01, ClampMax = 1.0))
+	UE_DEPRECATED(5.7, "VR Editor mode is deprecated; use the XR Creative Framework plugin instead.")
+	UPROPERTY(config)
 	float TriggerPressedThreshold_Vive;
 
 	/** The amount (between 0-1) you have to depress the Oculus Touch controller trigger to register a press */
-	UPROPERTY(EditAnywhere, config, Category = "Motion Controllers", meta = (DisplayName = "Trigger Pressed Threshold (Oculus Touch)", ClampMin = 0.01, ClampMax = 1.0))
+	UE_DEPRECATED(5.7, "VR Editor mode is deprecated; use the XR Creative Framework plugin instead.")
+	UPROPERTY(config)
 	float TriggerPressedThreshold_Rift;
 
-	UE_DEPRECATED(5.1, "Refer to UVREditorMode::InteractorClass, or create a derived mode to override the interactor class.")
-	UPROPERTY(meta = (DeprecatedProperty, DeprecationMessage = "Refer to UVREditorMode::InteractorClass, or create a derived mode to override the interactor class."))
-	TSoftClassPtr<UVREditorInteractor> InteractorClass;
-
-	UE_DEPRECATED(5.1, "Refer to UVREditorMode::TeleporterClass, or create a derived mode to override the teleporter class.")
-	UPROPERTY(meta = (DeprecatedProperty, DeprecationMessage = "Refer to UVREditorMode::TeleporterClass, or create a derived mode to override the teleporter class."))
-	TSoftClassPtr<AVREditorTeleporter> TeleporterClass;
-
 	/** The mode extension to use when UnrealEd is in VR mode. Use VREditorMode to get default editor behavior or select a custom mode. */
 	UPROPERTY(EditAnywhere, config, NoClear, Category = "General")
 	TSoftClassPtr<UVREditorModeBase> ModeClass;
diff --git a/Engine/Source/Editor/VREditor/Public/VRScoutingInteractor.h b/Engine/Source/Editor/VREditor/Public/VRScoutingInteractor.h
--- a/Engine/Source/Editor/VREditor/Public/VRScoutingInteractor.h
+++ b/Engine/Source/Editor/VREditor/Public/VRScoutingInteractor.h
@@ -11,6 +11,9 @@
 
 #define UE_API VREDITOR_API
 
+PRAGMA_DISABLE_DEPRECATION_WARNINGS
+
+
 class AActor;
 class UStaticMesh;
 class UStaticMeshSocket;
@@ -20,7 +23,7 @@ class UInputComponent;
  * Represents the interactor in the world
  */
 UCLASS(MinimalAPI, Abstract)
-class UVRScoutingInteractor : public UVREditorInteractor
+class UE_DEPRECATED(5.7, "VR Editor mode is deprecated; use the XR Creative Framework plugin instead.") UVRScoutingInteractor : public UVREditorInteractor
 {
 	GENERATED_BODY()
 
@@ -87,4 +90,7 @@ private:
 
 };
 
+PRAGMA_ENABLE_DEPRECATION_WARNINGS
+
 #undef UE_API
+
diff --git a/Engine/Source/Runtime/Engine/Private/Analytics/EditorAnalyticsSessionSummary.cpp b/Engine/Source/Runtime/Engine/Private/Analytics/EditorAnalyticsSessionSummary.cpp
--- a/Engine/Source/Runtime/Engine/Private/Analytics/EditorAnalyticsSessionSummary.cpp
+++ b/Engine/Source/Runtime/Engine/Private/Analytics/EditorAnalyticsSessionSummary.cpp
@@ -48,7 +48,7 @@ FEditorAnalyticsSessionSummary::FEditorAnalyticsSessionSummary(TSharedPtr<IAnaly
 	EditorAnalyticsProperties::Idle1Min.Set(GetStore(), 0);
 	EditorAnalyticsProperties::Idle5Min.Set(GetStore(), 0);
 	EditorAnalyticsProperties::Idle30Min.Set(GetStore(), 0);
-	EditorAnalyticsProperties::IsInVRMode.Set(GetStore(), IVREditorModule::Get().IsVREditorModeActive());
+	EditorAnalyticsProperties::IsInVRMode.Set(GetStore(), IVREditorModule::Get().GetVRModeBase() != nullptr);
 	EditorAnalyticsProperties::IsInPIE.Set(GetStore(), FPlayWorldCommandCallbacks::IsInPIE());
 	EditorAnalyticsProperties::IsUnattended.Set(GetStore(), FApp::IsUnattended());
 	EditorAnalyticsProperties::DirtyPackageCount.Set(GetStore(), 0);
@@ -103,7 +103,7 @@ bool FEditorAnalyticsSessionSummary::UpdateSessionProgressInternal(bool bCrashin
 		return true; // Yes, should persist.
 	}
 
-	EditorAnalyticsProperties::IsInVRMode.Set(GetStore(), IVREditorModule::Get().IsVREditorModeActive());
+	EditorAnalyticsProperties::IsInVRMode.Set(GetStore(), IVREditorModule::Get().GetVRModeBase() != nullptr);
 	EditorAnalyticsProperties::IsInPIE.Set(GetStore(), FPlayWorldCommandCallbacks::IsInPIE());
 
 	// Accumulate stall stats
