VR Editor: Fix for failed check in UWidgetComponent unregister during engine pre-exit.

FVREditorModeManager now disables VR Editor in response to FEditorDelegates::OnEditorPreExit.

#FailedValidations Preflight

diff --git a/Engine/Plugins/Experimental/XRCreativeFramework/Source/XRCreativeEditor/Private/XRCreativeVREditorMode.cpp b/Engine/Plugins/Experimental/XRCreativeFramework/Source/XRCreativeEditor/Private/XRCreativeVREditorMode.cpp
--- a/Engine/Plugins/Experimental/XRCreativeFramework/Source/XRCreativeEditor/Private/XRCreativeVREditorMode.cpp
+++ b/Engine/Plugins/Experimental/XRCreativeFramework/Source/XRCreativeEditor/Private/XRCreativeVREditorMode.cpp
@@ -262,7 +262,7 @@ void UXRCreativeVREditorMode::DisableStereo()
 {
 	UOpenXRInputFunctionLibrary::EndXRSession();
 
-	if (TSharedPtr<SLevelViewport> Viewport = GetVrLevelViewport(); ensure(Viewport))
+	if (TSharedPtr<SLevelViewport> Viewport = GetVrLevelViewport())
 	{
 		StereoViewportShutdown(Viewport.ToSharedRef());
 	}
diff --git a/Engine/Source/Editor/VREditor/Private/VREditorModeManager.cpp b/Engine/Source/Editor/VREditor/Private/VREditorModeManager.cpp
--- a/Engine/Source/Editor/VREditor/Private/VREditorModeManager.cpp
+++ b/Engine/Source/Editor/VREditor/Private/VREditorModeManager.cpp
@@ -44,6 +44,8 @@ FVREditorModeManager::FVREditorModeManager() :
 	{
 		HandleAssetFilesLoaded();
 	}
+
+	FEditorDelegates::OnEditorPreExit.AddRaw(this, &FVREditorModeManager::HandleEditorPreExit);
 }
 
 FVREditorModeManager::~FVREditorModeManager()
@@ -63,6 +65,8 @@ FVREditorModeManager::~FVREditorModeManager()
 		CurrentVREditorMode->OnVRModeEntryComplete().RemoveAll(this);
 		CurrentVREditorMode = nullptr;
 	}
+
+	FEditorDelegates::OnEditorPreExit.RemoveAll(this);
 }
 
 void FVREditorModeManager::Tick( const float DeltaTime )
@@ -368,4 +372,14 @@ void FVREditorModeManager::HandleAssetAddedOrRemoved(TConstArrayView<FAssetData>
 	UpdateCachedModeClassPaths();
 }
 
+void FVREditorModeManager::HandleEditorPreExit()
+{
+	if (IsVREditorActive())
+	{
+		const bool bEnable_false = false;
+		const bool bForceWithoutHMD_false = false;
+		EnableVREditor(bEnable_false, bForceWithoutHMD_false);
+	}
+}
+
 #undef LOCTEXT_NAMESPACE
diff --git a/Engine/Source/Editor/VREditor/Private/VREditorModeManager.h b/Engine/Source/Editor/VREditor/Private/VREditorModeManager.h
--- a/Engine/Source/Editor/VREditor/Private/VREditorModeManager.h
+++ b/Engine/Source/Editor/VREditor/Private/VREditorModeManager.h
@@ -108,6 +108,9 @@ private:
 	
 	/** On level changed */
 	void OnMapChanged( UWorld* World, EMapChangeType MapChangeType );
+
+	/** Ensure we shut down cleanly early in editor shutdown. */
+	void HandleEditorPreExit();
 	
 	/** The current mode, nullptr if none */
 	UPROPERTY()
